#UE4:
* added check for safety

diff --git a/Engine/Source/Runtime/OpenGLDrv/Public/OpenGLES2.h b/Engine/Source/Runtime/OpenGLDrv/Public/OpenGLES2.h
--- a/Engine/Source/Runtime/OpenGLDrv/Public/OpenGLES2.h
+++ b/Engine/Source/Runtime/OpenGLDrv/Public/OpenGLES2.h
@@ -207,11 +207,17 @@ struct FOpenGLES2 : public FOpenGLBase
 
 	static FORCEINLINE void EnableIndexed(GLenum Parameter, GLuint Index)
 	{
+		// We don't have MRT on ES2 and Index was used for RenderTargetIndex so Index can be ignore, other Parameters might not work.
+		check(Parameter == GL_BLEND);
+
 		glEnable(Parameter);
 	}
 
 	static FORCEINLINE void DisableIndexed(GLenum Parameter, GLuint Index)
 	{
+		// We don't have MRT on ES2 and Index was used for RenderTargetIndex so Index can be ignore, other Parameters might not work.
+		check(Parameter == GL_BLEND);
+
 		glDisable(Parameter);
 	}
 
