#UE4:
* fixed ES2 compile error

diff --git a/Engine/Source/Runtime/OpenGLDrv/Public/OpenGLES2.h b/Engine/Source/Runtime/OpenGLDrv/Public/OpenGLES2.h
--- a/Engine/Source/Runtime/OpenGLDrv/Public/OpenGLES2.h
+++ b/Engine/Source/Runtime/OpenGLDrv/Public/OpenGLES2.h
@@ -647,6 +647,9 @@ protected:
 #ifndef GL_RG32F
 #define GL_RG32F 0x8230
 #endif
+#ifndef GL_FRAMEBUFFER_SRG
+#define GL_FRAMEBUFFER_SRG 0x8DB9
+#endif
 #define GL_RG8I 0x8237
 #define GL_RG8UI 0x8238
 #define GL_RG16I 0x8239
