Workaround AMD driver bug with GL
#github change from AndrewScheidecker in https://github.com/EpicGames/UnrealEngine/pull/53

From AndrewS:
Work around an apparent bug in AMD drivers (Catalyst 14.2): The glPushDebugGroup docs says a negative length is valid, and means that the string is null terminated. However, the AMD drivers throw an error saying that the length must be between 0 and 1023. The change avoids the problem by passing in the length of the string explicitly.

diff --git a/Engine/Source/Runtime/OpenGLDrv/Private/Windows/OpenGLWindows.h b/Engine/Source/Runtime/OpenGLDrv/Private/Windows/OpenGLWindows.h
--- a/Engine/Source/Runtime/OpenGLDrv/Private/Windows/OpenGLWindows.h
+++ b/Engine/Source/Runtime/OpenGLDrv/Private/Windows/OpenGLWindows.h
@@ -373,7 +373,7 @@ struct FWindowsOpenGL : public FOpenGL4
 	{
 		if (glPushDebugGroup && bDebugContext)
 		{
-			glPushDebugGroup( GL_DEBUG_SOURCE_APPLICATION, 1, -1,Name);
+			glPushDebugGroup( GL_DEBUG_SOURCE_APPLICATION, 1, FCStringAnsi::Strlen(Name),Name);
 		}
 	}
 
