missing break statement

diff --git a/Engine/Source/Runtime/OpenGLDrv/Private/OpenGLShaders.cpp b/Engine/Source/Runtime/OpenGLDrv/Private/OpenGLShaders.cpp
--- a/Engine/Source/Runtime/OpenGLDrv/Private/OpenGLShaders.cpp
+++ b/Engine/Source/Runtime/OpenGLDrv/Private/OpenGLShaders.cpp
@@ -859,7 +859,7 @@ static void VerifyUniformLayout(const TCHAR* UniformName, const UniformData& GLS
 				FString BaseTypeName;
 				switch(Member.GetBaseType())
 				{
-					case UBMT_STRUCT:  BaseTypeName = TEXT("struct"); 
+					case UBMT_STRUCT:  BaseTypeName = TEXT("struct"); break;
 					case UBMT_BOOL:    BaseTypeName = TEXT("bool"); break;
 					case UBMT_INT32:   BaseTypeName = TEXT("int"); break;
 					case UBMT_UINT32:  BaseTypeName = TEXT("uint"); break;
