#ttp 332421 CRITICAL: Regression: OPENGL: BLACKJACK: Text on table does not read clearly
#summary Fix anisotropic filtering on OpenGL, we hadn't been turning it on for desktop OpenGL!
#change Make base FOpenGL responsible for testing for GL_EXT_texture_filter_anisotropic.
#removed Took the overrides out of GLES2 and HTML5, they aren't required anymore.

diff --git a/Engine/Source/Runtime/OpenGLDrv/Private/HTML5/HTML5OpenGL.cpp b/Engine/Source/Runtime/OpenGLDrv/Private/HTML5/HTML5OpenGL.cpp
--- a/Engine/Source/Runtime/OpenGLDrv/Private/HTML5/HTML5OpenGL.cpp
+++ b/Engine/Source/Runtime/OpenGLDrv/Private/HTML5/HTML5OpenGL.cpp
@@ -40,7 +40,6 @@ void FHTML5OpenGL::ProcessExtensions( const FString& ExtensionsString )
 	bSupportsVertexArrayObjects = ExtensionsString.Contains(TEXT("GL_OES_vertex_array_object")) ;
 	bSupportsDiscardFrameBuffer = ExtensionsString.Contains(TEXT("GL_EXT_discard_framebuffer"));
 	bSupportsNVFrameBufferBlit = ExtensionsString.Contains(TEXT("GL_NV_framebuffer_blit"));
-	bSupportsTextureFilterAnisotropic = ExtensionsString.Contains(TEXT("GL_EXT_texture_filter_anisotropic"));
 	bSupportsShaderTextureLod = ExtensionsString.Contains(TEXT("GL_EXT_shader_texture_lod"));
     bSupportsTextureCubeLodEXT = bSupportsShaderTextureLod;
 
diff --git a/Engine/Source/Runtime/OpenGLDrv/Private/OpenGLDrv.cpp b/Engine/Source/Runtime/OpenGLDrv/Private/OpenGLDrv.cpp
--- a/Engine/Source/Runtime/OpenGLDrv/Private/OpenGLDrv.cpp
+++ b/Engine/Source/Runtime/OpenGLDrv/Private/OpenGLDrv.cpp
@@ -438,6 +438,7 @@ bool  FOpenGLBase::bSupportsASTC = false;
 bool  FOpenGLBase::bSupportsCopyImage = false;
 bool  FOpenGLBase::bSupportsSeamlessCubemap = false;
 bool  FOpenGLBase::bSupportsVolumeTextureRendering = false;
+bool  FOpenGLBase::bSupportsTextureFilterAnisotropic = false;
 
 void FOpenGLBase::ProcessQueryGLInt()
 {
@@ -487,6 +488,8 @@ void FOpenGLBase::ProcessExtensions( const FString& ExtensionsString )
 	bSupportsCopyImage = ExtensionsString.Contains(TEXT("GL_ARB_copy_image"));
 
 	bSupportsSeamlessCubemap = ExtensionsString.Contains(TEXT("GL_ARB_seamless_cube_map"));
+	
+	bSupportsTextureFilterAnisotropic = ExtensionsString.Contains(TEXT("GL_EXT_texture_filter_anisotropic"));
 }
 
 void InitDefaultGLContextState(void)
diff --git a/Engine/Source/Runtime/OpenGLDrv/Private/OpenGLES2.cpp b/Engine/Source/Runtime/OpenGLDrv/Private/OpenGLES2.cpp
--- a/Engine/Source/Runtime/OpenGLDrv/Private/OpenGLES2.cpp
+++ b/Engine/Source/Runtime/OpenGLDrv/Private/OpenGLES2.cpp
@@ -86,9 +86,6 @@ int FOpenGLES2::ShaderHighPrecision = 0;
 /** GL_NV_framebuffer_blit */
 bool FOpenGLES2::bSupportsNVFrameBufferBlit = false;
 
-/** GL_EXT_texture_filter_anisotropic */
-bool FOpenGLES2::bSupportsTextureFilterAnisotropic = false;
-
 /** GL_OES_packed_depth_stencil */
 bool FOpenGLES2::bSupportsPackedDepthStencil = false;
 
@@ -158,7 +155,6 @@ void FOpenGLES2::ProcessExtensions( const FString& ExtensionsString )
 	bSupportsVertexArrayObjects = ExtensionsString.Contains(TEXT("GL_OES_vertex_array_object")) ;
 	bSupportsDiscardFrameBuffer = ExtensionsString.Contains(TEXT("GL_EXT_discard_framebuffer"));
 	bSupportsNVFrameBufferBlit = ExtensionsString.Contains(TEXT("GL_NV_framebuffer_blit"));
-	bSupportsTextureFilterAnisotropic = ExtensionsString.Contains(TEXT("GL_EXT_texture_filter_anisotropic"));
 	bSupportsPackedDepthStencil = ExtensionsString.Contains(TEXT("GL_OES_packed_depth_stencil"));
 	bSupportsShaderTextureLod = ExtensionsString.Contains(TEXT("GL_EXT_shader_texture_lod"));
 	bSupportsTextureStorageEXT = ExtensionsString.Contains(TEXT("GL_EXT_texture_storage"));
diff --git a/Engine/Source/Runtime/OpenGLDrv/Public/OpenGL.h b/Engine/Source/Runtime/OpenGLDrv/Public/OpenGL.h
--- a/Engine/Source/Runtime/OpenGLDrv/Public/OpenGL.h
+++ b/Engine/Source/Runtime/OpenGLDrv/Public/OpenGL.h
@@ -115,7 +115,7 @@ public:
 	static FORCEINLINE bool SupportsFastBufferData()					{ return true; }
 	static FORCEINLINE bool SupportsCopyImage()							{ return bSupportsCopyImage; }
 	static FORCEINLINE bool SupportsCopyTextureLevels()					{ return false; }
-	static FORCEINLINE bool SupportsTextureFilterAnisotropic()			{ return false; }
+	static FORCEINLINE bool SupportsTextureFilterAnisotropic()			{ return bSupportsTextureFilterAnisotropic; }
 	static FORCEINLINE bool SupportsPackedDepthStencil()				{ return true; }
 	static FORCEINLINE bool SupportsTextureCubeLodEXT()					{ return true; }
 	static FORCEINLINE bool SupportsShaderTextureLod()					{ return false; }
@@ -304,6 +304,9 @@ protected:
 	
 	/** Can we render to texture 2D array or 3D */
 	static bool bSupportsVolumeTextureRendering;
+	
+	/** GL_EXT_texture_filter_anisotropic Can we use anisotropic filtering? */
+	static bool bSupportsTextureFilterAnisotropic;
 };
 
 /** Unreal tokens that maps to different OpenGL tokens by platform. */
diff --git a/Engine/Source/Runtime/OpenGLDrv/Public/OpenGLES2.h b/Engine/Source/Runtime/OpenGLDrv/Public/OpenGLES2.h
--- a/Engine/Source/Runtime/OpenGLDrv/Public/OpenGLES2.h
+++ b/Engine/Source/Runtime/OpenGLDrv/Public/OpenGLES2.h
@@ -100,7 +100,6 @@ struct FOpenGLES2 : public FOpenGLBase
 	static FORCEINLINE bool SupportsETC1()								{ return bSupportsETC1; }
 	static FORCEINLINE bool SupportsETC2()								{ return bSupportsETC2; }
 	static FORCEINLINE bool SupportsCombinedDepthStencilAttachment()	{ return false; }
-	static FORCEINLINE bool SupportsTextureFilterAnisotropic()			{ return bSupportsTextureFilterAnisotropic; }
 	static FORCEINLINE bool SupportsPackedDepthStencil()				{ return bSupportsPackedDepthStencil; }
 	static FORCEINLINE bool SupportsTextureCubeLodEXT()					{ return bSupportsTextureCubeLodEXT; }
 	static FORCEINLINE bool SupportsShaderTextureLod()					{ return bSupportsShaderTextureLod; }
@@ -400,9 +399,6 @@ protected:
 	/** GL_NV_framebuffer_blit */
 	static bool bSupportsNVFrameBufferBlit;
 
-	/** GL_EXT_texture_filter_anisotropic  */
-	static bool bSupportsTextureFilterAnisotropic;
-
 	/** GL_OES_packed_depth_stencil */
 	static bool bSupportsPackedDepthStencil;
 
