Fix for -opengl on AMD

diff --git a/Engine/Source/Runtime/OpenGLDrv/Private/OpenGL4.cpp b/Engine/Source/Runtime/OpenGLDrv/Private/OpenGL4.cpp
--- a/Engine/Source/Runtime/OpenGLDrv/Private/OpenGL4.cpp
+++ b/Engine/Source/Runtime/OpenGLDrv/Private/OpenGL4.cpp
@@ -10,6 +10,7 @@
 
 GLint FOpenGL4::MaxComputeTextureImageUnits = -1;
 GLint FOpenGL4::MaxComputeUniformComponents = -1;
+bool FOpenGL4::bSupportsComputeShaders = true;
 
 void FOpenGL4::ProcessQueryGLInt()
 {
@@ -23,6 +24,10 @@ void FOpenGL4::ProcessQueryGLInt()
 	LOG_AND_GET_GL_INT(GL_MAX_COMPUTE_UNIFORM_COMPONENTS, 0, MaxComputeUniformComponents);
 	LOG_AND_GET_GL_INT(GL_MAX_TESS_CONTROL_UNIFORM_COMPONENTS, 0, MaxHullUniformComponents);
 	LOG_AND_GET_GL_INT(GL_MAX_TESS_EVALUATION_UNIFORM_COMPONENTS, 0, MaxDomainUniformComponents);
+	if ( MaxComputeTextureImageUnits == -1 || MaxComputeUniformComponents == -1 || MaxHullUniformComponents == -1 || MaxDomainUniformComponents == -1 )
+	{
+		bSupportsComputeShaders = false;
+	}
 }
 
 void FOpenGL4::ProcessExtensions( const FString& ExtensionsString )
diff --git a/Engine/Source/Runtime/OpenGLDrv/Public/OpenGL4.h b/Engine/Source/Runtime/OpenGLDrv/Public/OpenGL4.h
--- a/Engine/Source/Runtime/OpenGLDrv/Public/OpenGL4.h
+++ b/Engine/Source/Runtime/OpenGLDrv/Public/OpenGL4.h
@@ -15,7 +15,7 @@ struct FOpenGL4 : public FOpenGL3
 {
 	static FORCEINLINE bool SupportsSeparateAlphaBlend()				{ return true; }
 	static FORCEINLINE bool SupportsTessellation()						{ return true; }
-	static FORCEINLINE bool SupportsComputeShaders()					{ return true; }
+	static FORCEINLINE bool SupportsComputeShaders()					{ return bSupportsComputeShaders; }
 
 	// Optional
 
@@ -62,4 +62,5 @@ struct FOpenGL4 : public FOpenGL3
 protected:
 	static GLint MaxComputeTextureImageUnits;
 	static GLint MaxComputeUniformComponents;
+	static bool bSupportsComputeShaders;
 };
