Move RHI shutdown into it's own function so we don't end up calling virtual functions from the destructor.

diff --git a/Engine/Source/Runtime/EmptyRHI/Private/EmptyRHI.cpp b/Engine/Source/Runtime/EmptyRHI/Private/EmptyRHI.cpp
--- a/Engine/Source/Runtime/EmptyRHI/Private/EmptyRHI.cpp
+++ b/Engine/Source/Runtime/EmptyRHI/Private/EmptyRHI.cpp
@@ -114,16 +114,14 @@ FEmptyDynamicRHI::FEmptyDynamicRHI()
 	GIsRHIInitialized = true;
 }
 
-FEmptyDynamicRHI::~FEmptyDynamicRHI()
+void FEmptyDynamicRHI::Init()
 {
-	check(IsInGameThread() && IsInRenderingThread());
 
 }
 
-
-void FEmptyDynamicRHI::Init()
+void FEmptyDynamicRHI::Shutdown()
 {
-
+	check(IsInGameThread() && IsInRenderingThread());
 }
 
 void FEmptyDynamicRHI::RHIBeginFrame()
diff --git a/Engine/Source/Runtime/EmptyRHI/Public/EmptyRHI.h b/Engine/Source/Runtime/EmptyRHI/Public/EmptyRHI.h
--- a/Engine/Source/Runtime/EmptyRHI/Public/EmptyRHI.h
+++ b/Engine/Source/Runtime/EmptyRHI/Public/EmptyRHI.h
@@ -26,11 +26,12 @@ public:
 	FEmptyDynamicRHI();
 
 	/** Destructor */
-	~FEmptyDynamicRHI();
+	~FEmptyDynamicRHI() {}
 
 
 	// FDynamicRHI interface.
 	virtual void Init();
+	virtual void Shutdown();
 	virtual void PushEvent(const TCHAR* Name);
 	virtual void PopEvent();
 
diff --git a/Engine/Source/Runtime/NullDrv/Private/NullRHI.cpp b/Engine/Source/Runtime/NullDrv/Private/NullRHI.cpp
--- a/Engine/Source/Runtime/NullDrv/Private/NullRHI.cpp
+++ b/Engine/Source/Runtime/NullDrv/Private/NullRHI.cpp
@@ -40,6 +40,10 @@ void FNullDynamicRHI::Init()
 	GIsRHIInitialized = true;
 }
 
+void FNullDynamicRHI::Shutdown()
+{
+}
+
 /**
  * Return a shared large static buffer that can be used to return from any 
  * function that needs to return a valid pointer (but can be garbage data)
diff --git a/Engine/Source/Runtime/NullDrv/Public/NullRHI.h b/Engine/Source/Runtime/NullDrv/Public/NullRHI.h
--- a/Engine/Source/Runtime/NullDrv/Public/NullRHI.h
+++ b/Engine/Source/Runtime/NullDrv/Public/NullRHI.h
@@ -16,6 +16,7 @@ public:
 
 	// FDynamicRHI interface.
 	virtual void Init();
+	virtual void Shutdown();
 
 	// Implement the dynamic RHI interface using the null implementations defined in RHIMethods.h
 	#define DEFINE_RHIMETHOD(Type,Name,ParameterTypesAndNames,ParameterNames,ReturnStatement,NullImplementation) \
diff --git a/Engine/Source/Runtime/OpenGLDrv/Private/OpenGLDevice.cpp b/Engine/Source/Runtime/OpenGLDrv/Private/OpenGLDevice.cpp
--- a/Engine/Source/Runtime/OpenGLDrv/Private/OpenGLDevice.cpp
+++ b/Engine/Source/Runtime/OpenGLDrv/Private/OpenGLDevice.cpp
@@ -789,18 +789,6 @@ FOpenGLDynamicRHI::FOpenGLDynamicRHI()
 
 extern void DestroyShadersAndPrograms();
 
-FOpenGLDynamicRHI::~FOpenGLDynamicRHI()
-{
-	check(IsInGameThread() && IsInRenderingThread()); // require that the render thread has been shut down
-
-	Cleanup();
-
-	DestroyShadersAndPrograms();
-	PlatformDestroyOpenGLDevice(PlatformDevice);
-
-	PrivateOpenGLDevicePtr = NULL;
-}
-
 void FOpenGLDynamicRHI::Init()
 {
 	check(!GIsRHIInitialized);
@@ -840,6 +828,18 @@ void FOpenGLDynamicRHI::Init()
 	GIsRHIInitialized = true;
 }
 
+void FOpenGLDynamicRHI::Shutdown()
+{
+	check(IsInGameThread() && IsInRenderingThread()); // require that the render thread has been shut down
+
+	Cleanup();
+
+	DestroyShadersAndPrograms();
+	PlatformDestroyOpenGLDevice(PlatformDevice);
+
+	PrivateOpenGLDevicePtr = NULL;
+}
+
 void FOpenGLDynamicRHI::Cleanup()
 {
 	if(GIsRHIInitialized)
diff --git a/Engine/Source/Runtime/OpenGLDrv/Public/OpenGLDrv.h b/Engine/Source/Runtime/OpenGLDrv/Public/OpenGLDrv.h
--- a/Engine/Source/Runtime/OpenGLDrv/Public/OpenGLDrv.h
+++ b/Engine/Source/Runtime/OpenGLDrv/Public/OpenGLDrv.h
@@ -292,10 +292,11 @@ public:
 	FOpenGLDynamicRHI();
 
 	/** Destructor */
-	~FOpenGLDynamicRHI();
+	~FOpenGLDynamicRHI() {}
 
 	// FDynamicRHI interface.
 	virtual void Init();
+	virtual void Shutdown();
 
 	virtual void PushEvent(const TCHAR* Name);
 	virtual void PopEvent();
diff --git a/Engine/Source/Runtime/RHI/Private/DynamicRHI.cpp b/Engine/Source/Runtime/RHI/Private/DynamicRHI.cpp
--- a/Engine/Source/Runtime/RHI/Private/DynamicRHI.cpp
+++ b/Engine/Source/Runtime/RHI/Private/DynamicRHI.cpp
@@ -76,6 +76,7 @@ void RHIExit()
 	if ( !GUsingNullRHI )
 	{
 		// Destruct the dynamic RHI.
+		GDynamicRHI->Shutdown();
 		delete GDynamicRHI;
 		GDynamicRHI = NULL;
 	}
diff --git a/Engine/Source/Runtime/RHI/Public/DynamicRHI.h b/Engine/Source/Runtime/RHI/Public/DynamicRHI.h
--- a/Engine/Source/Runtime/RHI/Public/DynamicRHI.h
+++ b/Engine/Source/Runtime/RHI/Public/DynamicRHI.h
@@ -28,6 +28,8 @@ public:
 
 	/** Initializes the RHI; separate from IDynamicRHIModule::CreateRHI so that GDynamicRHI is set when it is called. */
 	virtual void Init() = 0;
+	/** Shutdown the RHI; handle shutdown and resource destruction before the RHI's actual destructor is called (so that all resources of the RHI are still available for shutdown). */
+	virtual void Shutdown() = 0;
 
 	virtual void PushEvent(const TCHAR* Name) {}
 	virtual void PopEvent() {}
diff --git a/Engine/Source/Runtime/Windows/D3D11RHI/Private/D3D11Device.cpp b/Engine/Source/Runtime/Windows/D3D11RHI/Private/D3D11Device.cpp
--- a/Engine/Source/Runtime/Windows/D3D11RHI/Private/D3D11Device.cpp
+++ b/Engine/Source/Runtime/Windows/D3D11RHI/Private/D3D11Device.cpp
@@ -189,7 +189,7 @@ FD3D11DynamicRHI::FD3D11DynamicRHI(IDXGIFactory* InDXGIFactory,D3D_FEATURE_LEVEL
 	DynamicIB = new FD3D11DynamicBuffer(this,D3D11_BIND_INDEX_BUFFER,DynamicIBSizes);
 }
 
-FD3D11DynamicRHI::~FD3D11DynamicRHI()
+void FD3D11DynamicRHI::Shutdown()
 {
 	check(IsInGameThread() && IsInRenderingThread());  // require that the render thread has been shut down
 
diff --git a/Engine/Source/Runtime/Windows/D3D11RHI/Private/D3D11RHIPrivate.h b/Engine/Source/Runtime/Windows/D3D11RHI/Private/D3D11RHIPrivate.h
--- a/Engine/Source/Runtime/Windows/D3D11RHI/Private/D3D11RHIPrivate.h
+++ b/Engine/Source/Runtime/Windows/D3D11RHI/Private/D3D11RHIPrivate.h
@@ -312,13 +312,14 @@ public:
 	FD3D11DynamicRHI(IDXGIFactory* InDXGIFactory,D3D_FEATURE_LEVEL InFeatureLevel,int32 InChosenAdapter);
 
 	/** Destructor */
-	virtual ~FD3D11DynamicRHI();
+	virtual ~FD3D11DynamicRHI() {}
 
 	/** If it hasn't been initialized yet, initializes the D3D device. */
 	virtual void InitD3DDevice();
 
 	// FDynamicRHI interface.
 	virtual void Init() OVERRIDE;
+	virtual void Shutdown() OVERRIDE;
 	virtual void PushEvent(const TCHAR* Name) OVERRIDE { GPUProfilingData.PushEvent(Name); }
 	virtual void PopEvent() OVERRIDE { GPUProfilingData.PopEvent(); }
 
