Enable ES 3.0 GLSL on all Android devices reporting ES 3.0 support.

- Fixes some GLSL errors that Adreno doesn't care about but Mali does.
  - gl_FragColor declaration before #extension
  - missing default precision on gl_FragColor
  - Declaring the output named gl_FragColor when variables must not start with gl_

diff --git a/Engine/Source/Runtime/OpenGLDrv/Private/Android/AndroidOpenGL.cpp b/Engine/Source/Runtime/OpenGLDrv/Private/Android/AndroidOpenGL.cpp
--- a/Engine/Source/Runtime/OpenGLDrv/Private/Android/AndroidOpenGL.cpp
+++ b/Engine/Source/Runtime/OpenGLDrv/Private/Android/AndroidOpenGL.cpp
@@ -293,6 +293,7 @@ void FAndroidOpenGL::ProcessExtensions(const FString& ExtensionsString)
 	glGetObjectLabelEXT = (PFNGLGETOBJECTLABELEXTPROC)((void*)eglGetProcAddress("glGetObjectLabelEXT"));
 
 	bSupportsETC2 = bES30Support;
+	bUseES30ShadingLanguage = bES30Support;
 
 	// Attempt to find ES 3.0 glTexStorage2D if we're on an Adreno device that supports it.
 	if( FString(ANSI_TO_TCHAR((const ANSICHAR*)glGetString(GL_RENDERER))).Contains(TEXT("Adreno")) )
@@ -301,10 +302,6 @@ void FAndroidOpenGL::ProcessExtensions(const FString& ExtensionsString)
 		if( glTexStorage2D != NULL )
 		{
 			bUseAdrenoHalfFloatTexStorage = true;
-			if( bES30Support )
-			{
-				bUseES30ShadingLanguage = true;
-			}
 		}
 		else
 		{
diff --git a/Engine/Source/Runtime/OpenGLDrv/Private/OpenGLShaders.cpp b/Engine/Source/Runtime/OpenGLDrv/Private/OpenGLShaders.cpp
--- a/Engine/Source/Runtime/OpenGLDrv/Private/OpenGLShaders.cpp
+++ b/Engine/Source/Runtime/OpenGLDrv/Private/OpenGLShaders.cpp
@@ -149,6 +149,45 @@ static const TCHAR* ShaderNameFromShaderType(GLenum ShaderType)
 // ============================================================================================================================
 
 
+static ANSICHAR* SkipShaderExtensionText(ANSICHAR* ShaderCode)
+{
+	char* ExtensionEnd = ShaderCode;
+	for (;;)
+	{
+		// See if we have an extension
+		char* Find = (char*)strstr(ExtensionEnd, "#extension");
+		if (Find)
+		{
+			// Skip to the end of the line
+			Find = (char*)strstr(Find, "\n");
+			if (Find == NULL)
+			{
+				// Abort if there is no end of line
+				return ShaderCode;
+			}
+			ExtensionEnd = Find;
+		}
+		else
+		{
+			break;
+		}
+	}
+
+	if (ExtensionEnd != ShaderCode)
+	{
+		// terminate the extension part of the shader code
+		*ExtensionEnd = '\0';
+		return ExtensionEnd + 1;
+	}
+	else
+	{
+		return NULL;
+	}
+}
+
+// ============================================================================================================================
+
+
 static void ReplaceShaderSubstring(ANSICHAR* ShaderCode, const ANSICHAR* OldString, const ANSICHAR* NewString)
 {
 	char* Curr = ShaderCode;
@@ -266,6 +305,19 @@ ShaderType* CompileOpenGLShader(const TArray<uint8>& Code)
 			strcpy(VersionString, "#version 100");
 			strcat(VersionString, "\n");
 		}
+
+		//remove "#version 100" line,  if found in existing GlslCode. 
+		ANSICHAR* VersionCodePointer = const_cast<ANSICHAR*>(GlslCode);
+		VersionCodePointer = strstr(VersionCodePointer, "#version 100");
+		if(VersionCodePointer != NULL)
+		{
+			for(int index=0; index < 12; index++)
+			{
+				VersionCodePointer[index] = ' ';
+			}
+		}
+
+		const ANSICHAR* ExtensionString = "";
 		
 		//This #define fixes compiler errors on Android (which doesnt seem to support textureCubeLodEXT)
 		const ANSICHAR* Prologue = NULL; 
@@ -286,7 +338,8 @@ ShaderType* CompileOpenGLShader(const TArray<uint8>& Code)
 			} 
 			else if (TypeEnum == GL_FRAGMENT_SHADER)
 			{
-				Prologue = "#define texture2D texture				\n"
+				Prologue = "\n"
+					"#define texture2D texture				\n"
 					"#define texture2DProj textureProj		\n"
 					"#define texture2DLod textureLod			\n"
 					"#define texture2DProjLod textureProjLod \n"
@@ -294,7 +347,16 @@ ShaderType* CompileOpenGLShader(const TArray<uint8>& Code)
 					"#define textureCubeLod textureLod		\n"
 					"#define textureCubeLodEXT textureLod	\n"
 					"\n"
-					"out vec4 gl_FragColor;";
+					"#define gl_FragColor out_FragColor	\n"
+					"out mediump vec4 out_FragColor;";
+
+				// See if we need to skip any #extension string
+				ANSICHAR* Temp = SkipShaderExtensionText(const_cast<ANSICHAR*>(GlslCode));
+				if (Temp)
+				{
+					ExtensionString = GlslCode;
+					GlslCode = Temp;
+				}
 
 				ReplaceShaderSubstring(const_cast<ANSICHAR*>(GlslCode), "varying", "in");
 			}
@@ -311,20 +373,10 @@ ShaderType* CompileOpenGLShader(const TArray<uint8>& Code)
 		{
 			Prologue = "";
 		}
-		//remove "#version 100" line,  if found in existing GlslCode. 
-		ANSICHAR* VersionCodePointer = const_cast<ANSICHAR*>(GlslCode);
-		VersionCodePointer = strstr(VersionCodePointer, "#version 100");
-		if(VersionCodePointer != NULL)
-		{
-			for(int index=0; index < 12; index++)
-			{
-				VersionCodePointer[index] = ' ';
-			}
-		}
 		//Assemble the source strings into an array to pass into the compiler.
-		const GLchar* ShaderSourceStrings[3] = { VersionString, Prologue, GlslCode };
-		const GLint ShaderSourceLen[3] = { (GLint)(strlen(VersionString)), (GLint)(strlen(Prologue)), (GLint)(strlen(GlslCode)) };
-		glShaderSource(Resource, 3, (const GLchar**)&ShaderSourceStrings, ShaderSourceLen);		
+		const GLchar* ShaderSourceStrings[4] = { VersionString, ExtensionString, Prologue, GlslCode };
+		const GLint ShaderSourceLen[4] = { (GLint)(strlen(VersionString)), (GLint)(strlen(ExtensionString)), (GLint)(strlen(Prologue)), (GLint)(strlen(GlslCode)) };
+		glShaderSource(Resource, 4, (const GLchar**)&ShaderSourceStrings, ShaderSourceLen);		
 		glCompileShader(Resource);
 #else
 		glShaderSource(Resource, 1, (const GLchar**)&GlslCode, &GlslCodeLength);
