Check for textureCubeLod availability to address Android on Mali-400 GPUs.

diff --git a/Engine/Source/Runtime/OpenGLDrv/Private/Android/AndroidOpenGL.cpp b/Engine/Source/Runtime/OpenGLDrv/Private/Android/AndroidOpenGL.cpp
--- a/Engine/Source/Runtime/OpenGLDrv/Private/Android/AndroidOpenGL.cpp
+++ b/Engine/Source/Runtime/OpenGLDrv/Private/Android/AndroidOpenGL.cpp
@@ -323,6 +323,12 @@ void FAndroidOpenGL::ProcessExtensions(const FString& ExtensionsString)
 	//On Android, there are problems compiling shaders with textureCubeLodEXT calls in the glsl code,
 	// so we set this to false to modify the glsl manually at compile-time.
 	bSupportsTextureCubeLodEXT = false;
+
+	// On some Android devices with Mali GPUs textureCubeLod is not available.
+	if( FString(ANSI_TO_TCHAR((const ANSICHAR*)glGetString(GL_RENDERER))).Contains(TEXT("Mali-400")) )
+	{
+		bSupportsShaderTextureCubeLod = false;
+	}
 }
 
 class FAndroidGPUInfo
diff --git a/Engine/Source/Runtime/OpenGLDrv/Private/OpenGLES2.cpp b/Engine/Source/Runtime/OpenGLDrv/Private/OpenGLES2.cpp
--- a/Engine/Source/Runtime/OpenGLDrv/Private/OpenGLES2.cpp
+++ b/Engine/Source/Runtime/OpenGLDrv/Private/OpenGLES2.cpp
@@ -95,6 +95,9 @@ bool FOpenGLES2::bSupportsTextureCubeLodEXT = true;
 /** GL_EXT_shader_texture_lod */
 bool FOpenGLES2::bSupportsShaderTextureLod = false;
 
+/** textureCubeLod */
+bool FOpenGLES2::bSupportsShaderTextureCubeLod = true;
+
 /** GL_APPLE_copy_texture_levels */
 bool FOpenGLES2::bSupportsCopyTextureLevels = false;
 
diff --git a/Engine/Source/Runtime/OpenGLDrv/Private/OpenGLShaders.cpp b/Engine/Source/Runtime/OpenGLDrv/Private/OpenGLShaders.cpp
--- a/Engine/Source/Runtime/OpenGLDrv/Private/OpenGLShaders.cpp
+++ b/Engine/Source/Runtime/OpenGLDrv/Private/OpenGLShaders.cpp
@@ -361,7 +361,7 @@ ShaderType* CompileOpenGLShader(const TArray<uint8>& Code)
 				ReplaceShaderSubstring(const_cast<ANSICHAR*>(GlslCode), "varying", "in");
 			}
 		}
-		else if(!FOpenGL::SupportsShaderTextureLod())
+		else if(!FOpenGL::SupportsShaderTextureLod() || !FOpenGL::SupportsShaderTextureCubeLod())
 		{
 			Prologue = "#define textureCubeLodEXT(a, b, c) textureCube(a, b) \n";
 		}
diff --git a/Engine/Source/Runtime/OpenGLDrv/Public/OpenGL.h b/Engine/Source/Runtime/OpenGLDrv/Public/OpenGL.h
--- a/Engine/Source/Runtime/OpenGLDrv/Public/OpenGL.h
+++ b/Engine/Source/Runtime/OpenGLDrv/Public/OpenGL.h
@@ -119,6 +119,7 @@ public:
 	static FORCEINLINE bool SupportsPackedDepthStencil()				{ return true; }
 	static FORCEINLINE bool SupportsTextureCubeLodEXT()					{ return true; }
 	static FORCEINLINE bool SupportsShaderTextureLod()					{ return false; }
+	static FORCEINLINE bool SupportsShaderTextureCubeLod()				{ return true; }
 	static FORCEINLINE bool SupportsSeparateAlphaBlend()				{ return false; }
 	static FORCEINLINE bool SupportsTessellation()						{ return false; }
 	static FORCEINLINE bool SupportsComputeShaders()					{ return false; }
diff --git a/Engine/Source/Runtime/OpenGLDrv/Public/OpenGLES2.h b/Engine/Source/Runtime/OpenGLDrv/Public/OpenGLES2.h
--- a/Engine/Source/Runtime/OpenGLDrv/Public/OpenGLES2.h
+++ b/Engine/Source/Runtime/OpenGLDrv/Public/OpenGLES2.h
@@ -103,6 +103,7 @@ struct FOpenGLES2 : public FOpenGLBase
 	static FORCEINLINE bool SupportsPackedDepthStencil()				{ return bSupportsPackedDepthStencil; }
 	static FORCEINLINE bool SupportsTextureCubeLodEXT()					{ return bSupportsTextureCubeLodEXT; }
 	static FORCEINLINE bool SupportsShaderTextureLod()					{ return bSupportsShaderTextureLod; }
+	static FORCEINLINE bool SupportsShaderTextureCubeLod()				{ return bSupportsShaderTextureCubeLod; }
 	static FORCEINLINE bool SupportsCopyTextureLevels()					{ return bSupportsCopyTextureLevels; }
 
 	// On iOS both glMapBufferOES() and glBufferSubData() for immediate vertex and index data
@@ -408,6 +409,9 @@ protected:
 	/** GL_EXT_shader_texture_lod */
 	static bool bSupportsShaderTextureLod;
 
+	/** textureCubeLod */
+	static bool bSupportsShaderTextureCubeLod;
+
 	/** GL_APPLE_copy_texture_levels */
 	static bool bSupportsCopyTextureLevels;
 
