Fixed startup crash with -opengl on low-end GPUs

diff --git a/Engine/Source/Runtime/OpenGLDrv/Private/OpenGL4.cpp b/Engine/Source/Runtime/OpenGLDrv/Private/OpenGL4.cpp
--- a/Engine/Source/Runtime/OpenGLDrv/Private/OpenGL4.cpp
+++ b/Engine/Source/Runtime/OpenGLDrv/Private/OpenGL4.cpp
@@ -24,7 +24,7 @@ void FOpenGL4::ProcessQueryGLInt()
 	LOG_AND_GET_GL_INT(GL_MAX_COMPUTE_UNIFORM_COMPONENTS, 0, MaxComputeUniformComponents);
 	LOG_AND_GET_GL_INT(GL_MAX_TESS_CONTROL_UNIFORM_COMPONENTS, 0, MaxHullUniformComponents);
 	LOG_AND_GET_GL_INT(GL_MAX_TESS_EVALUATION_UNIFORM_COMPONENTS, 0, MaxDomainUniformComponents);
-	if ( MaxComputeTextureImageUnits == -1 || MaxComputeUniformComponents == -1 || MaxHullUniformComponents == -1 || MaxDomainUniformComponents == -1 )
+	if ( MaxComputeTextureImageUnits == 0 || MaxComputeUniformComponents == 0 || MaxHullUniformComponents == 0 || MaxDomainUniformComponents == 0 )
 	{
 		bSupportsComputeShaders = false;
 	}
