Cleaned up frame timing for the renderthread. It was wrong for OpenGL platforms and could be for other platforms as well. Now it measured cross-platform instead.

diff --git a/Engine/Source/Runtime/Engine/Private/UnrealClient.cpp b/Engine/Source/Runtime/Engine/Private/UnrealClient.cpp
--- a/Engine/Source/Runtime/Engine/Private/UnrealClient.cpp
+++ b/Engine/Source/Runtime/Engine/Private/UnrealClient.cpp
@@ -885,8 +885,12 @@ void FViewport::EndRenderFrame( bool bPresent, bool bLockToVsync )
 {
 	check( IsInRenderingThread() );
 
-
+	uint32 StartTime = FPlatformTime::Cycles();
 	RHIEndDrawingViewport( GetViewportRHI(), bPresent, bLockToVsync );
+	uint32 EndTime = FPlatformTime::Cycles();
+
+	GRenderThreadIdle[ERenderThreadIdleTypes::WaitingForGPUPresent] += EndTime - StartTime;
+	GRenderThreadNumIdle[ERenderThreadIdleTypes::WaitingForGPUPresent]++;
 }
 
 void InsertVolume(IInterface_PostProcessVolume* Volume, TArray< IInterface_PostProcessVolume* >& VolumeArray)
diff --git a/Engine/Source/Runtime/OpenGLDrv/Private/Linux/OpenGLLinux.cpp b/Engine/Source/Runtime/OpenGLDrv/Private/Linux/OpenGLLinux.cpp
--- a/Engine/Source/Runtime/OpenGLDrv/Private/Linux/OpenGLLinux.cpp
+++ b/Engine/Source/Runtime/OpenGLDrv/Private/Linux/OpenGLLinux.cpp
@@ -350,8 +350,6 @@ void PlatformBlitToViewport(FPlatformOpenGLDevice* Device,
 
 		if ( bPresent )
 		{
-			uint32 IdleStart = FPlatformTime::Cycles();
-
 			int32 RealSyncInterval = bLockToVsync ? SyncInterval : 0;
 
 			if ( Context->SyncInterval != RealSyncInterval)
@@ -367,10 +365,7 @@ void PlatformBlitToViewport(FPlatformOpenGLDevice* Device,
 			SDL_GL_SwapWindow( Context->hWnd );
 
 			REPORT_GL_END_BUFFER_EVENT_FOR_FRAME_DUMP();
-			//			INITIATE_GL_FRAME_DUMP_EVERY_X_CALLS( 1000 );
-
-			GRenderThreadIdle[ERenderThreadIdleTypes::WaitingForGPUPresent] += FPlatformTime::Cycles() - IdleStart;
-			GRenderThreadNumIdle[ERenderThreadIdleTypes::WaitingForGPUPresent]++;
+//			INITIATE_GL_FRAME_DUMP_EVERY_X_CALLS( 1000 );
 		}
 	}
 }
diff --git a/Engine/Source/Runtime/OpenGLDrv/Private/Mac/OpenGLMac.cpp b/Engine/Source/Runtime/OpenGLDrv/Private/Mac/OpenGLMac.cpp
--- a/Engine/Source/Runtime/OpenGLDrv/Private/Mac/OpenGLMac.cpp
+++ b/Engine/Source/Runtime/OpenGLDrv/Private/Mac/OpenGLMac.cpp
@@ -816,8 +816,6 @@ void PlatformBlitToViewport( FPlatformOpenGLDevice* Device, FPlatformOpenGLConte
 					glColorMask(GL_TRUE,GL_TRUE,GL_TRUE,GL_TRUE);
 				}
 
-				uint32 IdleStart = FPlatformTime::Cycles();
-
 				int32 RealSyncInterval = bLockToVsync ? SyncInterval : 0;
 
 				if (Context->SyncInterval != RealSyncInterval)
@@ -837,9 +835,6 @@ void PlatformBlitToViewport( FPlatformOpenGLDevice* Device, FPlatformOpenGLConte
 					// Just hope that this will work OK...
 					[SlateCocoaWindow performDeferredOrderFront];
 				}
-				
-				GRenderThreadIdle[ERenderThreadIdleTypes::WaitingForGPUPresent] += FPlatformTime::Cycles() - IdleStart;
-				GRenderThreadNumIdle[ERenderThreadIdleTypes::WaitingForGPUPresent]++;
 			}
 		}
 	}
diff --git a/Engine/Source/Runtime/OpenGLDrv/Private/OpenGLViewport.cpp b/Engine/Source/Runtime/OpenGLDrv/Private/OpenGLViewport.cpp
--- a/Engine/Source/Runtime/OpenGLDrv/Private/OpenGLViewport.cpp
+++ b/Engine/Source/Runtime/OpenGLDrv/Private/OpenGLViewport.cpp
@@ -116,8 +116,6 @@ void FOpenGLDynamicRHI::RHIEndDrawingViewport(FViewportRHIParamRef ViewportRHI,b
 
 	FOpenGLTexture2D* BackBuffer = Viewport->GetBackBuffer();
 
-	uint32 IdleStart = FPlatformTime::Cycles();
-
 	PlatformBlitToViewport(PlatformDevice,
 		Viewport->OpenGLContext,
 		BackBuffer->GetSizeX(),
@@ -157,8 +155,6 @@ void FOpenGLDynamicRHI::RHIEndDrawingViewport(FViewportRHIParamRef ViewportRHI,b
 		}
 	}
 
-	GRenderThreadIdle[ERenderThreadIdleTypes::WaitingForGPUPresent] += FPlatformTime::Cycles() - IdleStart;
-
 	if (bRevertToSharedContextAfterDrawingViewport)
 	{
 		PlatformSharedContextSetup(PlatformDevice);
diff --git a/Engine/Source/Runtime/OpenGLDrv/Private/Windows/OpenGLWindows.cpp b/Engine/Source/Runtime/OpenGLDrv/Private/Windows/OpenGLWindows.cpp
--- a/Engine/Source/Runtime/OpenGLDrv/Private/Windows/OpenGLWindows.cpp
+++ b/Engine/Source/Runtime/OpenGLDrv/Private/Windows/OpenGLWindows.cpp
@@ -421,8 +421,6 @@ void PlatformBlitToViewport( FPlatformOpenGLDevice* Device, FPlatformOpenGLConte
 
 		if (bPresent)
 		{
-			uint32 IdleStart = FPlatformTime::Cycles();
-
 			int32 RealSyncInterval = bLockToVsync ? SyncInterval : 0;
 			if (wglSwapIntervalEXT_ProcAddress && Context->SyncInterval != RealSyncInterval)
 			{
@@ -433,9 +431,6 @@ void PlatformBlitToViewport( FPlatformOpenGLDevice* Device, FPlatformOpenGLConte
 			::SwapBuffers(Context->DeviceContext);
 			REPORT_GL_END_BUFFER_EVENT_FOR_FRAME_DUMP();
 //			INITIATE_GL_FRAME_DUMP_EVERY_X_CALLS( 1000 );
-
-			GRenderThreadIdle[ERenderThreadIdleTypes::WaitingForGPUPresent] += FPlatformTime::Cycles() - IdleStart;
-			GRenderThreadNumIdle[ERenderThreadIdleTypes::WaitingForGPUPresent]++;
 		}
 	}
 }
diff --git a/Engine/Source/Runtime/SlateRHIRenderer/Private/SlateRHIRenderer.cpp b/Engine/Source/Runtime/SlateRHIRenderer/Private/SlateRHIRenderer.cpp
--- a/Engine/Source/Runtime/SlateRHIRenderer/Private/SlateRHIRenderer.cpp
+++ b/Engine/Source/Runtime/SlateRHIRenderer/Private/SlateRHIRenderer.cpp
@@ -483,6 +483,8 @@ void FSlateRHIRenderer::DrawWindow_RenderThread( const FViewportInfo& ViewportIn
 
 	FThreadIdleStats& RenderThread = FThreadIdleStats::Get();
 	GRenderThreadIdle[ERenderThreadIdleTypes::WaitingForAllOtherSleep] = RenderThread.Waits;
+	GRenderThreadIdle[ERenderThreadIdleTypes::WaitingForGPUPresent] += GSwapBufferTime;
+	GRenderThreadNumIdle[ERenderThreadIdleTypes::WaitingForGPUPresent]++;
 	RenderThread.Waits = 0;
 
 	SET_CYCLE_COUNTER(STAT_RenderingIdleTime_RenderThreadSleepTime, GRenderThreadIdle[0]);
diff --git a/Engine/Source/Runtime/Windows/D3D11RHI/Private/D3D11Viewport.cpp b/Engine/Source/Runtime/Windows/D3D11RHI/Private/D3D11Viewport.cpp
--- a/Engine/Source/Runtime/Windows/D3D11RHI/Private/D3D11Viewport.cpp
+++ b/Engine/Source/Runtime/Windows/D3D11RHI/Private/D3D11Viewport.cpp
@@ -515,10 +515,7 @@ void FD3D11DynamicRHI::RHIEndDrawingViewport(FViewportRHIParamRef ViewportRHI,bo
 
 	if(bPresent)
 	{
-		uint32 IdleStart = FPlatformTime::Cycles();
 		Viewport->Present(bLockToVsync);
-		GRenderThreadIdle[ERenderThreadIdleTypes::WaitingForGPUPresent] += FPlatformTime::Cycles() - IdleStart;
-		GRenderThreadNumIdle[ERenderThreadIdleTypes::WaitingForGPUPresent]++;
 	}
 
 	// Don't wait on the GPU when using SLI, let the driver determine how many frames behind the GPU should be allowed to get
