Stop deleting textures bound to ShaderResourceViews out from underneath our FOpenGLTexture classes, fixing glitches with screen space reflections.
The OpenGL ShaderResourceView only need delete TEXTURE_BUFFER target textures, the others will be disposed of with their parent FOpenGLTexture. This will be replaced by a superior fix as part of adding GL_texture_view support.

diff --git a/Engine/Source/Runtime/OpenGLDrv/Private/OpenGLUAV.cpp b/Engine/Source/Runtime/OpenGLDrv/Private/OpenGLUAV.cpp
--- a/Engine/Source/Runtime/OpenGLDrv/Private/OpenGLUAV.cpp
+++ b/Engine/Source/Runtime/OpenGLDrv/Private/OpenGLUAV.cpp
@@ -36,7 +36,9 @@ FShaderResourceViewRHIRef FOpenGLDynamicRHI::RHICreateShaderResourceView(FVertex
 
 FOpenGLShaderResourceView::~FOpenGLShaderResourceView()
 {
-	if (Resource)
+    // Don't delete GL resource out from underneath FOpenGLTexture for other targets.
+    // For D3D RHI the texture object is actually retained by the SRV, we may want to match those semantics here too.
+	if (Resource && Target == GL_TEXTURE_BUFFER)
 	{
 		OpenGLRHI->InvalidateTextureResourceInCache( Resource );
 		glDeleteTextures(1, &Resource);
