On Mac OS X clear the depth & stencil separately to avoid a bug first found in 2011 with the ATi 2xx0 & 4xx0 series drivers which causes the combined depth/stencil clear to fail to clear stencil.
This was causing incorrect particle rendering as noted in #ttp 331090 "SHOWCASES: EFFECTS: Mac: Snow particles are black".

diff --git a/Engine/Source/Runtime/OpenGLDrv/Private/OpenGLCommands.cpp b/Engine/Source/Runtime/OpenGLDrv/Private/OpenGLCommands.cpp
--- a/Engine/Source/Runtime/OpenGLDrv/Private/OpenGLCommands.cpp
+++ b/Engine/Source/Runtime/OpenGLDrv/Private/OpenGLCommands.cpp
@@ -2333,6 +2333,30 @@ void FOpenGLDynamicRHI::RHIClear(bool bClearColor,const FLinearColor& Color,bool
 
 static inline void ClearCurrentDepthStencilWithCurrentScissor( int8 ClearType, float Depth, uint32 Stencil )
 {
+#if PLATFORM_MAC
+	switch (ClearType)
+	{
+        case CT_DepthStencil:	// Clear depth and stencil separately to avoid an AMD Dx00 bug which causes depth to clear, but not stencil.
+                                // Especially irritatingly this bug will not manifest when stepping though the program with GL Profiler.
+                                // This was a bug found by me during dev. on Tropico 3 circa. Q4 2011 in the ATi Mac 2xx0 & 4xx0 drivers.
+                                // It was never fixed & has re-emerged in the AMD Mac FirePro Dx00 drivers.
+        case CT_Stencil:	// Clear stencil only
+            FOpenGL::ClearBufferiv(GL_STENCIL, 0, (const GLint*)&Stencil);
+            
+            // If not also clearing depth break
+            if(!(ClearType & CT_Depth))
+            {
+                break;
+            }
+            // Otherwise fall through to perform a separate depth clear.
+        case CT_Depth:	// Clear depth only
+            FOpenGL::ClearBufferfv(GL_DEPTH, 0, &Depth);
+            break;
+            
+        default:
+            break;	// impossible anyway
+	}
+#else
 	switch (ClearType)
 	{
 	case CT_DepthStencil:	// Clear depth and stencil
@@ -2350,6 +2374,7 @@ static inline void ClearCurrentDepthStencilWithCurrentScissor( int8 ClearType, f
 	default:
 		break;	// impossible anyway
 	}
+#endif
 }
 
 void FOpenGLDynamicRHI::ClearCurrentFramebufferWithCurrentScissor(FOpenGLContextState& ContextState, int8 ClearType, int32 NumClearColors, const FLinearColor* ClearColorArray, float Depth, uint32 Stencil)
