Disable sampling from a depth/stencil texture while also depth or stencil testing against it on the AMD FirePro Dx00 cards on OS X to fix decal rendering errors.
Test the current CGL Renderer ID and if using the new kCGLRendererATIRadeonX4000ID renderer (which covers FirePro D300, 500, 700 and Radeon 7950 Mac cards) disable GSupportsDepthFetchDuringDepthTest to fix the decal rendering which samples depth from the depth-stencil surface while it is bound for stencil testing (but not depth testing).
Also when the GL_EXT_debug_marker extension is present enable GEmitDrawEvents in non-shipping builds to make debugging vastly easier.

diff --git a/Engine/Source/Runtime/OpenGLDrv/Private/Mac/OpenGLMac.cpp b/Engine/Source/Runtime/OpenGLDrv/Private/Mac/OpenGLMac.cpp
--- a/Engine/Source/Runtime/OpenGLDrv/Private/Mac/OpenGLMac.cpp
+++ b/Engine/Source/Runtime/OpenGLDrv/Private/Mac/OpenGLMac.cpp
@@ -607,6 +607,22 @@ FPlatformOpenGLContext* PlatformCreateOpenGLContext(FPlatformOpenGLDevice* Devic
 		GIsRunningOnIntelCard = true;
 	}
 	
+	// Renderer IDs matchup to driver kexts, so switching based on them will allow us to target workarouds to many GPUs
+	// which exhibit the same unfortunate driver bugs without having to parse their individual ID strings.
+	GLint RendererID = 0;
+	if(CGLGetParameter((CGLContextObj)[Context->OpenGLContext CGLContextObj], kCGLCPCurrentRendererID, &RendererID) == kCGLNoError)
+	{
+		switch((RendererID & kCGLRendererIDMatchingMask))
+		{
+			case kCGLRendererATIRadeonX4000ID:
+			{
+				// @todo: remove once AMD fix the AMDX4000 driver for GCN cards so that it is possible to sample the depth while also stencil testing to the same DEPTH_STENCIL texture - it works on all other cards we have.
+				GSupportsDepthFetchDuringDepthTest = false;
+				break;
+			}
+		}
+    }
+	
 #if UE_EMULATE_TIMESTAMP
 	// Bind the platform context into the CGL context
 	Context->TimeElapsed = 0;
@@ -690,12 +706,12 @@ void PlatformBlitToViewport( FPlatformOpenGLDevice* Device, FPlatformOpenGLConte
 			glBindFramebuffer(GL_READ_FRAMEBUFFER, Context->ViewportFramebuffer);
 			glReadBuffer(GL_COLOR_ATTACHMENT0);
 
-			glBlitFramebuffer(
+            glBlitFramebuffer(
 				0, 0, BackbufferSizeX, BackbufferSizeY,
 				0, BackbufferSizeY, BackbufferSizeX, 0,
 				GL_COLOR_BUFFER_BIT,
 				GL_NEAREST
-				);
+            );
 			
 			if(Context->OpenGLView.bNeedsUpdate)
 			{
@@ -825,6 +841,7 @@ void PlatformBlitToViewport( FPlatformOpenGLDevice* Device, FPlatformOpenGLConte
 				}
 
 				[Context->OpenGLContext flushBuffer];
+
 				REPORT_GL_END_BUFFER_EVENT_FOR_FRAME_DUMP();
 //				INITIATE_GL_FRAME_DUMP_EVERY_X_CALLS( 1000 );
 				
@@ -1379,6 +1396,9 @@ void FMacOpenGL::ProcessExtensions(const FString& ExtensionsString)
 	{
 		glPushGroupMarkerEXT = (PFNGLPUSHGROUPMARKEREXTPROC)dlsym(RTLD_SELF, "glPushGroupMarkerEXT");
 		glPopGroupMarkerEXT = (PFNGLPOPGROUPMARKEREXTPROC)dlsym(RTLD_SELF, "glPopGroupMarkerEXT");
+#if !UE_BUILD_SHIPPING // For debuggable builds emit draw events when the extension is GL_EXT_debug_marker present.
+        GEmitDrawEvents = true;
+#endif
 	}
 	
 	if(ExtensionsString.Contains(TEXT("GL_ARB_tessellation_shader")))
