Don't use an sRGB render target on AMD GPUs under Mac OS X by default. There is presently a bug that means that the diffuse rendering ends up being linear even though we've requested sRGB. Non-AMD GPUs work correctly.
This is handled by a console variable, so you can turn it on or off & see the effect this bug has - making it easier to debug. Setting this value in an .ini file will override the automatic detection as well.
The underlying bug has been filed as radr://16754329 AMD Cards don't always perform FRAMEBUFFER_SRGB if the draw FBO has mixed sRGB & non-SRGB colour attachments

diff --git a/Engine/Source/Runtime/OpenGLDrv/Private/Mac/OpenGLMac.cpp b/Engine/Source/Runtime/OpenGLDrv/Private/Mac/OpenGLMac.cpp
--- a/Engine/Source/Runtime/OpenGLDrv/Private/Mac/OpenGLMac.cpp
+++ b/Engine/Source/Runtime/OpenGLDrv/Private/Mac/OpenGLMac.cpp
@@ -369,6 +369,14 @@ struct OpenGLContextInfo
 
 extern void OnQueryInvalidation( void );
 
+// @todo: remove once Apple fixes radr://16754329 AMD Cards don't always perform FRAMEBUFFER_SRGB if the draw FBO has mixed sRGB & non-SRGB colour attachments
+static TAutoConsoleVariable<int32> CVarMacUseFrameBufferSRGB(
+		TEXT("r.Mac.UseFrameBufferSRGB"),
+		-1,
+		TEXT("Flag to toggle use of GL_FRAMEBUFFER_SRGB for better color accuracy.\n"),
+		ECVF_RenderThreadSafe
+		);
+
 bool GIsRunningOnIntelCard = false; // @todo: remove once Apple fixes radr://16223045 Changes to the GL separate blend state aren't always respected on Intel cards
 
 struct FPlatformOpenGLDevice
@@ -606,6 +614,18 @@ FPlatformOpenGLContext* PlatformCreateOpenGLContext(FPlatformOpenGLDevice* Devic
 	{
 		GIsRunningOnIntelCard = true;
 	}
+	else if (VendorName.Contains(TEXT("AMD ")) || VendorName.Contains(TEXT("ATI ")))
+	{
+		if(CVarMacUseFrameBufferSRGB.GetValueOnRenderThread() == -1)
+		{
+			CVarMacUseFrameBufferSRGB.AsVariable()->Set(0);
+		}
+	}
+	
+	if(CVarMacUseFrameBufferSRGB.GetValueOnRenderThread() == -1)
+	{
+		CVarMacUseFrameBufferSRGB.AsVariable()->Set(1);
+	}
 	
 	// Renderer IDs matchup to driver kexts, so switching based on them will allow us to target workarouds to many GPUs
 	// which exhibit the same unfortunate driver bugs without having to parse their individual ID strings.
@@ -705,6 +725,7 @@ void PlatformBlitToViewport( FPlatformOpenGLDevice* Device, FPlatformOpenGLConte
 			glDrawBuffer(GL_BACK);
 			glBindFramebuffer(GL_READ_FRAMEBUFFER, Context->ViewportFramebuffer);
 			glReadBuffer(GL_COLOR_ATTACHMENT0);
+			glDisable(GL_FRAMEBUFFER_SRGB);
 
             glBlitFramebuffer(
 				0, 0, BackbufferSizeX, BackbufferSizeY,
@@ -842,6 +863,9 @@ void PlatformBlitToViewport( FPlatformOpenGLDevice* Device, FPlatformOpenGLConte
 
 				[Context->OpenGLContext flushBuffer];
 
+				
+				glEnable(GL_FRAMEBUFFER_SRGB);
+				
 				REPORT_GL_END_BUFFER_EVENT_FOR_FRAME_DUMP();
 //				INITIATE_GL_FRAME_DUMP_EVERY_X_CALLS( 1000 );
 				
diff --git a/Engine/Source/Runtime/Renderer/Private/PostProcess/SceneRenderTargets.cpp b/Engine/Source/Runtime/Renderer/Private/PostProcess/SceneRenderTargets.cpp
--- a/Engine/Source/Runtime/Renderer/Private/PostProcess/SceneRenderTargets.cpp
+++ b/Engine/Source/Runtime/Renderer/Private/PostProcess/SceneRenderTargets.cpp
@@ -850,8 +850,14 @@ void FSceneRenderTargets::AllocateDeferredShadingPathRenderTargets()
 		// Create the diffuse color g-buffer.
 		{
 			const EPixelFormat DiffuseGBufferFormat = bHighPrecisionGBuffers ? PF_FloatRGBA : PF_B8G8R8A8;
+			uint32 DiffuseGBufferFlags = TexCreate_SRGB;
 
-			FPooledRenderTargetDesc Desc(FPooledRenderTargetDesc::Create2DDesc(BufferSize, DiffuseGBufferFormat, TexCreate_SRGB, TexCreate_RenderTargetable, false));
+#if PLATFORM_MAC // @todo: remove once Apple fixes radr://16754329 AMD Cards don't always perform FRAMEBUFFER_SRGB if the draw FBO has mixed sRGB & non-SRGB colour attachments
+			static const auto CVar = IConsoleManager::Get().FindTConsoleVariableDataInt(TEXT("r.Mac.UseFrameBufferSRGB"));
+			DiffuseGBufferFlags = CVar && CVar->GetValueOnRenderThread() ? TexCreate_SRGB : TexCreate_None;
+#endif
+
+			FPooledRenderTargetDesc Desc(FPooledRenderTargetDesc::Create2DDesc(BufferSize, DiffuseGBufferFormat, DiffuseGBufferFlags, TexCreate_RenderTargetable, false));
 			GRenderTargetPool.FindFreeElement(Desc, GBufferC, TEXT("GBufferC"));
 		}
 
