Prevent crash setting a cube depth-stencil surface for BeginRenderingCubeShadowDepth.
Use the OpenGL specific X & Y texture dimension accessors to get the required viewport size rather than casting to Texture2D here.

diff --git a/Engine/Source/Runtime/OpenGLDrv/Private/OpenGLCommands.cpp b/Engine/Source/Runtime/OpenGLDrv/Private/OpenGLCommands.cpp
--- a/Engine/Source/Runtime/OpenGLDrv/Private/OpenGLCommands.cpp
+++ b/Engine/Source/Runtime/OpenGLDrv/Private/OpenGLCommands.cpp
@@ -1580,9 +1580,8 @@ void FOpenGLDynamicRHI::RHISetRenderTargets(
 	else if( NewDepthStencilTargetRHI )
 	{
 		// Set viewport size to new depth target size.
-		FOpenGLTexture2D* NewDepthTargetTexture = (FOpenGLTexture2D*)NewDepthStencilTargetRHI->GetTexture2D();
-		PendingState.Viewport.Max.X = NewDepthTargetTexture->GetSizeX();
-		PendingState.Viewport.Max.Y = NewDepthTargetTexture->GetSizeY();
+		PendingState.Viewport.Max.X = GetOpenGLTextureSizeXFromRHITexture(NewDepthStencilTargetRHI);
+		PendingState.Viewport.Max.Y = GetOpenGLTextureSizeYFromRHITexture(NewDepthStencilTargetRHI);
 	}
 }
 
