Fix FOpenGLDynamicRHI::ReadSurfaceDataRaw so that reading back from a float texture doesn't introduce a yellow tint to the retrieved image. This was causing asset thumbnails to render incorrectly on OpenGL platforms. The minimum & maximum value arrays were being incorrectly initialised.

diff --git a/Engine/Source/Runtime/OpenGLDrv/Private/OpenGLRenderTarget.cpp b/Engine/Source/Runtime/OpenGLDrv/Private/OpenGLRenderTarget.cpp
--- a/Engine/Source/Runtime/OpenGLDrv/Private/OpenGLRenderTarget.cpp
+++ b/Engine/Source/Runtime/OpenGLDrv/Private/OpenGLRenderTarget.cpp
@@ -614,8 +614,8 @@ void FOpenGLDynamicRHI::ReadSurfaceDataRaw(FOpenGLContextState& ContextState, FT
 		else 
 			glReadPixels(Rect.Min.X, Rect.Min.Y, SizeX, SizeY, GL_RGBA, GL_FLOAT, FloatBGRAData );
 		// Determine minimal and maximal float values present in received data. Treat each component separately.
-		float MinValue[4] = { 0.0f };
-		float MaxValue[4] = { 1.0f };
+		float MinValue[4] = { 0.0f, 0.0f, 0.0f, 0.0f };
+		float MaxValue[4] = { 1.0f, 1.0f, 1.0f, 1.0f };
 		float* DataPtr = FloatBGRAData;
 		for( int32 PixelComponentIndex = 0; PixelComponentIndex < PixelComponentCount; ++PixelComponentIndex, ++DataPtr )
 		{
