Fix crash opening or playing some maps caused by accessing an invalid pointer when accessing render target dimensions to set the glViewport.

diff --git a/Engine/Source/Runtime/OpenGLDrv/Private/OpenGLCommands.cpp b/Engine/Source/Runtime/OpenGLDrv/Private/OpenGLCommands.cpp
--- a/Engine/Source/Runtime/OpenGLDrv/Private/OpenGLCommands.cpp
+++ b/Engine/Source/Runtime/OpenGLDrv/Private/OpenGLCommands.cpp
@@ -1575,7 +1575,7 @@ void FOpenGLDynamicRHI::RHISetRenderTargets(
 				else
 				{
 					FOpenGLTexture2DArray* NewRenderTarget2DArray = (FOpenGLTexture2DArray*)NewRenderTargetsRHI[PendingState.FirstNonzeroRenderTarget].Texture->GetTexture2DArray();
-					if(NewRenderTarget3D)
+					if(NewRenderTarget2DArray)
 					{
 						Width = NewRenderTarget2DArray->GetSizeX();
 						Height = NewRenderTarget2DArray->GetSizeY();
