Fixed texture streaming

diff --git a/Engine/Source/Runtime/Engine/Private/ContentStreaming.cpp b/Engine/Source/Runtime/Engine/Private/ContentStreaming.cpp
--- a/Engine/Source/Runtime/Engine/Private/ContentStreaming.cpp
+++ b/Engine/Source/Runtime/Engine/Private/ContentStreaming.cpp
@@ -5,6 +5,7 @@
 =============================================================================*/
 
 #include "EnginePrivate.h"
+#include "GenericPlatformMemoryPoolStats.h"
 
 DEFINE_LOG_CATEGORY_STATIC(LogContentStreaming, Log, All);
 
@@ -906,11 +907,13 @@ struct FStreamingContext
 			// set total size for the pool (used to available)
 			SET_MEMORY_STAT(STAT_TexturePoolAllocatedSize, (int32)AllocatedMemorySize);
 			SET_MEMORY_STAT(STAT_TexturePoolSize, (int32)Stats.TexturePoolSize);
+			SET_MEMORY_STAT(MCR_TexturePool, (int32)Stats.TexturePoolSize);
 		}
 		else
 		{
 			SET_MEMORY_STAT(STAT_TexturePoolAllocatedSize,0);
 			SET_MEMORY_STAT(STAT_TexturePoolSize, 0);
+			SET_MEMORY_STAT(MCR_TexturePool, 0);
 		}
 	}
 
@@ -3719,7 +3722,7 @@ void FStreamingManagerTexture::CheckUserSettings()
 				int64 TotalGraphicsMemory = Stats.TotalGraphicsMemory;
 				if ( GCurrentRendertargetMemorySize > 0 )
 				{
-					TotalGraphicsMemory -= GCurrentRendertargetMemorySize * 1024;
+					TotalGraphicsMemory -= int64(GCurrentRendertargetMemorySize) * 1024;
 				}
 				float PoolSize = float(GPoolSizeVRAMPercentage) * 0.01f * float(TotalGraphicsMemory);
 
diff --git a/Engine/Source/Runtime/OpenGLDrv/Private/OpenGLTexture.cpp b/Engine/Source/Runtime/OpenGLDrv/Private/OpenGLTexture.cpp
--- a/Engine/Source/Runtime/OpenGLDrv/Private/OpenGLTexture.cpp
+++ b/Engine/Source/Runtime/OpenGLDrv/Private/OpenGLTexture.cpp
@@ -98,11 +98,11 @@ void OpenGLTextureAllocated(FRHITexture* Texture, uint32 Flags)
 
 	if( bRenderTarget )
 	{
-		GCurrentRendertargetMemorySize += TextureSize / 1024;
+		GCurrentRendertargetMemorySize += Align(TextureSize, 1024) / 1024;
 	}
 	else
 	{
-		GCurrentTextureMemorySize += TextureSize;
+		GCurrentTextureMemorySize += Align(TextureSize, 1024) / 1024;
 	}
 }
 
@@ -165,11 +165,11 @@ void OpenGLTextureDeleted( FRHITexture* Texture )
 
 	if( bRenderTarget )
 	{
-		GCurrentRendertargetMemorySize -= TextureSize / 1024;
+		GCurrentRendertargetMemorySize -= Align(TextureSize, 1024) / 1024;
 	}
 	else
 	{
-		GCurrentTextureMemorySize -= TextureSize;
+		GCurrentTextureMemorySize -= Align(TextureSize, 1024) / 1024;
 	}
 }
 
@@ -185,7 +185,7 @@ void FOpenGLDynamicRHI::RHIGetTextureMemoryStats(FTextureMemoryStats& OutStats)
     OutStats.SharedSystemMemory = 0;
 	OutStats.TotalGraphicsMemory = GOpenGLTotalGraphicsMemory ? GOpenGLTotalGraphicsMemory : -1;
 
-	OutStats.AllocatedMemorySize = GCurrentTextureMemorySize * 1024;
+	OutStats.AllocatedMemorySize = int64(GCurrentTextureMemorySize) * 1024;
 	OutStats.TexturePoolSize = GTexturePoolSize;
 	OutStats.PendingMemoryAdjustment = 0;
 }
diff --git a/Engine/Source/Runtime/Windows/D3D11RHI/Private/D3D11Texture.cpp b/Engine/Source/Runtime/Windows/D3D11RHI/Private/D3D11Texture.cpp
--- a/Engine/Source/Runtime/Windows/D3D11RHI/Private/D3D11Texture.cpp
+++ b/Engine/Source/Runtime/Windows/D3D11RHI/Private/D3D11Texture.cpp
@@ -76,11 +76,11 @@ void UpdateD3D11TextureStats(uint32 BindFlags, uint32 MiscFlags, int64 TextureSi
 	
 	if(ShouldCountAsTextureMemory(BindFlags))
 	{
-		FPlatformAtomics::InterlockedAdd( &GCurrentTextureMemorySize, TextureSize / 1024 );
+		FPlatformAtomics::InterlockedAdd(&GCurrentTextureMemorySize, Align(TextureSize, 1024) / 1024);
 	}
 	else
 	{
-		FPlatformAtomics::InterlockedAdd( &GCurrentRendertargetMemorySize, TextureSize / 1024 );
+		FPlatformAtomics::InterlockedAdd(&GCurrentRendertargetMemorySize, Align(TextureSize, 1024) / 1024);
 	}
 
 	bool bCubeMap = (MiscFlags & D3D11_RESOURCE_MISC_TEXTURECUBE) != 0;
@@ -218,7 +218,7 @@ void FD3D11DynamicRHI::RHIGetTextureMemoryStats(FTextureMemoryStats& OutStats)
     OutStats.SharedSystemMemory = GSharedSystemMemory;
 	OutStats.TotalGraphicsMemory = GTotalGraphicsMemory ? GTotalGraphicsMemory : -1;
 
-	OutStats.AllocatedMemorySize = GCurrentTextureMemorySize * 1024;
+	OutStats.AllocatedMemorySize = int64(GCurrentTextureMemorySize) * 1024;
 	OutStats.TexturePoolSize = GTexturePoolSize;
 	OutStats.PendingMemoryAdjustment = 0;
 }
