Use shader instructions to emulate FRAMEBUFFER_SRGB on Mac OS X until there is a fix for the bug in Mac AMD drivers. All Macs, even those equipped with Nvidia & Intel GPUs, will use this mode by default.
Other platforms will not be affected.

diff --git a/Engine/Shaders/DeferredShadingCommon.usf b/Engine/Shaders/DeferredShadingCommon.usf
--- a/Engine/Shaders/DeferredShadingCommon.usf
+++ b/Engine/Shaders/DeferredShadingCommon.usf
@@ -24,6 +24,7 @@ float3 DecodeSpecularColor(float3 SpecularColor)
 	return SpecularColor * SpecularColor;
 }
 
+#if USE_FRAMEBUFFER_SRGB
 float3 EncodeDiffuseColor(float3 DiffuseColor)
 {
 	// we use sRGB on the render target to give more precision to the darks
@@ -48,6 +49,29 @@ float3 DecodeBaseColor(float3 BaseColor)
 	// we use sRGB on the render target to give more precision to the darks
 	return BaseColor;
 }
+#else // @todo: remove once Apple fixes radr://16754329 AMD Cards don't always perform FRAMEBUFFER_SRGB if the draw FBO has mixed sRGB & non-SRGB colour attachments
+float3 EncodeDiffuseColor(float3 DiffuseColor)
+{
+	// Allocate more precision to the darks
+	return sqrt(saturate(DiffuseColor));
+}
+
+float3 DecodeDiffuseColor(float3 DiffuseColor)
+{
+	return DiffuseColor * DiffuseColor;
+}
+
+float3 EncodeBaseColor(float3 BaseColor)
+{
+	// Gamma 2.0
+	return sqrt( saturate(BaseColor) );
+}
+
+float3 DecodeBaseColor(float3 BaseColor)
+{
+	return Square( BaseColor );
+}
+#endif
 
 float3 EncodeSubsurfaceColor(float3 SubsurfaceColor)
 {
diff --git a/Engine/Source/Runtime/Engine/Private/ShaderCompiler/ShaderCompiler.cpp b/Engine/Source/Runtime/Engine/Private/ShaderCompiler/ShaderCompiler.cpp
--- a/Engine/Source/Runtime/Engine/Private/ShaderCompiler/ShaderCompiler.cpp
+++ b/Engine/Source/Runtime/Engine/Private/ShaderCompiler/ShaderCompiler.cpp
@@ -1991,6 +1991,15 @@ void GlobalBeginCompileShader(
 		static IConsoleVariable* CVar = IConsoleManager::Get().FindConsoleVariable(TEXT("r.DBuffer"));
 		Input.Environment.SetDefine(TEXT("USE_DBUFFER"), CVar ? CVar->GetInt() : 0);
 	}
+	
+	{
+		int32 UseFrameBufferSRGB = 1;
+#if PLATFORM_MAC // @todo: remove once Apple fixes radr://16754329 AMD Cards don't always perform FRAMEBUFFER_SRGB if the draw FBO has mixed sRGB & non-SRGB colour attachments
+		static const auto CVar = IConsoleManager::Get().FindTConsoleVariableDataInt(TEXT("r.Mac.UseFrameBufferSRGB"));
+		UseFrameBufferSRGB = CVar ? (CVar->GetValueOnGameThread() != 0) : 0;
+#endif
+		Input.Environment.SetDefine(TEXT("USE_FRAMEBUFFER_SRGB"), UseFrameBufferSRGB);
+	}
 
 	NewJobs.Add(NewJob);
 }
diff --git a/Engine/Source/Runtime/OpenGLDrv/Private/Mac/OpenGLMac.cpp b/Engine/Source/Runtime/OpenGLDrv/Private/Mac/OpenGLMac.cpp
--- a/Engine/Source/Runtime/OpenGLDrv/Private/Mac/OpenGLMac.cpp
+++ b/Engine/Source/Runtime/OpenGLDrv/Private/Mac/OpenGLMac.cpp
@@ -372,7 +372,7 @@ extern void OnQueryInvalidation( void );
 // @todo: remove once Apple fixes radr://16754329 AMD Cards don't always perform FRAMEBUFFER_SRGB if the draw FBO has mixed sRGB & non-SRGB colour attachments
 static TAutoConsoleVariable<int32> CVarMacUseFrameBufferSRGB(
 		TEXT("r.Mac.UseFrameBufferSRGB"),
-		-1,
+		0,
 		TEXT("Flag to toggle use of GL_FRAMEBUFFER_SRGB for better color accuracy.\n"),
 		ECVF_RenderThreadSafe
 		);
@@ -614,18 +614,6 @@ FPlatformOpenGLContext* PlatformCreateOpenGLContext(FPlatformOpenGLDevice* Devic
 	{
 		GIsRunningOnIntelCard = true;
 	}
-	else if (VendorName.Contains(TEXT("AMD ")) || VendorName.Contains(TEXT("ATI ")))
-	{
-		if(CVarMacUseFrameBufferSRGB.GetValueOnRenderThread() == -1)
-		{
-			CVarMacUseFrameBufferSRGB.AsVariable()->Set(0);
-		}
-	}
-	
-	if(CVarMacUseFrameBufferSRGB.GetValueOnRenderThread() == -1)
-	{
-		CVarMacUseFrameBufferSRGB.AsVariable()->Set(1);
-	}
 	
 	// Renderer IDs matchup to driver kexts, so switching based on them will allow us to target workarouds to many GPUs
 	// which exhibit the same unfortunate driver bugs without having to parse their individual ID strings.
