Fix for objects outline issue in the editor on Macs with NVIDIA cards (a workaround for a driver bug)

diff --git a/Engine/Source/Runtime/OpenGLDrv/Private/OpenGLCommands.cpp b/Engine/Source/Runtime/OpenGLDrv/Private/OpenGLCommands.cpp
--- a/Engine/Source/Runtime/OpenGLDrv/Private/OpenGLCommands.cpp
+++ b/Engine/Source/Runtime/OpenGLDrv/Private/OpenGLCommands.cpp
@@ -2652,19 +2652,20 @@ static inline void ClearCurrentDepthStencilWithCurrentScissor( int8 ClearType, f
                                 // Especially irritatingly this bug will not manifest when stepping though the program with GL Profiler.
                                 // This was a bug found by me during dev. on Tropico 3 circa. Q4 2011 in the ATi Mac 2xx0 & 4xx0 drivers.
                                 // It was never fixed & has re-emerged in the AMD Mac FirePro Dx00 drivers.
-        case CT_Stencil:	// Clear stencil only
-            FOpenGL::ClearBufferiv(GL_STENCIL, 0, (const GLint*)&Stencil);
-            
+								// Also, on NVIDIA depth must be cleared first.
+        case CT_Depth:	// Clear depth only
+            FOpenGL::ClearBufferfv(GL_DEPTH, 0, &Depth);
+
             // If not also clearing depth break
-            if(!(ClearType & CT_Depth))
+            if(!(ClearType & CT_Stencil))
             {
                 break;
             }
-            // Otherwise fall through to perform a separate depth clear.
-        case CT_Depth:	// Clear depth only
-            FOpenGL::ClearBufferfv(GL_DEPTH, 0, &Depth);
+            // Otherwise fall through to perform a separate stencil clear.
+        case CT_Stencil:	// Clear stencil only
+            FOpenGL::ClearBufferiv(GL_STENCIL, 0, (const GLint*)&Stencil);
             break;
-            
+
         default:
             break;	// impossible anyway
 	}
