Fixing Linux client compilation.

diff --git a/Engine/Source/Runtime/OpenGLDrv/Private/Linux/OpenGLLinux.h b/Engine/Source/Runtime/OpenGLDrv/Private/Linux/OpenGLLinux.h
--- a/Engine/Source/Runtime/OpenGLDrv/Private/Linux/OpenGLLinux.h
+++ b/Engine/Source/Runtime/OpenGLDrv/Private/Linux/OpenGLLinux.h
@@ -290,7 +290,8 @@
 	EnumMacro(PFNGLGETQUERYOBJECTUI64VPROC, glGetQueryObjectui64v)\
 	EnumMacro(PFNGLFENCESYNCPROC, glFenceSync)\
 	EnumMacro(PFNGLGETSYNCIVPROC, glGetSynciv)\
-	EnumMacro(PFNGLCLIENTWAITSYNCPROC, glClientWaitSync)
+	EnumMacro(PFNGLCLIENTWAITSYNCPROC, glClientWaitSync)\
+	EnumMacro(PFNGLBINDBUFFERRANGEPROC, glBindBufferRange)
 
 #define ENUM_GL_ENTRYPOINTS_OPTIONAL(EnumMacro) \
 	EnumMacro(PFNGLDEBUGMESSAGECALLBACKARBPROC,glDebugMessageCallbackARB) \
