OpenGL4: use bSupportsTesselation value in the rest of code (minor fix).

diff --git a/Engine/Source/Runtime/OpenGLDrv/Public/OpenGL4.h b/Engine/Source/Runtime/OpenGLDrv/Public/OpenGL4.h
--- a/Engine/Source/Runtime/OpenGLDrv/Public/OpenGL4.h
+++ b/Engine/Source/Runtime/OpenGLDrv/Public/OpenGL4.h
@@ -14,7 +14,7 @@
 struct FOpenGL4 : public FOpenGL3
 {
 	static FORCEINLINE bool SupportsSeparateAlphaBlend()				{ return true; }
-	static FORCEINLINE bool SupportsTessellation()						{ return true; }
+	static FORCEINLINE bool SupportsTessellation()						{ return bSupportsTessellation; }
 	static FORCEINLINE bool SupportsComputeShaders()					{ return bSupportsComputeShaders; }
 	static FORCEINLINE bool SupportsDrawIndirect()						{ return true; }
 	static FORCEINLINE bool SupportsVertexAttribBinding()				{ return bSupportsVertexAttribBinding; }
