Better acquiring of integer GL parameters.

- Set to default when cannot acquire.
- Unified macros for most platforms.

Fixes TTP #331890.

diff --git a/Engine/Source/Runtime/OpenGLDrv/Private/Mac/OpenGLMac.cpp b/Engine/Source/Runtime/OpenGLDrv/Private/Mac/OpenGLMac.cpp
--- a/Engine/Source/Runtime/OpenGLDrv/Private/Mac/OpenGLMac.cpp
+++ b/Engine/Source/Runtime/OpenGLDrv/Private/Mac/OpenGLMac.cpp
@@ -1359,12 +1359,6 @@ PFNGLPATCHPARAMETERIPROC FMacOpenGL::glPatchParameteri = NULL;
 
 void FMacOpenGL::ProcessQueryGLInt()
 {
-#ifndef __clang__
-#define LOG_AND_GET_GL_INT(IntEnum, Default, Dest) do { if (IntEnum) {glGetIntegerv(IntEnum, &Dest);} else {Dest = Default;} /*FPlatformMisc::LowLevelOutputDebugStringf(TEXT("  ") ## TEXT(#IntEnum) ## TEXT(": %d"), Dest);*/ } while(0)
-#else
-#define LOG_AND_GET_GL_INT(IntEnum, Default, Dest) do { if (IntEnum) {glGetIntegerv(IntEnum, &Dest);} else {Dest = Default;} /*FPlatformMisc::LowLevelOutputDebugStringf(TEXT("  " #IntEnum ": %d"), Dest);*/ } while(0)
-#endif
-	
 	LOG_AND_GET_GL_INT(GL_MAX_TESS_CONTROL_UNIFORM_COMPONENTS, 0, MaxHullUniformComponents);
 	LOG_AND_GET_GL_INT(GL_MAX_TESS_EVALUATION_UNIFORM_COMPONENTS, 0, MaxDomainUniformComponents);
 }
diff --git a/Engine/Source/Runtime/OpenGLDrv/Private/OpenGL3.cpp b/Engine/Source/Runtime/OpenGLDrv/Private/OpenGL3.cpp
--- a/Engine/Source/Runtime/OpenGLDrv/Private/OpenGL3.cpp
+++ b/Engine/Source/Runtime/OpenGLDrv/Private/OpenGL3.cpp
@@ -18,11 +18,6 @@ bool FOpenGL3::bDebugContext = false;
 
 void FOpenGL3::ProcessQueryGLInt()
 {
-#ifndef __clang__
-#define LOG_AND_GET_GL_INT(IntEnum, Default, Dest) do { if (IntEnum) {glGetIntegerv(IntEnum, &Dest);} else {Dest = Default;} /*FPlatformMisc::LowLevelOutputDebugStringf(TEXT("  ") ## TEXT(#IntEnum) ## TEXT(": %d"), Dest);*/ } while(0)
-#else
-#define LOG_AND_GET_GL_INT(IntEnum, Default, Dest) do { if (IntEnum) {glGetIntegerv(IntEnum, &Dest);} else {Dest = Default;} /*FPlatformMisc::LowLevelOutputDebugStringf(TEXT("  " #IntEnum ": %d"), Dest);*/ } while(0)
-#endif
 	LOG_AND_GET_GL_INT(GL_MAX_VERTEX_UNIFORM_COMPONENTS, 0, MaxVertexUniformComponents);
 	LOG_AND_GET_GL_INT(GL_MAX_FRAGMENT_UNIFORM_COMPONENTS, 0, MaxPixelUniformComponents);
 	LOG_AND_GET_GL_INT(GL_MAX_GEOMETRY_UNIFORM_COMPONENTS, 0, MaxGeometryUniformComponents);
@@ -39,7 +34,6 @@ void FOpenGL3::ProcessQueryGLInt()
 	LOG_AND_GET_GL_QUERY_INT(GL_TIMESTAMP, 0, TimestampQueryBits);
 	
 #undef LOG_AND_GET_GL_QUERY_INT
-#undef LOG_AND_GET_GL_INT
 }
 
 void FOpenGL3::ProcessExtensions( const FString& ExtensionsString )
diff --git a/Engine/Source/Runtime/OpenGLDrv/Private/OpenGL4.cpp b/Engine/Source/Runtime/OpenGLDrv/Private/OpenGL4.cpp
--- a/Engine/Source/Runtime/OpenGLDrv/Private/OpenGL4.cpp
+++ b/Engine/Source/Runtime/OpenGLDrv/Private/OpenGL4.cpp
@@ -20,12 +20,6 @@ bool FOpenGL4::bSupportsTextureView = true;
 
 void FOpenGL4::ProcessQueryGLInt()
 {
-#ifndef __clang__
-#define LOG_AND_GET_GL_INT(IntEnum, Default, Dest) do { if (IntEnum) {glGetIntegerv(IntEnum, &Dest);} else {Dest = Default;} /*FPlatformMisc::LowLevelOutputDebugStringf(TEXT("  ") ## TEXT(#IntEnum) ## TEXT(": %d"), Dest);*/ } while(0)
-#else
-#define LOG_AND_GET_GL_INT(IntEnum, Default, Dest) do { if (IntEnum) {glGetIntegerv(IntEnum, &Dest);} else {Dest = Default;} /*FPlatformMisc::LowLevelOutputDebugStringf(TEXT("  " #IntEnum ": %d"), Dest);*/ } while(0)
-#endif
-
 	if (bSupportsComputeShaders)
 	{
 		LOG_AND_GET_GL_INT(GL_MAX_COMPUTE_TEXTURE_IMAGE_UNITS, 0, MaxComputeTextureImageUnits);
diff --git a/Engine/Source/Runtime/OpenGLDrv/Private/OpenGLDevice.cpp b/Engine/Source/Runtime/OpenGLDrv/Private/OpenGLDevice.cpp
--- a/Engine/Source/Runtime/OpenGLDrv/Private/OpenGLDevice.cpp
+++ b/Engine/Source/Runtime/OpenGLDrv/Private/OpenGLDevice.cpp
@@ -523,38 +523,37 @@ static void InitRHICapabilitiesForGL()
 
 	// Log and get various limits.
 #ifndef __clang__
-	#define LOG_AND_GET_GL_INT(IntEnum,Default) GLint Value_##IntEnum; if (IntEnum) {glGetIntegerv(IntEnum, &Value_##IntEnum);} else {Value_##IntEnum = Default;} UE_LOG(LogRHI, Log, TEXT("  ") ## TEXT(#IntEnum) ## TEXT(": %d"), Value_##IntEnum)
+	#define LOG_AND_GET_GL_INT_TEMP(IntEnum,Default) GLint Value_##IntEnum; if (IntEnum) {glGetIntegerv(IntEnum, &Value_##IntEnum);} else {Value_##IntEnum = Default;} UE_LOG(LogRHI, Log, TEXT("  ") ## TEXT(#IntEnum) ## TEXT(": %d"), Value_##IntEnum)
 #else
-	#undef LOG_AND_GET_GL_INT
-	#define LOG_AND_GET_GL_INT(IntEnum,Default) GLint Value_##IntEnum; if (IntEnum) {glGetIntegerv(IntEnum, &Value_##IntEnum);} else {Value_##IntEnum = Default;} UE_LOG(LogRHI, Log, TEXT("  " #IntEnum ": %d"), Value_##IntEnum)
+	#define LOG_AND_GET_GL_INT_TEMP(IntEnum,Default) GLint Value_##IntEnum; if (IntEnum) {glGetIntegerv(IntEnum, &Value_##IntEnum);} else {Value_##IntEnum = Default;} UE_LOG(LogRHI, Log, TEXT("  " #IntEnum ": %d"), Value_##IntEnum)
 #endif
-	LOG_AND_GET_GL_INT(GL_MAX_TEXTURE_SIZE,0);
-	LOG_AND_GET_GL_INT(GL_MAX_CUBE_MAP_TEXTURE_SIZE,0);
+	LOG_AND_GET_GL_INT_TEMP(GL_MAX_TEXTURE_SIZE, 0);
+	LOG_AND_GET_GL_INT_TEMP(GL_MAX_CUBE_MAP_TEXTURE_SIZE, 0);
 #if GL_MAX_ARRAY_TEXTURE_LAYERS
-	LOG_AND_GET_GL_INT(GL_MAX_ARRAY_TEXTURE_LAYERS,0);
+	LOG_AND_GET_GL_INT_TEMP(GL_MAX_ARRAY_TEXTURE_LAYERS, 0);
 #endif
 #if GL_MAX_3D_TEXTURE_SIZE
-	LOG_AND_GET_GL_INT(GL_MAX_3D_TEXTURE_SIZE,0);
+	LOG_AND_GET_GL_INT_TEMP(GL_MAX_3D_TEXTURE_SIZE, 0);
 #endif 
-	LOG_AND_GET_GL_INT(GL_MAX_RENDERBUFFER_SIZE,0);
-	LOG_AND_GET_GL_INT(GL_MAX_TEXTURE_IMAGE_UNITS,0);
+	LOG_AND_GET_GL_INT_TEMP(GL_MAX_RENDERBUFFER_SIZE, 0);
+	LOG_AND_GET_GL_INT_TEMP(GL_MAX_TEXTURE_IMAGE_UNITS, 0);
 	if (FOpenGL::SupportsDrawBuffers())
 	{
-		LOG_AND_GET_GL_INT(GL_MAX_DRAW_BUFFERS,1);
+		LOG_AND_GET_GL_INT_TEMP(GL_MAX_DRAW_BUFFERS, 1);
 	}
-	LOG_AND_GET_GL_INT(GL_MAX_COLOR_ATTACHMENTS,1);
-	LOG_AND_GET_GL_INT(GL_MAX_SAMPLES,1);
-	LOG_AND_GET_GL_INT(GL_MAX_COLOR_TEXTURE_SAMPLES,1);
-	LOG_AND_GET_GL_INT(GL_MAX_DEPTH_TEXTURE_SAMPLES,1);
-	LOG_AND_GET_GL_INT(GL_MAX_INTEGER_SAMPLES,1);
-	LOG_AND_GET_GL_INT(GL_MAX_COMBINED_TEXTURE_IMAGE_UNITS,0);
-	LOG_AND_GET_GL_INT(GL_MAX_VERTEX_ATTRIBS,0);
+	LOG_AND_GET_GL_INT_TEMP(GL_MAX_COLOR_ATTACHMENTS, 1);
+	LOG_AND_GET_GL_INT_TEMP(GL_MAX_SAMPLES, 1);
+	LOG_AND_GET_GL_INT_TEMP(GL_MAX_COLOR_TEXTURE_SAMPLES, 1);
+	LOG_AND_GET_GL_INT_TEMP(GL_MAX_DEPTH_TEXTURE_SAMPLES, 1);
+	LOG_AND_GET_GL_INT_TEMP(GL_MAX_INTEGER_SAMPLES, 1);
+	LOG_AND_GET_GL_INT_TEMP(GL_MAX_COMBINED_TEXTURE_IMAGE_UNITS, 0);
+	LOG_AND_GET_GL_INT_TEMP(GL_MAX_VERTEX_ATTRIBS, 0);
 	if( FOpenGL::SupportsTextureFilterAnisotropic())
 	{
-		LOG_AND_GET_GL_INT(GL_MAX_TEXTURE_MAX_ANISOTROPY_EXT,0);
+		LOG_AND_GET_GL_INT_TEMP(GL_MAX_TEXTURE_MAX_ANISOTROPY_EXT, 0);
 		GMaxOpenGLTextureFilterAnisotropic = Value_GL_MAX_TEXTURE_MAX_ANISOTROPY_EXT;
 	}
-#undef LOG_AND_GET_GL_INT
+#undef LOG_AND_GET_GL_INT_TEMP
 
 	GMaxOpenGLColorSamples = Value_GL_MAX_COLOR_TEXTURE_SAMPLES;
 	GMaxOpenGLDepthSamples = Value_GL_MAX_DEPTH_TEXTURE_SAMPLES;
diff --git a/Engine/Source/Runtime/OpenGLDrv/Private/OpenGLDrv.cpp b/Engine/Source/Runtime/OpenGLDrv/Private/OpenGLDrv.cpp
--- a/Engine/Source/Runtime/OpenGLDrv/Private/OpenGLDrv.cpp
+++ b/Engine/Source/Runtime/OpenGLDrv/Private/OpenGLDrv.cpp
@@ -449,15 +449,9 @@ bool  FOpenGLBase::bSupportsTextureFilterAnisotropic = false;
 
 void FOpenGLBase::ProcessQueryGLInt()
 {
-#ifndef __clang__
-#define LOG_AND_GET_GL_INT(IntEnum, Default, Dest) do { if (IntEnum) {glGetIntegerv(IntEnum, &Dest);} else {Dest = Default;} /*FPlatformMisc::LowLevelOutputDebugStringf(TEXT("  ") ## TEXT(#IntEnum) ## TEXT(": %d"), Dest);*/ } while(0)
-#else
-#define LOG_AND_GET_GL_INT(IntEnum, Default, Dest) do { if (IntEnum) {glGetIntegerv(IntEnum, &Dest);} else {Dest = Default;} /*FPlatformMisc::LowLevelOutputDebugStringf(TEXT("  " #IntEnum ": %d"), Dest);*/ } while(0)
-#endif
 	LOG_AND_GET_GL_INT(GL_MAX_TEXTURE_IMAGE_UNITS, 0, MaxTextureImageUnits);
 	LOG_AND_GET_GL_INT(GL_MAX_VERTEX_TEXTURE_IMAGE_UNITS, 0, MaxVertexTextureImageUnits);
 	LOG_AND_GET_GL_INT(GL_MAX_COMBINED_TEXTURE_IMAGE_UNITS, 0, MaxCombinedTextureImageUnits);
-#undef LOG_AND_GET_GL_INT
 }
 
 void FOpenGLBase::ProcessExtensions( const FString& ExtensionsString )
diff --git a/Engine/Source/Runtime/OpenGLDrv/Private/OpenGLDrvPrivate.h b/Engine/Source/Runtime/OpenGLDrv/Private/OpenGLDrvPrivate.h
--- a/Engine/Source/Runtime/OpenGLDrv/Private/OpenGLDrvPrivate.h
+++ b/Engine/Source/Runtime/OpenGLDrv/Private/OpenGLDrvPrivate.h
@@ -11,6 +11,43 @@
 
 #define GL_CHECK(x)		x; do { GLint Err = glGetError(); if (Err != 0) {FPlatformMisc::LowLevelOutputDebugStringf(TEXT("(%s:%d) GL_CHECK Failed '%s'! %d (%x)\n"), ANSI_TO_TCHAR(__FILE__), __LINE__, ANSI_TO_TCHAR( #x ), Err, Err); check(!Err);}} while (0)
 
+#ifndef __clang__
+	#define LOG_AND_GET_GL_INT(IntEnum, Default, Dest) \
+		do \
+		{ \
+			Dest = GL_INVALID_ENUM; \
+			if (IntEnum) \
+			{ \
+				glGetIntegerv(IntEnum, &Dest); \
+			}; \
+			\
+			if (Dest == GL_INVALID_ENUM)  \
+			{ \
+				Dest = Default; \
+			}; \
+			/*FPlatformMisc::LowLevelOutputDebugStringf(TEXT("  ") ## TEXT(#IntEnum) ## TEXT(": %d"), Dest);*/ \
+		} \
+		while (0)
+#else
+	#define LOG_AND_GET_GL_INT(IntEnum, Default, Dest) \
+		do \
+		{ \
+			Dest = GL_INVALID_ENUM; \
+			if (IntEnum) \
+			{ \
+				glGetIntegerv(IntEnum, &Dest); \
+			}; \
+			\
+			if (Dest == GL_INVALID_ENUM)  \
+			{ \
+				Dest = Default; \
+			}; \
+			/*FPlatformMisc::LowLevelOutputDebugStringf(TEXT("  " #IntEnum ": %d"), Dest);*/ \
+		} \
+		while (0)
+#endif
+
+
 /**
  * The OpenGL RHI stats.
  */
diff --git a/Engine/Source/Runtime/OpenGLDrv/Private/OpenGLES2.cpp b/Engine/Source/Runtime/OpenGLDrv/Private/OpenGLES2.cpp
--- a/Engine/Source/Runtime/OpenGLDrv/Private/OpenGLES2.cpp
+++ b/Engine/Source/Runtime/OpenGLDrv/Private/OpenGLES2.cpp
@@ -114,17 +114,11 @@ bool FOpenGLES2::SupportsDisjointTimeQueries()
 
 void FOpenGLES2::ProcessQueryGLInt()
 {
-#ifndef __clang__
-#define LOG_AND_GET_GL_INT(IntEnum,Default,Dest) if (IntEnum) {glGetIntegerv(IntEnum, &Dest);} else {Dest = Default;} UE_LOG(LogRHI, Log, TEXT("  ") ## TEXT(#IntEnum) ## TEXT(": %d"), Dest)
-#else
-#define LOG_AND_GET_GL_INT(IntEnum,Default,Dest) if (IntEnum) {glGetIntegerv(IntEnum, &Dest);} else {Dest = Default;} UE_LOG(LogRHI, Log, TEXT("  " #IntEnum ": %d"), Dest)
-#endif
 	LOG_AND_GET_GL_INT(GL_MAX_VERTEX_UNIFORM_VECTORS, 0, MaxVertexUniformComponents);
 	LOG_AND_GET_GL_INT(GL_MAX_FRAGMENT_UNIFORM_VECTORS, 0, MaxPixelUniformComponents);
 	MaxVertexUniformComponents *= 4;
 	MaxPixelUniformComponents *= 4;
 	MaxGeometryUniformComponents = 0;
-#undef LOG_AND_GET_GL_INT
 
 	MaxGeometryTextureImageUnits = 0;
 	MaxHullTextureImageUnits = 0;
