To improve performance working with the Editor on Mac OS X use NSDisableScreenUpdates & NSEnableScreenUpdates to ensure that window composition only occurs when flushing the OpenGL context.
This means we don't need to render the default window view, nor stall for the render thread during live-resize operations, improving resize performance. Opening & moving a window will also temporarily disable
the screen updates which helps reduce the peformance impact when multiple notification windows appear.

diff --git a/Engine/Source/Runtime/OpenGLDrv/Private/Mac/OpenGLMac.cpp b/Engine/Source/Runtime/OpenGLDrv/Private/Mac/OpenGLMac.cpp
--- a/Engine/Source/Runtime/OpenGLDrv/Private/Mac/OpenGLMac.cpp
+++ b/Engine/Source/Runtime/OpenGLDrv/Private/Mac/OpenGLMac.cpp
@@ -148,6 +148,9 @@ private:
 
 void DeleteQueriesForCurrentContext( NSOpenGLContext* Context );
 
+/** Used to temporarily disable Cocoa screen updates to make window updates happen only on the render thread. */
+bool GMacEnableCocoaScreenUpdates = true;
+
 @interface NSView(NSThemeFramePrivate)
 - (float)roundedCornerRadius;
 @end
@@ -184,27 +187,31 @@ void DeleteQueriesForCurrentContext( NSOpenGLContext* Context );
 
 - (void) _surfaceNeedsUpdate:(NSNotification*)notification
 {
-	self.bNeedsUpdate = true;;
+	self.bNeedsUpdate = true;
 }
 
 - (void)drawRect:(NSRect)DirtyRect
 {
 	SCOPED_AUTORELEASE_POOL;
 	
-	[super drawRect:DirtyRect];
-	
 	FSlateCocoaWindow* SlateCocoaWindow = [[self window] isKindOfClass:[FSlateCocoaWindow class]] ? (FSlateCocoaWindow*)[self window] : nil;
 	if (SlateCocoaWindow)
 	{
 		[SlateCocoaWindow redrawContents];
-		
-		if([SlateCocoaWindow inLiveResize])
-		{
-			FlushRenderingCommands();
-		}
 	}
 }
 
+- (void)renewGState
+{
+	if(GMacEnableCocoaScreenUpdates)
+	{
+		GMacEnableCocoaScreenUpdates = false;
+		NSDisableScreenUpdates();
+	}
+	
+	[super renewGState];
+}
+
 - (BOOL)acceptsFirstMouse:(NSEvent *)Event
 {
 	return YES;
@@ -848,7 +855,13 @@ void PlatformBlitToViewport( FPlatformOpenGLDevice* Device, FPlatformOpenGLConte
 					[Context->OpenGLContext setValues: &RealSyncInterval forParameter: NSOpenGLCPSwapInterval];
 					Context->SyncInterval = RealSyncInterval;
 				}
-
+				
+				if(!GMacEnableCocoaScreenUpdates)
+				{
+					GMacEnableCocoaScreenUpdates = true;
+					NSEnableScreenUpdates();
+				}
+				
 				[Context->OpenGLContext flushBuffer];
 
 				
