Fix use of OpenGL sync objects across different contexts - they must be flushed to the GPU before attempting to wait on them on a different context. Flushing the context after issuing as suggested by Apple which means we shouldn't need to do so when waiting on the sync object.

diff --git a/Engine/Source/Runtime/OpenGLDrv/Private/OpenGLQuery.cpp b/Engine/Source/Runtime/OpenGLDrv/Private/OpenGLQuery.cpp
--- a/Engine/Source/Runtime/OpenGLDrv/Private/OpenGLQuery.cpp
+++ b/Engine/Source/Runtime/OpenGLDrv/Private/OpenGLQuery.cpp
@@ -254,8 +254,8 @@ void FOpenGLEventQuery::IssueEvent()
 		Sync = 0;
 	}
 	Sync = FOpenGL::FenceSync(GL_SYNC_GPU_COMMANDS_COMPLETE, 0);
-
-
+	glFlush();
+	
 	checkSlow(FOpenGL::IsSync(Sync));
 
 }
@@ -269,7 +269,7 @@ void FOpenGLEventQuery::WaitForCompletion()
 
 
 	// Wait up to 1/2 second for sync execution
-	FOpenGL::EFenceResult Status = FOpenGL::ClientWaitSync( Sync, GL_SYNC_FLUSH_COMMANDS_BIT, 500*1000*1000);
+	FOpenGL::EFenceResult Status = FOpenGL::ClientWaitSync( Sync, 0, 500*1000*1000);
 
 	if ( Status != FOpenGL::FR_AlreadySignaled && Status != FOpenGL::FR_ConditionSatisfied )
 	{
