Put the GL_SYNC_FLUSH_COMMANDS_BIT back when calling FOpenGL::ClientWaitSync for the benefit of iOS.

diff --git a/Engine/Source/Runtime/OpenGLDrv/Private/OpenGLQuery.cpp b/Engine/Source/Runtime/OpenGLDrv/Private/OpenGLQuery.cpp
--- a/Engine/Source/Runtime/OpenGLDrv/Private/OpenGLQuery.cpp
+++ b/Engine/Source/Runtime/OpenGLDrv/Private/OpenGLQuery.cpp
@@ -269,7 +269,7 @@ void FOpenGLEventQuery::WaitForCompletion()
 
 
 	// Wait up to 1/2 second for sync execution
-	FOpenGL::EFenceResult Status = FOpenGL::ClientWaitSync( Sync, 0, 500*1000*1000);
+	FOpenGL::EFenceResult Status = FOpenGL::ClientWaitSync( Sync, GL_SYNC_FLUSH_COMMANDS_BIT, 500*1000*1000);
 
 	if ( Status != FOpenGL::FR_AlreadySignaled && Status != FOpenGL::FR_ConditionSatisfied )
 	{
