fix flush bit for iOS in a different manner
#ue4
#ios

diff --git a/Engine/Source/Runtime/OpenGLDrv/Private/IOS/IOSOpenGL.h b/Engine/Source/Runtime/OpenGLDrv/Private/IOS/IOSOpenGL.h
--- a/Engine/Source/Runtime/OpenGLDrv/Private/IOS/IOSOpenGL.h
+++ b/Engine/Source/Runtime/OpenGLDrv/Private/IOS/IOSOpenGL.h
@@ -48,7 +48,7 @@ struct FIOSOpenGL : public FOpenGLES2
 
 	static FORCEINLINE EFenceResult ClientWaitSync(UGLsync Sync, GLbitfield Flags, GLuint64 Timeout)
 	{
-		check( Flags == GL_SYNC_FLUSH_COMMANDS_BIT );
+//		check( Flags == GL_SYNC_FLUSH_COMMANDS_BIT );
 		GLenum Result = glClientWaitSyncAPPLE( Sync, GL_SYNC_FLUSH_COMMANDS_BIT_APPLE, Timeout );
 		switch (Result)
 		{
diff --git a/Engine/Source/Runtime/OpenGLDrv/Private/OpenGLQuery.cpp b/Engine/Source/Runtime/OpenGLDrv/Private/OpenGLQuery.cpp
--- a/Engine/Source/Runtime/OpenGLDrv/Private/OpenGLQuery.cpp
+++ b/Engine/Source/Runtime/OpenGLDrv/Private/OpenGLQuery.cpp
@@ -269,7 +269,7 @@ void FOpenGLEventQuery::WaitForCompletion()
 
 
 	// Wait up to 1/2 second for sync execution
-	FOpenGL::EFenceResult Status = FOpenGL::ClientWaitSync( Sync, GL_SYNC_FLUSH_COMMANDS_BIT, 500*1000*1000);
+	FOpenGL::EFenceResult Status = FOpenGL::ClientWaitSync( Sync, 0, 500*1000*1000);
 
 	if ( Status != FOpenGL::FR_AlreadySignaled && Status != FOpenGL::FR_ConditionSatisfied )
 	{
