Fix assertion error introduced by CL#2074162.

diff --git a/Engine/Source/Runtime/OpenGLDrv/Private/Android/AndroidOpenGL.h b/Engine/Source/Runtime/OpenGLDrv/Private/Android/AndroidOpenGL.h
--- a/Engine/Source/Runtime/OpenGLDrv/Private/Android/AndroidOpenGL.h
+++ b/Engine/Source/Runtime/OpenGLDrv/Private/Android/AndroidOpenGL.h
@@ -154,7 +154,7 @@ struct FAndroidOpenGL : public FOpenGLES2
 	{
 		if (GUseThreadedRendering)
 		{
-			check( Flags == GL_SYNC_FLUSH_COMMANDS_BIT );
+			// check( Flags == GL_SYNC_FLUSH_COMMANDS_BIT );
 			GLenum Result = eglClientWaitSyncKHR( AndroidEGL::GetInstance()->GetDisplay(), Sync, EGL_SYNC_FLUSH_COMMANDS_BIT_KHR, Timeout );
 			switch (Result)
 			{
