Fix native resource accessor for OpenGL resource.

diff --git a/Engine/Source/Runtime/OpenGLDrv/Public/OpenGLResources.h b/Engine/Source/Runtime/OpenGLDrv/Public/OpenGLResources.h
--- a/Engine/Source/Runtime/OpenGLDrv/Public/OpenGLResources.h
+++ b/Engine/Source/Runtime/OpenGLDrv/Public/OpenGLResources.h
@@ -959,7 +959,7 @@ public:
 	/** FRHITexture override.  See FRHITexture::GetNativeResource() */
 	virtual void* GetNativeResource() const OVERRIDE
 	{ 
-		return (void*)Resource;
+		return const_cast<void *>(reinterpret_cast<const void*>(&Resource));
 	}
 
 	/**
