Merging

//depot/UE4-TappyChicken/Engine/Source/...

to //depot/UE4/Engine/Source/...

diff --git a/Engine/Source/Developer/HTML5/HTML5TargetPlatform/Private/HTML5TargetPlatform.cpp b/Engine/Source/Developer/HTML5/HTML5TargetPlatform/Private/HTML5TargetPlatform.cpp
--- a/Engine/Source/Developer/HTML5/HTML5TargetPlatform/Private/HTML5TargetPlatform.cpp
+++ b/Engine/Source/Developer/HTML5/HTML5TargetPlatform/Private/HTML5TargetPlatform.cpp
@@ -119,6 +119,7 @@ void FHTML5TargetPlatform::GetTextureFormats( const UTexture* Texture, TArray<FN
 	static FName NameBGRA8(TEXT("BGRA8"));
 	static FName NameG8(TEXT("G8"));
 	static FName NameRGBA16F(TEXT("RGBA16F"));
+	static FName NameRGBA8(TEXT("RGBA8"));
 
 	bool bNoCompression = Texture->CompressionNone				// Code wants the texture uncompressed.
 		|| (HasEditorOnlyData() && Texture->DeferCompression)	// The user wishes to defer compression, this is ok for the Editor only.
@@ -168,7 +169,7 @@ void FHTML5TargetPlatform::GetTextureFormats( const UTexture* Texture, TArray<FN
 	}
 	else if (Texture->CompressionSettings == TC_VectorDisplacementmap)
 	{
-		TextureFormatName = NameBGRA8;
+		TextureFormatName = NameRGBA8;
 	}
 	else if (Texture->CompressionSettings == TC_Grayscale)
 	{
diff --git a/Engine/Source/Developer/ShaderFormatOpenGL/Private/OpenGLShaderCompiler.cpp b/Engine/Source/Developer/ShaderFormatOpenGL/Private/OpenGLShaderCompiler.cpp
--- a/Engine/Source/Developer/ShaderFormatOpenGL/Private/OpenGLShaderCompiler.cpp
+++ b/Engine/Source/Developer/ShaderFormatOpenGL/Private/OpenGLShaderCompiler.cpp
@@ -1,5 +1,5 @@
 // Copyright 1998-2014 Epic Games, Inc. All Rights Reserved..
-
+// 
 #include "ShaderFormatOpenGL.h"
 #include "Core.h"
 
@@ -13,7 +13,7 @@
 	#include "Windows/MinWindows.h"
 #include "HideWindowsPlatformTypes.h"
 #endif
-#include "ShaderPreprocessor.h"
+#include "ShaderPreprocessor.h" 
 #include "hlslcc.h"
 #include "glsl/ir_gen_glsl.h"
 #if PLATFORM_WINDOWS
@@ -44,7 +44,7 @@ DEFINE_LOG_CATEGORY_STATIC(LogOpenGLShaderCompiler, Log, All);
 
 
 #define ENABLE_IMAGINATION_COMPILER		1
-
+ 
 #define MAX_SAMPLERS_PER_SHADER_GLSL_150 16
 #define MAX_SAMPLERS_PER_SHADER_GLSL_430 32
 
@@ -90,7 +90,7 @@ struct FPlatformOpenGLContext
 };
 
 /**
- * A dummy wndproc.
+ * A dummy wndproc. 
  */
 static LRESULT CALLBACK PlatformDummyGLWndproc(HWND hWnd, uint32 Message, WPARAM wParam, LPARAM lParam)
 {
diff --git a/Engine/Source/Developer/TextureFormatUncompressed/Private/TextureFormatUncompressed.cpp b/Engine/Source/Developer/TextureFormatUncompressed/Private/TextureFormatUncompressed.cpp
--- a/Engine/Source/Developer/TextureFormatUncompressed/Private/TextureFormatUncompressed.cpp
+++ b/Engine/Source/Developer/TextureFormatUncompressed/Private/TextureFormatUncompressed.cpp
@@ -18,7 +18,8 @@ DEFINE_LOG_CATEGORY_STATIC(LogTextureFormatUncompressed, Log, All);
 	op(G8) \
 	op(VU8) \
 	op(RGBA16F) \
-	op(XGXR8)
+	op(XGXR8) \
+	op(RGBA8)
 
 #define DECL_FORMAT_NAME(FormatName) static FName GTextureFormatName##FormatName = FName(TEXT(#FormatName));
 ENUM_SUPPORTED_FORMATS(DECL_FORMAT_NAME);
@@ -116,6 +117,33 @@ class FTextureFormatUncompressed : public ITextureFormat
 
 			return true;
 		}
+		else if (BuildSettings.TextureFormatName == GTextureFormatNameRGBA8)
+		{
+			FImage Image;
+			InImage.CopyTo(Image, ERawImageFormat::BGRA8, BuildSettings.bSRGB);
+
+			OutCompressedImage.SizeX = Image.SizeX;
+			OutCompressedImage.SizeY = Image.SizeY;
+			OutCompressedImage.PixelFormat = PF_B8G8R8A8;
+
+			// swizzle each texel
+			uint32 NumTexels = Image.SizeX * Image.SizeY * Image.NumSlices;
+			OutCompressedImage.RawData.Empty(NumTexels * 4);
+			OutCompressedImage.RawData.AddUninitialized(NumTexels * 4);
+			const FColor* FirstColor = Image.AsBGRA8();
+			const FColor* LastColor = FirstColor + NumTexels;
+			int8* Dest = (int8*)OutCompressedImage.RawData.GetTypedData();
+
+			for (const FColor* Color = FirstColor; Color < LastColor; ++Color)
+			{
+				*Dest++ = (int32)Color->R;
+				*Dest++ = (int32)Color->G;
+				*Dest++ = (int32)Color->B;
+				*Dest++ = (int32)Color->A;
+			}
+
+			return true;
+		}
 		else if (BuildSettings.TextureFormatName == GTextureFormatNameXGXR8)
 		{
 			FImage Image;
diff --git a/Engine/Source/Programs/AutomationTool/Android/AndroidPlatform.Automation.cs b/Engine/Source/Programs/AutomationTool/Android/AndroidPlatform.Automation.cs
--- a/Engine/Source/Programs/AutomationTool/Android/AndroidPlatform.Automation.cs
+++ b/Engine/Source/Programs/AutomationTool/Android/AndroidPlatform.Automation.cs
@@ -103,7 +103,8 @@ public class AndroidPlatform : Platform
 		// packaging just takes a pak file and makes it the .obb
         UEBuildConfiguration.bOBBinAPK = Params.OBBinAPK; // Make sure this setting is sync'd pre-build
         var Deploy = UEBuildDeploy.GetBuildDeploy(UnrealTargetPlatform.Android);
-		Deploy.PrepForUATPackageOrDeploy(Params.ShortProjectName, SC.ProjectRoot, SOName, SC.LocalRoot + "/Engine", Params.Distribution);
+		string CookFlavor = SC.CookPlatform.IndexOf("_") > 0 ? SC.CookPlatform.Substring(SC.CookPlatform.IndexOf("_")) : "";
+		Deploy.PrepForUATPackageOrDeploy(Params.ShortProjectName, SC.ProjectRoot, SOName, SC.LocalRoot + "/Engine", Params.Distribution, CookFlavor);
         
 		// first, look for a .pak file in the staged directory
 		string[] PakFiles = Directory.GetFiles(SC.StageDirectory, "*.pak", SearchOption.AllDirectories);
@@ -137,9 +138,9 @@ public class AndroidPlatform : Platform
 			"set ADB=%ANDROID_HOME%\\platform-tools\\adb.exe",
 			"%ADB% uninstall " + PackageName,
 			"%ADB% install " + Path.GetFileName(ApkName),
-			"%ADB% shell rm -r /mnt/sdcard/" + Params.ShortProjectName,
-			SC.IsCodeBasedProject ? "" : "%ADB% shell rm /mnt/sdcard/UE4Game/UE4CommandLine.txt", // we need to delete the commandline in UE4Game or it will mess up loading
-			"%ADB% shell rm -r /mnt/sdcard/obb/" + PackageName,
+			"%ADB% shell rm -rf /mnt/sdcard/" + Params.ShortProjectName,
+			"%ADB% shell rm -rf /mnt/sdcard/UE4Game/UE4CommandLine.txt", // we need to delete the commandline in UE4Game or it will mess up loading
+			"%ADB% shell rm -rf /mnt/sdcard/obb/" + PackageName,
 			Params.OBBinAPK ? "" : "%ADB% push " + Path.GetFileName(LocalObbName) + " " + DeviceObbName,
 		};
 		File.WriteAllLines(BatchName, BatchLines);
@@ -202,7 +203,8 @@ public class AndroidPlatform : Platform
 
 		// make sure APK is up to date (this is fast if so)
 		var Deploy = UEBuildDeploy.GetBuildDeploy(UnrealTargetPlatform.Android);
-		Deploy.PrepForUATPackageOrDeploy(Params.ShortProjectName, SC.ProjectRoot, SOName, SC.LocalRoot + "/Engine", Params.Distribution);
+		string CookFlavor = SC.CookPlatform.IndexOf("_") > 0 ? SC.CookPlatform.Substring(SC.CookPlatform.IndexOf("_")) : "";
+		Deploy.PrepForUATPackageOrDeploy(Params.ShortProjectName, SC.ProjectRoot, SOName, SC.LocalRoot + "/Engine", Params.Distribution, CookFlavor);
 
 		// now we can use the apk to get more info
 		string DeviceObbName = GetDeviceObbName(ApkName);
@@ -245,12 +247,12 @@ public class AndroidPlatform : Platform
 				}
 
 				string FinalRemoteDir = RemoteDir;
-				// handle the special case of the UE4Commandline.txt when using content only game (UE4Game)
+/*				// handle the special case of the UE4Commandline.txt when using content only game (UE4Game)
 				if (!Params.IsCodeBasedProject &&
 					Path.GetFileName(Filename).Equals("UE4CommandLine.txt", StringComparison.InvariantCultureIgnoreCase))
 				{
 					FinalRemoteDir = "/mnt/sdcard/UE4Game";
-				}
+				}*/
 
 				string RemoteFilename = Filename.Replace(SC.StageDirectory, FinalRemoteDir).Replace("\\", "/");
  				string Commandline = string.Format("{0} \"{1}\" \"{2}\"", BaseCommandline, Filename, RemoteFilename);
@@ -267,11 +269,11 @@ public class AndroidPlatform : Platform
 			string RemoteDir = "/mnt/sdcard/" + Params.ShortProjectName;
 
 			string FinalRemoteDir = RemoteDir;
-			// handle the special case of the UE4Commandline.txt when using content only game (UE4Game)
+/*			// handle the special case of the UE4Commandline.txt when using content only game (UE4Game)
 			if (!Params.IsCodeBasedProject)
 			{
 				FinalRemoteDir = "/mnt/sdcard/UE4Game";
-			}
+			}*/
 
 			string RemoteFilename = IntermediateCmdLineFile.Replace(SC.StageDirectory, FinalRemoteDir).Replace("\\", "/");
 			string Commandline = string.Format("{0} \"{1}\" \"{2}\"", BaseCommandline, IntermediateCmdLineFile, RemoteFilename);
diff --git a/Engine/Source/Programs/AutomationTool/IOS/IOSPlatform.Automation.cs b/Engine/Source/Programs/AutomationTool/IOS/IOSPlatform.Automation.cs
--- a/Engine/Source/Programs/AutomationTool/IOS/IOSPlatform.Automation.cs
+++ b/Engine/Source/Programs/AutomationTool/IOS/IOSPlatform.Automation.cs
@@ -81,7 +81,7 @@ public class IOSPlatform : Platform
 				Path.GetDirectoryName(Params.RawProjectPath),
 				CombinePaths(Path.GetDirectoryName(Params.ProjectGameExeFilename), SC.StageExecutables[0]),
 				CombinePaths(SC.LocalRoot, "Engine"),
-				Params.Distribution);
+				Params.Distribution, "");
 
 			// figure out where to pop in the staged files
 			string AppDirectory = string.Format("{0}/Payload/{1}.app",
diff --git a/Engine/Source/Programs/UnrealBuildTool/Android/UEBuildAndroid.cs b/Engine/Source/Programs/UnrealBuildTool/Android/UEBuildAndroid.cs
--- a/Engine/Source/Programs/UnrealBuildTool/Android/UEBuildAndroid.cs
+++ b/Engine/Source/Programs/UnrealBuildTool/Android/UEBuildAndroid.cs
@@ -243,6 +243,7 @@ namespace UnrealBuildTool
 			InBuildTarget.GlobalCompileEnvironment.Config.Definitions.Add("WITH_DATABASE_SUPPORT=0");
 			InBuildTarget.GlobalCompileEnvironment.Config.Definitions.Add("WITH_EDITOR=0");
 			InBuildTarget.GlobalCompileEnvironment.Config.Definitions.Add("USE_NULL_RHI=0");
+			InBuildTarget.GlobalCompileEnvironment.Config.Definitions.Add("REQUIRES_ALIGNED_INT_ACCESS");
 
 			if (InBuildTarget.GlobalCompileEnvironment.Config.TargetArchitecture == "-armv7")
 			{
diff --git a/Engine/Source/Programs/UnrealBuildTool/Android/UEDeployAndroid.cs b/Engine/Source/Programs/UnrealBuildTool/Android/UEDeployAndroid.cs
--- a/Engine/Source/Programs/UnrealBuildTool/Android/UEDeployAndroid.cs
+++ b/Engine/Source/Programs/UnrealBuildTool/Android/UEDeployAndroid.cs
@@ -263,7 +263,7 @@ namespace UnrealBuildTool.Android
         }
 
 
-		private void MakeAPK(string ProjectName, string ProjectDirectory, string OutputPath, string EngineDirectory, bool bForDistribution)
+		private void MakeAPK(string ProjectName, string ProjectDirectory, string OutputPath, string EngineDirectory, bool bForDistribution, string CookFlavor)
 		{
 			// cache some build product paths
 			string SourceSOName = OutputPath;
@@ -324,7 +324,7 @@ namespace UnrealBuildTool.Android
             // Currently we'll just support 1 of 'em
             if (UEBuildConfiguration.bOBBinAPK)
             {
-                string PAKFileLocation = ProjectDirectory + "/Saved/StagedBuilds/Android/" + ProjectName + "/Content/Paks";
+                string PAKFileLocation = ProjectDirectory + "/Saved/StagedBuilds/Android" + CookFlavor + "/" + ProjectName + "/Content/Paks";
                 Console.WriteLine("Pak location {0}", PAKFileLocation);
                 string PAKFileDestination = UE4BuildPath + "/assets";
                 Console.WriteLine("Pak destination location {0}", PAKFileDestination);
@@ -623,12 +623,12 @@ namespace UnrealBuildTool.Android
 
 		public override bool PrepTargetForDeployment(UEBuildTarget InTarget)
 		{
-			return PrepForUATPackageOrDeploy(InTarget.AppName, InTarget.ProjectDirectory, InTarget.OutputPath, BuildConfiguration.RelativeEnginePath, false);
+			return PrepForUATPackageOrDeploy(InTarget.AppName, InTarget.ProjectDirectory, InTarget.OutputPath, BuildConfiguration.RelativeEnginePath, false, "");
 		}
 
-		public override bool PrepForUATPackageOrDeploy(string ProjectName, string ProjectDirectory, string ExecutablePath, string EngineDirectory, bool bForDistribution)
+		public override bool PrepForUATPackageOrDeploy(string ProjectName, string ProjectDirectory, string ExecutablePath, string EngineDirectory, bool bForDistribution, string CookFlavor)
 		{
-			MakeAPK(ProjectName, ProjectDirectory, ExecutablePath, EngineDirectory, bForDistribution);
+			MakeAPK(ProjectName, ProjectDirectory, ExecutablePath, EngineDirectory, bForDistribution, CookFlavor);
 			return true;
 		}
 
diff --git a/Engine/Source/Programs/UnrealBuildTool/Configuration/UEBuildDeploy.cs b/Engine/Source/Programs/UnrealBuildTool/Configuration/UEBuildDeploy.cs
--- a/Engine/Source/Programs/UnrealBuildTool/Configuration/UEBuildDeploy.cs
+++ b/Engine/Source/Programs/UnrealBuildTool/Configuration/UEBuildDeploy.cs
@@ -94,7 +94,7 @@ namespace UnrealBuildTool
 		/**
 		 * Prepare the target for deployment
 		 */
-		public virtual bool PrepForUATPackageOrDeploy(string ProjectName, string ProjectDirectory, string ExecutablePath, string EngineDirectory, bool bForDistribution)
+		public virtual bool PrepForUATPackageOrDeploy(string ProjectName, string ProjectDirectory, string ExecutablePath, string EngineDirectory, bool bForDistribution, string CookFlavor)
 		{
 			return true;
 		}
diff --git a/Engine/Source/Programs/UnrealBuildTool/IOS/UEDeployIOS.cs b/Engine/Source/Programs/UnrealBuildTool/IOS/UEDeployIOS.cs
--- a/Engine/Source/Programs/UnrealBuildTool/IOS/UEDeployIOS.cs
+++ b/Engine/Source/Programs/UnrealBuildTool/IOS/UEDeployIOS.cs
@@ -65,7 +65,7 @@ namespace UnrealBuildTool.IOS
 			}
 		}
 
-		public override bool PrepForUATPackageOrDeploy(string InProjectName, string InProjectDirectory, string InExecutablePath, string InEngineDir, bool bForDistribution)
+		public override bool PrepForUATPackageOrDeploy(string InProjectName, string InProjectDirectory, string InExecutablePath, string InEngineDir, bool bForDistribution, string CookFlavor)
 		{
 			bool bIsUE4Game = InExecutablePath.Contains ("UE4Game");
 			string BinaryPath = Path.GetDirectoryName (InExecutablePath);
@@ -229,7 +229,7 @@ namespace UnrealBuildTool.IOS
 				string BuildPath = InTarget.ProjectDirectory + "/Binaries/IOS";
 				string ProjectDirectory = InTarget.ProjectDirectory;
 
-				return PrepForUATPackageOrDeploy(GameName, ProjectDirectory, BuildPath + "/" + DecoratedGameName, "../../Engine", false);
+				return PrepForUATPackageOrDeploy(GameName, ProjectDirectory, BuildPath + "/" + DecoratedGameName, "../../Engine", false, "");
 			}
 			else
 			{
diff --git a/Engine/Source/Runtime/Core/Private/Android/AndroidFile.cpp b/Engine/Source/Runtime/Core/Private/Android/AndroidFile.cpp
--- a/Engine/Source/Runtime/Core/Private/Android/AndroidFile.cpp
+++ b/Engine/Source/Runtime/Core/Private/Android/AndroidFile.cpp
@@ -129,8 +129,8 @@ class CORE_API FFileHandlePakAndroid : public IFileHandle
 	enum { READWRITE_SIZE = 1024 * 1024 };
 	AAsset * File;
 	int32 FileHandle;
-	off64_t Start;
-	off64_t Length;
+	off_t Start;
+	off_t Length;
 
 	FORCEINLINE bool IsValid()
 	{
@@ -141,7 +141,7 @@ public:
 	FFileHandlePakAndroid(AAsset * file)
 		: File(file), FileHandle(0), Start(0), Length(0)
 	{
-		FileHandle = AAsset_openFileDescriptor64(file, &Start, &Length);
+		FileHandle = AAsset_openFileDescriptor(file, &Start, &Length);
 	}
 
 	virtual ~FFileHandlePakAndroid()
@@ -154,7 +154,7 @@ public:
 	virtual int64 Tell() OVERRIDE
 	{
 		check(IsValid());
-		int64 pos = lseek64(FileHandle, 0, SEEK_CUR);
+		int64 pos = lseek(FileHandle, 0, SEEK_CUR);
 		return pos - Start; // We are treating 'tell' as a virtual location from file Start
 	}
 
@@ -165,7 +165,7 @@ public:
 		NewPosition += Start;
 		check(NewPosition >= 0);
 		
-		return lseek64(FileHandle, NewPosition, SEEK_SET) != -1;
+		return lseek(FileHandle, NewPosition, SEEK_SET) != -1;
 	}
 
     virtual bool SeekFromEnd(int64 NewPositionRelativeToEnd = 0) OVERRIDE
@@ -174,7 +174,7 @@ public:
 		check(NewPositionRelativeToEnd <= 0);
 		// We need to convert this to a virtual offset inside the file we are interested in
 		int64 position = Start + (Length - NewPositionRelativeToEnd);
-		return lseek64(FileHandle, position, SEEK_SET) != -1;
+		return lseek(FileHandle, position, SEEK_SET) != -1;
 	}
 
 	virtual bool Read(uint8* Destination, int64 BytesToRead) OVERRIDE
@@ -282,7 +282,7 @@ public:
 		if (convertedPath.Contains(TEXT(".pak")) && GOBBinAPK)
 		{
 			AAsset * file = AAssetManager_open(AssetMgr, TCHAR_TO_UTF8(*convertedPath), AASSET_MODE_RANDOM); // seems reasonable?
-			FileInfo.st_size = AAsset_getLength64(file);
+			FileInfo.st_size = AAsset_getLength(file);
 			AAsset_close(file);
 		}
 		else
diff --git a/Engine/Source/Runtime/Core/Private/Android/AndroidMisc.cpp b/Engine/Source/Runtime/Core/Private/Android/AndroidMisc.cpp
--- a/Engine/Source/Runtime/Core/Private/Android/AndroidMisc.cpp
+++ b/Engine/Source/Runtime/Core/Private/Android/AndroidMisc.cpp
@@ -12,6 +12,11 @@ DEFINE_LOG_CATEGORY_STATIC(LogEngine, Log, All);
 
 void* FAndroidMisc::NativeWindow = NULL;
 
+// run time compatibility information
+FString FAndroidMisc::AndroidVersion; // version of android we are running eg "4.0.4"
+FString FAndroidMisc::DeviceMake; // make of the device we are running on eg. "samsung"
+FString FAndroidMisc::DeviceModel; // model of the device we are running on eg "SAMSUNG-SGH-I437"
+
 GenericApplication* FAndroidMisc::CreateApplication()
 {
 	return FAndroidApplication::CreateAndroidApplication();
@@ -220,6 +225,32 @@ bool FAndroidMisc::GetUseVirtualJoysticks()
 	return !FParse::Param(FCommandLine::Get(),TEXT("joystick"));
 }
 
+
+
+void FAndroidMisc::SetVersionInfo( FString InAndroidVersion, FString InDeviceMake, FString InDeviceModel )
+{
+	AndroidVersion = InAndroidVersion;
+	DeviceMake = InDeviceMake;
+	DeviceModel = InDeviceModel;
+
+	UE_LOG(LogEngine, Display, TEXT("Android Version Make Model: %s %s %s"), *AndroidVersion, *DeviceMake, *DeviceModel);
+}
+
+const FString FAndroidMisc::GetAndroidVersion()
+{
+	return AndroidVersion;
+}
+
+const FString FAndroidMisc::GetDeviceMake()
+{
+	return DeviceMake;
+}
+
+const FString FAndroidMisc::GetDeviceModel()
+{
+	return DeviceModel;
+}
+
 uint32 FAndroidMisc::GetKeyMap( uint16* KeyCodes, FString* KeyNames, uint32 MaxMappings )
 {
 #define ADDKEYMAP(KeyCode, KeyName)		if (NumMappings<MaxMappings) { KeyCodes[NumMappings]=KeyCode; if(KeyNames) { KeyNames[NumMappings]=KeyName; } ++NumMappings; };
diff --git a/Engine/Source/Runtime/Core/Public/Android/AndroidMisc.h b/Engine/Source/Runtime/Core/Public/Android/AndroidMisc.h
--- a/Engine/Source/Runtime/Core/Public/Android/AndroidMisc.h
+++ b/Engine/Source/Runtime/Core/Public/Android/AndroidMisc.h
@@ -76,6 +76,10 @@ struct CORE_API FAndroidMisc : public FGenericPlatformMisc
 	static const TCHAR* GetDefaultDeviceProfileName() { return TEXT("Android"); }
 
 	// ANDROID ONLY:
+	static void SetVersionInfo( FString AndroidVersion, FString DeviceMake, FString DeviceModel );
+	static const FString GetAndroidVersion();
+	static const FString GetDeviceMake();
+	static const FString GetDeviceModel();
 	static FString GetGPUFamily();
 	static FString GetGLVersion();
 	static bool SupportsFloatingPointRenderTargets();
@@ -108,6 +112,12 @@ struct CORE_API FAndroidMisc : public FGenericPlatformMisc
 
 
 	static void* NativeWindow ; //raw platform Main window
+	
+	// run time compatibility information
+	static FString AndroidVersion; // version of android we are running eg "4.0.4"
+	static FString DeviceMake; // make of the device we are running on eg. "samsung"
+	static FString DeviceModel; // model of the device we are running on eg "SAMSUNG-SGH-I437"
+
 };
 
 typedef FAndroidMisc FPlatformMisc;
diff --git a/Engine/Source/Runtime/Engine/Private/ShaderCompiler/ShaderCompiler.cpp b/Engine/Source/Runtime/Engine/Private/ShaderCompiler/ShaderCompiler.cpp
--- a/Engine/Source/Runtime/Engine/Private/ShaderCompiler/ShaderCompiler.cpp
+++ b/Engine/Source/Runtime/Engine/Private/ShaderCompiler/ShaderCompiler.cpp
@@ -1,7 +1,7 @@
 // Copyright 1998-2014 Epic Games, Inc. All Rights Reserved.
 
 /*=============================================================================
-	ShaderCompiler.cpp: Platform independent shader compilation.
+	ShaderCompiler.cpp: Platform independent shader compilations.
 =============================================================================*/
 
 #include "EnginePrivate.h"
@@ -15,7 +15,6 @@ DEFINE_LOG_CATEGORY(LogShaderCompilers);
 
 // Set to 1 to debug ShaderCompilerWorker.exe. Set a breakpoint in LaunchWorker() to get the cmd-line.
 #define DEBUG_SHADERCOMPILEWORKER 0
-
 // Default value comes from bPromptToRetryFailedShaderCompiles in BaseEngine.ini
 // This is set as a global variable to allow changing in the debugger even in release
 // For example if there are a lot of content shader compile errors you want to skip over without relaunching
diff --git a/Engine/Source/Runtime/Launch/Private/Android/LaunchAndroid.cpp b/Engine/Source/Runtime/Launch/Private/Android/LaunchAndroid.cpp
--- a/Engine/Source/Runtime/Launch/Private/Android/LaunchAndroid.cpp
+++ b/Engine/Source/Runtime/Launch/Private/Android/LaunchAndroid.cpp
@@ -139,6 +139,7 @@ void UpdateGameInterruptions()
 		}
 		else
 		{
+			UE_LOG(LogAndroid, Display, TEXT("Acquiring Thread Ownership"));
 			RHIAcquireThreadOwnership();
 		}
 	}
@@ -241,6 +242,8 @@ int32 AndroidMain(struct android_app* state)
 		IgnoredGamepadKeyCodes.Add(IgnoredGamepadKeyCodesList[i]);
 	}
 
+	FPlatformProcess::Sleep(1.0f);
+
 	// read the command line file
 	InitCommandLine();
 	FPlatformMisc::LowLevelOutputDebugStringf(TEXT("Final commandline: %s\n"), FCommandLine::Get());
@@ -865,7 +868,7 @@ extern "C" void Java_com_epicgames_ue4_GameActivity_nativeConsoleCommand(JNIEnv*
 	jenv->ReleaseStringUTFChars(commandString, javaChars);
 }
 
-//This function is declared in the Java-defined class, GameActivity.java: "public native void nativeIsGooglePlayEnabled(String commandString);"
+//This function is declared in the Java-defined class, GameActivity.java: "public native void nativeIsGooglePlayEnabled();"
 extern "C" jboolean Java_com_epicgames_ue4_GameActivity_nativeIsGooglePlayEnabled(JNIEnv* jenv, jobject thiz)
 {
 	bool bEnabled = true;
@@ -873,3 +876,18 @@ extern "C" jboolean Java_com_epicgames_ue4_GameActivity_nativeIsGooglePlayEnable
 	UE_LOG(LogOnline, Log, TEXT("Checking whether Google Play is enabled. bEnableGooglePlaySupport = %d"), bEnabled);
 	return bEnabled;
 }
+
+
+extern "C" void Java_com_epicgames_ue4_GameActivity_nativeSetAndroidVersionInformation(JNIEnv* jenv, jobject thiz, jstring androidVersion, jstring phoneMake, jstring phoneModel )
+{
+	const char *javaAndroidVersion = jenv->GetStringUTFChars(androidVersion, 0 );
+	FString UEAndroidVersion = FString(UTF8_TO_TCHAR( javaAndroidVersion ));
+
+	const char *javaPhoneMake = jenv->GetStringUTFChars(phoneMake, 0 );
+	FString UEPhoneMake = FString(UTF8_TO_TCHAR( javaPhoneMake ));
+
+	const char *javaPhoneModel = jenv->GetStringUTFChars(phoneModel, 0 );
+	FString UEPhoneModel = FString(UTF8_TO_TCHAR( javaPhoneModel ));
+
+	FAndroidMisc::SetVersionInfo( UEAndroidVersion, UEPhoneMake, UEPhoneModel );
+}
\ No newline at end of file
diff --git a/Engine/Source/Runtime/Launch/Private/LaunchEngineLoop.cpp b/Engine/Source/Runtime/Launch/Private/LaunchEngineLoop.cpp
--- a/Engine/Source/Runtime/Launch/Private/LaunchEngineLoop.cpp
+++ b/Engine/Source/Runtime/Launch/Private/LaunchEngineLoop.cpp
@@ -1008,6 +1008,17 @@ int32 FEngineLoop::PreInit( const TCHAR* CmdLine )
 	if (FApp::ShouldUseThreadingForPerformance())
 	{
 		GUseThreadedRendering = true;
+
+#if PLATFORM_ANDROID 
+		// there is a crash with the nvidia tegra dual core processors namely the optimus 2x and xoom 
+		// when running multithreaded it can't handle multiple threads using opengl (bug)
+		// tested with lg optimus 2x and motorola xoom 
+		// come back and revisit this later 
+		if ( FAndroidMisc::GetGPUFamily() == FString(TEXT("NVIDIA Tegra")) && FPlatformMisc::NumberOfCores() <= 2 )
+		{
+			GUseThreadedRendering = false;
+		}
+#endif
 	}
 #endif
 
@@ -1138,8 +1149,6 @@ int32 FEngineLoop::PreInit( const TCHAR* CmdLine )
 	// Delete temporary files in cache.
 	FPlatformProcess::CleanFileCache();
 
-
-
 #if !UE_BUILD_SHIPPING
 	GIsDemoMode = FParse::Param( FCommandLine::Get(), TEXT( "DEMOMODE" ) );
 #endif
diff --git a/Engine/Source/Runtime/OpenGLDrv/Private/OpenGLDevice.cpp b/Engine/Source/Runtime/OpenGLDrv/Private/OpenGLDevice.cpp
--- a/Engine/Source/Runtime/OpenGLDrv/Private/OpenGLDevice.cpp
+++ b/Engine/Source/Runtime/OpenGLDrv/Private/OpenGLDevice.cpp
@@ -807,6 +807,67 @@ FOpenGLDynamicRHI::FOpenGLDynamicRHI()
 
 extern void DestroyShadersAndPrograms();
 
+#if PLATFORM_ANDROID
+// only used to test for shader compatibility issues
+static bool VerifyCompiledShader(GLuint Shader, const ANSICHAR* GlslCode, bool IsFatal )
+{
+	SCOPE_CYCLE_COUNTER(STAT_OpenGLShaderCompileVerifyTime);
+
+	GLint CompileStatus;
+	glGetShaderiv(Shader, GL_COMPILE_STATUS, &CompileStatus);
+	if (CompileStatus != GL_TRUE)
+	{
+		GLint LogLength;
+		ANSICHAR DefaultLog[] = "No log";
+		ANSICHAR *CompileLog = DefaultLog;
+		glGetShaderiv(Shader, GL_INFO_LOG_LENGTH, &LogLength);
+#if PLATFORM_ANDROID
+		if ( LogLength == 0 )
+		{
+			// make it big anyway
+			// there was a bug in android 2.2 where glGetShaderiv would return 0 even though there was a error message
+			// https://code.google.com/p/android/issues/detail?id=9953
+			LogLength = 4096;
+		}
+#endif
+		if (LogLength > 1)
+		{
+			CompileLog = (ANSICHAR *)FMemory::Malloc(LogLength);
+			glGetShaderInfoLog(Shader, LogLength, NULL, CompileLog);
+		}
+
+#if DEBUG_GL_SHADERS
+		if (GlslCode)
+		{
+			UE_LOG(LogRHI,Warning,TEXT("Shader:\n%s"),ANSI_TO_TCHAR(GlslCode));
+
+
+			const ANSICHAR *Temp = GlslCode;
+
+			for ( int i = 0; i < 30 && (*Temp != '\0'); ++i )
+			{
+				FString Converted = ANSI_TO_TCHAR( Temp );
+				Converted.LeftChop( 256 );
+
+				UE_LOG(LogRHI,Display,TEXT("%s"), *Converted );
+				Temp += Converted.Len();
+			}
+
+		}	
+#endif
+		UE_LOG(LogRHI,Warning,TEXT("Failed to compile shader. Compile log:\n%s"), ANSI_TO_TCHAR(CompileLog));
+		
+		if (LogLength > 1)
+		{
+			FMemory::Free(CompileLog);
+		}
+		return false;
+	}
+	return true;
+}
+#endif
+
+
 void FOpenGLDynamicRHI::Init()
 {
 	check(!GIsRHIInitialized);
@@ -871,6 +932,94 @@ void FOpenGLDynamicRHI::Init()
 
 	// Set the RHI initialized flag.
 	GIsRHIInitialized = true;
+
+	
+#if PLATFORM_ANDROID
+	{
+		UE_LOG(LogRHI,Display,TEXT("Testing for shader compiler compatibility"));
+		// This code creates a sample program and finds out which hacks are required to compile it
+		const ANSICHAR* TestFragmentProgram = "\n"
+			"#version 100\n"
+			"#ifndef DONTEMITEXTENSIONSHADERTEXTURELODENABLE\n"
+			"#extension GL_EXT_shader_texture_lod : enable\n"
+			"#endif\n"
+			"precision mediump float;\n"
+			"precision mediump int;\n"
+			"#ifndef DONTEMITSAMPLERDEFAULTPRECISION\n"
+			"precision mediump sampler2D;\n"
+			"precision mediump samplerCube;\n"
+			"#endif\n"
+			"varying vec3 TexCoord;\n"
+			"uniform samplerCube Texture;\n"
+			"void main()\n"
+			"{\n"
+			"	gl_FragColor = textureCubeLodEXT(Texture,TexCoord, 4.0);\n"
+			"}\n";
+
+		/*TArray<uint8> Code;
+		Code.Add( strlen( TestFragmentProgram ) + 1 );
+		FMemory::Memcpy( Code.GetData(), TestFragmentProgram, Code.Num() );*/
+
+		FOpenGL::bRequiresDontEmitPrecisionForTextureSamplers = false;
+		FOpenGL::bRequiresTextureCubeLodEXTToTextureCubeLodDefine = false;
+
+
+		FOpenGLCodeHeader Header;
+		Header.FrequencyMarker = 0x5053;
+		Header.GlslMarker = 0x474c534c;
+		TArray<uint8> Code;
+		{
+			FMemoryWriter Writer(Code );
+			Writer << Header;
+			Writer.Serialize( (void*)(TestFragmentProgram), strlen(TestFragmentProgram)+1);
+			Writer.Close();
+		}
+		// try to compile without any hacks
+		TRefCountPtr<FOpenGLPixelShader> PixelShader = (FOpenGLPixelShader*)(RHICreatePixelShader(Code).GetReference());
+
+		if ( VerifyCompiledShader( PixelShader->Resource, TestFragmentProgram, false ) )
+		{
+			UE_LOG(LogRHI,Display,TEXT("Shaders compile fine no need to enable hacks"));
+			// we are done
+			return;
+		}
+		
+
+
+		FOpenGLES2::bRequiresDontEmitPrecisionForTextureSamplers = true;
+		FOpenGLES2::bRequiresTextureCubeLodEXTToTextureCubeLodDefine = false;
+
+		// second most number of devices fall into this hack category
+		// try to compile without using precision for texture samplers 
+		// Samsung Galaxy Express	Samsung Galaxy S3	Samsung Galaxy S3 mini	Samsung Galaxy Tab GT-P1000	Samsung Galaxy Tab 2
+		PixelShader = (FOpenGLPixelShader*)(RHICreatePixelShader(Code).GetReference());
+
+		if ( VerifyCompiledShader( PixelShader->Resource, TestFragmentProgram, false ) )
+		{
+			UE_LOG(LogRHI,Warning,TEXT("Enabling shader compiler hack to remove precision modifiers for texture samplers"));
+
+			// we are done
+			return;
+		}
+
+
+		FOpenGL::bRequiresDontEmitPrecisionForTextureSamplers = false;
+		FOpenGL::bRequiresTextureCubeLodEXTToTextureCubeLodDefine = true;
+
+		// third most likely Samsung Galaxy Tab GT-P1000
+		PixelShader = (FOpenGLPixelShader*)(RHICreatePixelShader(Code).GetReference());
+
+		if ( VerifyCompiledShader( PixelShader->Resource, TestFragmentProgram, false ) )
+		{
+			UE_LOG(LogRHI,Warning,TEXT("Enabling shader compiler hack to redefine textureCubeLodEXT to textureCubeLod"));
+			// we are done
+			return;
+		}
+		
+		UE_LOG(LogRHI, Warning, TEXT("Unable to find a test shader that compiles try running anyway"));
+
+	}
+#endif
 }
 
 void FOpenGLDynamicRHI::Shutdown()
diff --git a/Engine/Source/Runtime/OpenGLDrv/Private/OpenGLES2.cpp b/Engine/Source/Runtime/OpenGLDrv/Private/OpenGLES2.cpp
--- a/Engine/Source/Runtime/OpenGLDrv/Private/OpenGLES2.cpp
+++ b/Engine/Source/Runtime/OpenGLDrv/Private/OpenGLES2.cpp
@@ -104,6 +104,12 @@ bool FOpenGLES2::bSupportsCopyTextureLevels = false;
 /** GL_EXT_texture_storage */
 bool FOpenGLES2::bSupportsTextureStorageEXT = false;
 
+/* This is a hack to remove the calls to "precision sampler" defaults which are produced by the cross compiler however don't compile on some android platforms */
+bool FOpenGLES2::bRequiresDontEmitPrecisionForTextureSamplers = false;
+
+/* Some android platforms require textureCubeLod to be used some require textureCubeLodEXT however they either inconsistently or don't use the GL_TextureCubeLodEXT extension definition */
+bool FOpenGLES2::bRequiresTextureCubeLodEXTToTextureCubeLodDefine = false;
+
 
 bool FOpenGLES2::SupportsDisjointTimeQueries()
 {
diff --git a/Engine/Source/Runtime/OpenGLDrv/Private/OpenGLShaders.cpp b/Engine/Source/Runtime/OpenGLDrv/Private/OpenGLShaders.cpp
--- a/Engine/Source/Runtime/OpenGLDrv/Private/OpenGLShaders.cpp
+++ b/Engine/Source/Runtime/OpenGLDrv/Private/OpenGLShaders.cpp
@@ -20,9 +20,9 @@ const uint32 SizeOfFloat4 = 16;
 const uint32 NumFloatsInFloat4 = 4;
 
 /**
- * Verify that an OpenGL shader has compiled successfully.
+ * Verify that an OpenGL shader has compiled successfully. 
  */
-static bool VerifyCompiledShader(GLuint Shader, const ANSICHAR* GlslCode)
+static bool VerifyCompiledShader(GLuint Shader, const ANSICHAR* GlslCode )
 {
 	SCOPE_CYCLE_COUNTER(STAT_OpenGLShaderCompileVerifyTime);
 
@@ -34,17 +34,39 @@ static bool VerifyCompiledShader(GLuint Shader, const ANSICHAR* GlslCode)
 		ANSICHAR DefaultLog[] = "No log";
 		ANSICHAR *CompileLog = DefaultLog;
 		glGetShaderiv(Shader, GL_INFO_LOG_LENGTH, &LogLength);
+#if PLATFORM_ANDROID
+		if ( LogLength == 0 )
+		{
+			// make it big anyway
+			// there was a bug in android 2.2 where glGetShaderiv would return 0 even though there was a error message
+			// https://code.google.com/p/android/issues/detail?id=9953
+			LogLength = 4096;
+		}
+#endif
 		if (LogLength > 1)
 		{
 			CompileLog = (ANSICHAR *)FMemory::Malloc(LogLength);
 			glGetShaderInfoLog(Shader, LogLength, NULL, CompileLog);
 		}
-
+	
 #if DEBUG_GL_SHADERS
 		if (GlslCode)
 		{
 			UE_LOG(LogRHI,Error,TEXT("Shader:\n%s"),ANSI_TO_TCHAR(GlslCode));
-		}
+
+#if 0
+			const ANSICHAR *Temp = GlslCode;
+
+			for ( int i = 0; i < 30 && (*Temp != '\0'); ++i )
+			{
+				FString Converted = ANSI_TO_TCHAR( Temp );
+				Converted.LeftChop( 256 );
+
+				UE_LOG(LogRHI,Display,TEXT("%s"), *Converted );
+				Temp += Converted.Len();
+			}
+#endif
+		}	
 #endif
 		UE_LOG(LogRHI,Fatal,TEXT("Failed to compile shader. Compile log:\n%s"), ANSI_TO_TCHAR(CompileLog));
 
@@ -349,7 +371,7 @@ ShaderType* CompileOpenGLShader(const TArray<uint8>& Code)
 					"#define textureCubeLodEXT textureLod	\n"
 					"\n"
 					"#define gl_FragColor out_FragColor	\n"
-					"out mediump vec4 out_FragColor;";
+					"out mediump vec4 out_FragColor;\n";
 
 				// See if we need to skip any #extension string
 				ANSICHAR* Temp = SkipShaderExtensionText(const_cast<ANSICHAR*>(GlslCode));
@@ -362,18 +384,39 @@ ShaderType* CompileOpenGLShader(const TArray<uint8>& Code)
 				ReplaceShaderSubstring(const_cast<ANSICHAR*>(GlslCode), "varying", "in");
 			}
 		}
-		else if(!FOpenGL::SupportsShaderTextureLod() || !FOpenGL::SupportsShaderTextureCubeLod())
+#if PLATFORM_ANDROID
+		else if ( (TypeEnum == GL_FRAGMENT_SHADER) &&
+			FOpenGL::RequiresDontEmitPrecisionForTextureSamplers() )
 		{
-			Prologue = "#define textureCubeLodEXT(a, b, c) textureCube(a, b) \n";
+			// Daniel: This device has some shader compiler compatibility issues force them to be disabled
+			//			The cross compiler will put the DONTEMITEXTENSIONSHADERTEXTURELODENABLE define around incompatible sections of code
+			Prologue = "#define DONTEMITEXTENSIONSHADERTEXTURELODENABLE\n"
+				"#define DONTEMITSAMPLERDEFAULTPRECISION\n"
+				"#define texture2DLodEXT(a, b, c) texture2D(a, b) \n"
+				"#define textureCubeLodEXT(a, b, c) textureCube(a, b) \n";
 		}
+		else if ( ( TypeEnum == GL_FRAGMENT_SHADER) && 
+			FOpenGL::RequiresTextureCubeLodEXTToTextureCubeLodDefine() )
+		{
+			Prologue = "#define textureCubeLodEXT textureCubeLod \n";
+		}
+#else
 		else if(!FOpenGL::SupportsTextureCubeLodEXT())
 		{
 			Prologue = "#define textureCubeLodEXT textureCubeLod \n";
 		}
+#endif
+		else if(!FOpenGL::SupportsShaderTextureLod() || !FOpenGL::SupportsShaderTextureCubeLod())
+		{
+			Prologue = "#define textureCubeLodEXT(a, b, c) textureCube(a, b) \n";
+		}
 		else 
 		{
 			Prologue = "";
 		}
+
+		//UE_LOG(LogRHI, Display,TEXT("Prologue is %s"), ANSI_TO_TCHAR( Prologue ));
+
 		//Assemble the source strings into an array to pass into the compiler.
 		const GLchar* ShaderSourceStrings[4] = { VersionString, ExtensionString, Prologue, GlslCode };
 		const GLint ShaderSourceLen[4] = { (GLint)(strlen(VersionString)), (GLint)(strlen(ExtensionString)), (GLint)(strlen(Prologue)), (GLint)(strlen(GlslCode)) };
diff --git a/Engine/Source/Runtime/OpenGLDrv/Private/OpenGLViewport.cpp b/Engine/Source/Runtime/OpenGLDrv/Private/OpenGLViewport.cpp
--- a/Engine/Source/Runtime/OpenGLDrv/Private/OpenGLViewport.cpp
+++ b/Engine/Source/Runtime/OpenGLDrv/Private/OpenGLViewport.cpp
@@ -83,7 +83,7 @@ void FOpenGLDynamicRHI::RHIBeginDrawingViewport(FViewportRHIParamRef ViewportRHI
 	if( CurrentContext != CONTEXT_Rendering )
 	{
 		check(CurrentContext == CONTEXT_Shared);
-		check(!bIsRenderingContextAcquired);
+		check(!bIsRenderingContextAcquired || !GUseThreadedRendering);
 		bRevertToSharedContextAfterDrawingViewport = true;
 		PlatformRenderingContextSetup(PlatformDevice);
 	}
diff --git a/Engine/Source/Runtime/OpenGLDrv/Public/OpenGLES2.h b/Engine/Source/Runtime/OpenGLDrv/Public/OpenGLES2.h
--- a/Engine/Source/Runtime/OpenGLDrv/Public/OpenGLES2.h
+++ b/Engine/Source/Runtime/OpenGLDrv/Public/OpenGLES2.h
@@ -106,6 +106,10 @@ struct FOpenGLES2 : public FOpenGLBase
 	static FORCEINLINE bool SupportsShaderTextureCubeLod()				{ return bSupportsShaderTextureCubeLod; }
 	static FORCEINLINE bool SupportsCopyTextureLevels()					{ return bSupportsCopyTextureLevels; }
 
+	static FORCEINLINE bool RequiresDontEmitPrecisionForTextureSamplers() { return bRequiresDontEmitPrecisionForTextureSamplers; }
+	static FORCEINLINE bool RequiresTextureCubeLodEXTToTextureCubeLodDefine() { return bRequiresTextureCubeLodEXTToTextureCubeLodDefine; }
+
+
 	// On iOS both glMapBufferOES() and glBufferSubData() for immediate vertex and index data
 	// is the slow path (they both hit GPU sync and data cache flush in driver according to profiling in driver symbols).
 	// Turning this to false reverts back to not using vertex and index buffers 
@@ -417,6 +421,14 @@ protected:
 
 	/** GL_EXT_texture_storage */
 	static bool bSupportsTextureStorageEXT;
+
+public:
+	/* This is a hack to remove the calls to "precision sampler" defaults which are produced by the cross compiler however don't compile on some android platforms */
+	static bool bRequiresDontEmitPrecisionForTextureSamplers;
+
+	/* Some android platforms require textureCubeLod to be used some require textureCubeLodEXT however they either inconsistently or don't use the GL_TextureCubeLodEXT extension definition */
+	static bool bRequiresTextureCubeLodEXTToTextureCubeLodDefine;
+
 };
 
 
