fixes for HTML5 build and execution errors, results of some merge problems
#ue4
#html5

diff --git a/Engine/Source/Runtime/Engine/Private/WheeledVehicle.cpp b/Engine/Source/Runtime/Engine/Private/WheeledVehicle.cpp
--- a/Engine/Source/Runtime/Engine/Private/WheeledVehicle.cpp
+++ b/Engine/Source/Runtime/Engine/Private/WheeledVehicle.cpp
@@ -34,7 +34,7 @@ void AWheeledVehicle::DisplayDebug(UCanvas* Canvas, const FDebugDisplayInfo& Deb
 
 	Super::DisplayDebug(Canvas, DebugDisplay, YL, YPos);
 
-#if WITH_PHYSX
+#if WITH_VEHICLE
 	if (DebugDisplay.IsDisplayOn(NAME_Vehicle))
 	{
 		GetVehicleMovementComponent()->DrawDebug(Canvas, YL, YPos);
diff --git a/Engine/Source/Runtime/OpenGLDrv/Private/OpenGLShaders.cpp b/Engine/Source/Runtime/OpenGLDrv/Private/OpenGLShaders.cpp
--- a/Engine/Source/Runtime/OpenGLDrv/Private/OpenGLShaders.cpp
+++ b/Engine/Source/Runtime/OpenGLDrv/Private/OpenGLShaders.cpp
@@ -400,16 +400,17 @@ ShaderType* CompileOpenGLShader(const TArray<uint8>& Code)
 		{
 			Prologue = "#define textureCubeLodEXT textureCubeLod \n";
 		}
-#else
-		else if(!FOpenGL::SupportsTextureCubeLodEXT())
-		{
-			Prologue = "#define textureCubeLodEXT textureCubeLod \n";
-		}
 #endif
 		else if(!FOpenGL::SupportsShaderTextureLod() || !FOpenGL::SupportsShaderTextureCubeLod())
 		{
 			Prologue = "#define textureCubeLodEXT(a, b, c) textureCube(a, b) \n";
 		}
+#if !PLATFORM_ANDROID
+		else if(!FOpenGL::SupportsTextureCubeLodEXT())
+		{
+			Prologue = "#define textureCubeLodEXT textureCubeLod \n";
+		}
+#endif
 		else 
 		{
 			Prologue = "";
