Fix the off-by-default additional OpenGL RHI stats for MapBuffer & DrawPrimitive time.

diff --git a/Engine/Source/Runtime/OpenGLDrv/Private/OpenGLDrvPrivate.h b/Engine/Source/Runtime/OpenGLDrv/Private/OpenGLDrvPrivate.h
--- a/Engine/Source/Runtime/OpenGLDrv/Private/OpenGLDrvPrivate.h
+++ b/Engine/Source/Runtime/OpenGLDrv/Private/OpenGLDrvPrivate.h
@@ -76,8 +76,8 @@ DECLARE_CYCLE_STAT_EXTERN(TEXT("Emulated Uniform buffer time"), STAT_OpenGLEmula
 DECLARE_DWORD_ACCUMULATOR_STAT_EXTERN(TEXT("Uniform buffer pool num free"),STAT_OpenGLNumFreeUniformBuffers,STATGROUP_OpenGLRHI, );
 
 #if OPENGLRHI_DETAILED_STATS
-DECLARE_CYCLE_STAT_EXTERN(TEXT("PrawPrimitive Time"),STAT_OpenGLDrawPrimitiveTime,STATGROUP_OpenGLRHI, );
-DECLARE_CYCLE_STAT_EXTERN(TEXT("PrawPrimitiveUP Time"),STAT_OpenGLDrawPrimitiveUPTime,STATGROUP_OpenGLRHI, );
+DECLARE_CYCLE_STAT_EXTERN(TEXT("DrawPrimitive Time"),STAT_OpenGLDrawPrimitiveTime,STATGROUP_OpenGLRHI, );
+DECLARE_CYCLE_STAT_EXTERN(TEXT("DrawPrimitiveUP Time"),STAT_OpenGLDrawPrimitiveUPTime,STATGROUP_OpenGLRHI, );
 #endif
 
 enum EOpenGLCurrentContext
diff --git a/Engine/Source/Runtime/OpenGLDrv/Private/OpenGLUtil.cpp b/Engine/Source/Runtime/OpenGLDrv/Private/OpenGLUtil.cpp
--- a/Engine/Source/Runtime/OpenGLDrv/Private/OpenGLUtil.cpp
+++ b/Engine/Source/Runtime/OpenGLDrv/Private/OpenGLUtil.cpp
@@ -53,6 +53,13 @@ DEFINE_STAT(STAT_OpenGLEmulatedUniformBufferTime);
 DEFINE_STAT(STAT_OpenGLFreeUniformBufferMemory);
 DEFINE_STAT(STAT_OpenGLNumFreeUniformBuffers);
 
+#if OPENGLRHI_DETAILED_STATS
+DEFINE_STAT(STAT_OpenGLDrawPrimitiveTime);
+DEFINE_STAT(STAT_OpenGLDrawPrimitiveUPTime);
+DEFINE_STAT(STAT_OpenGLMapBufferTime);
+DEFINE_STAT(STAT_OpenGLUnmapBufferTime);
+#endif
+
 void IncrementBufferMemory(GLenum Type, bool bStructuredBuffer, uint32 NumBytes)
 {
 	if (bStructuredBuffer)
