Marginal gains for MTGL performance on Mac OS X.
When MTGL is enabled don't use glMapBuffer - this API causes a full MTGL thread synchronisation until the corresponding glUnmapBuffer which makes MTGL less effective. Don't label objects in MTGL mode as they will cause an MTGL thread synchronisation. Also, use glFlushRenderAPPLE within the Mac specific code as it should be slightly faster.

diff --git a/Engine/Source/Runtime/OpenGLDrv/Private/Mac/OpenGLMac.cpp b/Engine/Source/Runtime/OpenGLDrv/Private/Mac/OpenGLMac.cpp
--- a/Engine/Source/Runtime/OpenGLDrv/Private/Mac/OpenGLMac.cpp
+++ b/Engine/Source/Runtime/OpenGLDrv/Private/Mac/OpenGLMac.cpp
@@ -98,6 +98,14 @@ static NSOpenGLContext* CreateContext( NSOpenGLContext* SharedContext )
 	if (FParse::Param(FCommandLine::Get(),TEXT("openglUseMacMTEngine")))
 	{
 		CGLEnable((CGLContextObj)[Context CGLContextObj], kCGLCEMPEngine);
+		
+		// Disable OpenGL.UseMapBuffer when using MTGL to reduce the number of context synchronisation points.
+		// All calls to glMapBuffer/Range will stall the MTGL thread, even with the unsynchronized bit set, so we want to avoid it.
+		static auto CVar = IConsoleManager::Get().FindConsoleVariable(TEXT("OpenGL.UseMapBuffer"));
+		if(CVar && CVar->GetInt() == 1)
+		{
+			CVar->Set(0);
+		}
 	}
 
 	return Context;
@@ -117,7 +125,7 @@ public:
 		{
 			if (PreviousContext)
 			{
-				glFlush();
+				glFlushRenderAPPLE();
 			}
 			[Context makeCurrentContext];
 		}
@@ -129,7 +137,7 @@ public:
 		{
 			SCOPED_AUTORELEASE_POOL;
 
-			glFlush();
+			glFlushRenderAPPLE();
 			if (PreviousContext)
 			{
 				[PreviousContext makeCurrentContext];
@@ -432,7 +440,7 @@ struct FPlatformOpenGLDevice
 #endif
 		InitDefaultGLContextState();
 		
-		glFlush();
+		glFlushRenderAPPLE();
 		
 		[SharedContext.OpenGLContext makeCurrentContext];
 		InitDebugContext();
@@ -714,6 +722,12 @@ void PlatformBlitToViewport( FPlatformOpenGLDevice* Device, FPlatformOpenGLConte
 		{
 			FScopeContext ScopeContext(Context->OpenGLContext);
 			
+			if(Context->OpenGLView.bNeedsUpdate)
+			{
+				Context->OpenGLView.bNeedsUpdate = false;
+				[Context->OpenGLContext update];
+			}
+			
 			// OpenGL state necessary for blit is set up in PlatformResizeGLContext(), and should be correct here,
 			// as viewport contexts aren't bound at any other occasion.
 			glBindFramebuffer(GL_DRAW_FRAMEBUFFER, 0);
@@ -728,12 +742,6 @@ void PlatformBlitToViewport( FPlatformOpenGLDevice* Device, FPlatformOpenGLConte
 				GL_COLOR_BUFFER_BIT,
 				GL_NEAREST
             );
-			
-			if(Context->OpenGLView.bNeedsUpdate)
-			{
-				Context->OpenGLView.bNeedsUpdate = false;
-				[Context->OpenGLContext update];
-			}
 
 			if (bPresent)
 			{
@@ -890,7 +898,7 @@ void PlatformRenderingContextSetup(FPlatformOpenGLDevice* Device)
 
 	if ([NSOpenGLContext currentContext])
 	{
-		glFlush();
+		glFlushRenderAPPLE();
 	}
 	[Device->RenderingContext.OpenGLContext makeCurrentContext];
 }
@@ -903,7 +911,7 @@ void PlatformSharedContextSetup(FPlatformOpenGLDevice* Device)
 
 	if ([NSOpenGLContext currentContext])
 	{
-		glFlush();
+		glFlushRenderAPPLE();
 	}
 	[Device->SharedContext.OpenGLContext makeCurrentContext];
 }
@@ -913,7 +921,7 @@ void PlatformNULLContextSetup()
 	SCOPED_AUTORELEASE_POOL;
 	if ([NSOpenGLContext currentContext])
 	{
-		glFlush();
+		glFlushRenderAPPLE();
 	}
 	[NSOpenGLContext clearCurrentContext];
 }
@@ -1406,7 +1414,8 @@ void FMacOpenGL::ProcessExtensions(const FString& ExtensionsString)
 		glBlendEquationi = (PFNGLBLENDEQUATIONIARBPROC)dlsym(RTLD_SELF, "glBlendEquationi");
 	}
 	
-	if(ExtensionsString.Contains(TEXT("GL_EXT_debug_label")))
+	// Don't label objects with MTGL, it causes synchronisation of the MTGL thread.
+	if(ExtensionsString.Contains(TEXT("GL_EXT_debug_label")) && !FParse::Param(FCommandLine::Get(),TEXT("openglUseMacMTEngine")))
 	{
 		glLabelObjectEXT = (PFNGLLABELOBJECTEXTPROC)dlsym(RTLD_SELF, "glLabelObjectEXT");
 	}
diff --git a/Engine/Source/Runtime/OpenGLDrv/Private/Mac/OpenGLMac.h b/Engine/Source/Runtime/OpenGLDrv/Private/Mac/OpenGLMac.h
--- a/Engine/Source/Runtime/OpenGLDrv/Private/Mac/OpenGLMac.h
+++ b/Engine/Source/Runtime/OpenGLDrv/Private/Mac/OpenGLMac.h
@@ -102,7 +102,7 @@ struct FMacOpenGL : public FOpenGL3
 		if( glTexStorage3D )
 		{
 			glTexStorage3D( Target, Levels, InternalFormat, Width, Height, Depth);
-			glFlush(); // Workaround for radr://15553950, TTP# 315197
+			glFlushRenderAPPLE(); // Workaround for radr://15553950, TTP# 315197
 		}
 		else
 		{
@@ -131,7 +131,7 @@ struct FMacOpenGL : public FOpenGL3
 		if( glTexStorage2D )
 		{
 			glTexStorage2D(Target, Levels, InternalFormat, Width, Height);
-			glFlush(); // Workaround for radr://15553950, TTP# 315197
+			glFlushRenderAPPLE(); // Workaround for radr://15553950, TTP# 315197
 			return true;
 		}
 		else
