Fixes for a few iOS rendering bugs.

1) View.ExposureScale is now a float4 instead of a scalar. This appears to workaround an iOS driver bug. This workaround can be moved in to the cross compiler after 4.2.
2) Clamp the number of supported vertex uniform vectors to be at least 256. The engine assumes it has at least 256 to work with for skinned meshes. This prevents a memory overwrite bug. The iOS driver works with more uniforms than it claims to support.
3) Fixed an off-by-one bug in dirty range tracking for global uniform buffers.

diff --git a/Engine/Shaders/BasePassForForwardShadingPixelShader.usf b/Engine/Shaders/BasePassForForwardShadingPixelShader.usf
--- a/Engine/Shaders/BasePassForForwardShadingPixelShader.usf
+++ b/Engine/Shaders/BasePassForForwardShadingPixelShader.usf
@@ -225,7 +225,9 @@ void Main(
 	#endif
 
 	#if !MATERIALBLENDING_MODULATE
-		OutColor.rgb *= View.ExposureScale;
+		// The exposure scale is just a scalar but needs to be a float4 to workaround a driver bug on IOS.
+		// After 4.2 we can put the workaround in the cross compiler.
+		OutColor.rgb *= View.ExposureScale.xyz;
 	#endif
 
 	#if OUTPUT_GAMMA_SPACE
diff --git a/Engine/Source/Runtime/Engine/Public/SceneManagement.h b/Engine/Source/Runtime/Engine/Public/SceneManagement.h
--- a/Engine/Source/Runtime/Engine/Public/SceneManagement.h
+++ b/Engine/Source/Runtime/Engine/Public/SceneManagement.h
@@ -2052,6 +2052,9 @@ BEGIN_UNIFORM_BUFFER_STRUCT(FViewUniformShaderParameters,ENGINE_API)
 	DECLARE_UNIFORM_BUFFER_STRUCT_MEMBER(FVector4,ViewSizeAndSceneTexelSize)
 	DECLARE_UNIFORM_BUFFER_STRUCT_MEMBER(FVector4,ViewOrigin)
 	DECLARE_UNIFORM_BUFFER_STRUCT_MEMBER(FVector4,TranslatedViewOrigin)
+	// The exposure scale is just a scalar but needs to be a float4 to workaround a driver bug on IOS.
+	// After 4.2 we can put the workaround in the cross compiler.
+	DECLARE_UNIFORM_BUFFER_STRUCT_MEMBER_EX(FVector4,ExposureScale, EShaderPrecisionModifier::Half)
 	DECLARE_UNIFORM_BUFFER_STRUCT_MEMBER_EX(FVector4,DiffuseOverrideParameter, EShaderPrecisionModifier::Half)
 	DECLARE_UNIFORM_BUFFER_STRUCT_MEMBER_EX(FVector4,SpecularOverrideParameter, EShaderPrecisionModifier::Half)
 	DECLARE_UNIFORM_BUFFER_STRUCT_MEMBER_EX(FVector4,NormalOverrideParameter, EShaderPrecisionModifier::Half)
@@ -2076,7 +2079,6 @@ BEGIN_UNIFORM_BUFFER_STRUCT(FViewUniformShaderParameters,ENGINE_API)
 	DECLARE_UNIFORM_BUFFER_STRUCT_MEMBER_ARRAY(FVector4,TranslucencyLightingVolumeMin,[TVC_MAX])
 	DECLARE_UNIFORM_BUFFER_STRUCT_MEMBER_ARRAY(FVector4,TranslucencyLightingVolumeInvSize,[TVC_MAX])
 	DECLARE_UNIFORM_BUFFER_STRUCT_MEMBER(FVector4,TemporalAAParams)
-	DECLARE_UNIFORM_BUFFER_STRUCT_MEMBER_EX(float,ExposureScale, EShaderPrecisionModifier::Half)
 	DECLARE_UNIFORM_BUFFER_STRUCT_MEMBER(float,DepthOfFieldFocalDistance)
 	DECLARE_UNIFORM_BUFFER_STRUCT_MEMBER(float,DepthOfFieldScale)
 	DECLARE_UNIFORM_BUFFER_STRUCT_MEMBER(float,DepthOfFieldFocalLength)
diff --git a/Engine/Source/Runtime/OpenGLDrv/Private/OpenGLES2.cpp b/Engine/Source/Runtime/OpenGLDrv/Private/OpenGLES2.cpp
--- a/Engine/Source/Runtime/OpenGLDrv/Private/OpenGLES2.cpp
+++ b/Engine/Source/Runtime/OpenGLDrv/Private/OpenGLES2.cpp
@@ -122,6 +122,18 @@ void FOpenGLES2::ProcessQueryGLInt()
 {
 	LOG_AND_GET_GL_INT(GL_MAX_VERTEX_UNIFORM_VECTORS, 0, MaxVertexUniformComponents);
 	LOG_AND_GET_GL_INT(GL_MAX_FRAGMENT_UNIFORM_VECTORS, 0, MaxPixelUniformComponents);
+
+	const GLint RequiredMaxVertexUniformComponents = 256;
+	if (MaxVertexUniformComponents < RequiredMaxVertexUniformComponents)
+	{
+		UE_LOG(LogRHI,Warning,
+			TEXT("Device reports support for %d vertex uniform vectors, UE4 requires %d. Rendering artifacts may occur, especially with skeletal meshes. Some drivers, e.g. iOS, report a smaller number than is actually supported."),
+			MaxVertexUniformComponents,
+			RequiredMaxVertexUniformComponents
+			);
+	}
+	MaxVertexUniformComponents = FMath::Max<GLint>(MaxVertexUniformComponents, RequiredMaxVertexUniformComponents);
+
 	MaxVertexUniformComponents *= 4;
 	MaxPixelUniformComponents *= 4;
 	MaxGeometryUniformComponents = 0;
diff --git a/Engine/Source/Runtime/OpenGLDrv/Private/OpenGLShaders.cpp b/Engine/Source/Runtime/OpenGLDrv/Private/OpenGLShaders.cpp
--- a/Engine/Source/Runtime/OpenGLDrv/Private/OpenGLShaders.cpp
+++ b/Engine/Source/Runtime/OpenGLDrv/Private/OpenGLShaders.cpp
@@ -3,7 +3,7 @@
 /*=============================================================================
 	OpenGLShaders.cpp: OpenGL shader RHI implementation.
 =============================================================================*/
-
+ 
 #include "OpenGLDrvPrivate.h"
 #include "Shader.h"
 #include "GlobalShader.h"
@@ -1732,6 +1732,9 @@ void FOpenGLShaderParameterCache::InitializeResources(int32 UniformArraySize)
 {
 	check(GlobalUniformArraySize == -1);
 
+	// Uniform arrays have to be multiples of float4s.
+	UniformArraySize = Align(UniformArraySize,SizeOfFloat4);
+
 	PackedGlobalUniforms[0] = (uint8*)FMemory::Malloc(UniformArraySize * OGL_PACKED_TYPEINDEX_MAX);
 	PackedUniformsScratch[0] = (uint8*)FMemory::Malloc(UniformArraySize * OGL_PACKED_TYPEINDEX_MAX);
 
@@ -1747,7 +1750,7 @@ void FOpenGLShaderParameterCache::InitializeResources(int32 UniformArraySize)
 	for (int32 ArrayIndex = 0; ArrayIndex < OGL_PACKED_TYPEINDEX_MAX; ++ArrayIndex)
 	{
 		PackedGlobalUniformDirty[ArrayIndex].LowVector = 0;
-		PackedGlobalUniformDirty[ArrayIndex].HighVector = UniformArraySize;
+		PackedGlobalUniformDirty[ArrayIndex].HighVector = UniformArraySize / SizeOfFloat4;
 	}
 }
 
@@ -1774,7 +1777,7 @@ void FOpenGLShaderParameterCache::MarkAllDirty()
 	for (int32 ArrayIndex = 0; ArrayIndex < OGL_PACKED_TYPEINDEX_MAX; ++ArrayIndex)
 	{
 		PackedGlobalUniformDirty[ArrayIndex].LowVector = 0;
-		PackedGlobalUniformDirty[ArrayIndex].HighVector = GlobalUniformArraySize;
+		PackedGlobalUniformDirty[ArrayIndex].HighVector = GlobalUniformArraySize / SizeOfFloat4;
 	}
 }
 
@@ -1824,10 +1827,14 @@ void FOpenGLShaderParameterCache::CommitPackedGlobals(const FOpenGLLinkedProgram
 		const int32 NumVectors = PackedArrays[PackedUniform].Size / BytesPerRegister;
 		GLint Location = UniformInfo.Location;
 		const void* UniformData = PackedGlobalUniforms[ArrayIndex];
-		if (PackedGlobalUniformDirty[ArrayIndex].HighVector > PackedGlobalUniformDirty[ArrayIndex].LowVector)
+
+		// This has to be >=. If LowVector == HighVector it means that particular vector was written to.
+		if (PackedGlobalUniformDirty[ArrayIndex].HighVector >= PackedGlobalUniformDirty[ArrayIndex].LowVector)
 		{
 			const uint32 StartVector = PackedGlobalUniformDirty[ArrayIndex].LowVector;
-			uint32 NumDirtyVectors = PackedGlobalUniformDirty[ArrayIndex].HighVector - StartVector;
+			// The number of dirty vectors is the index of the highest vector written minus the first plus one.
+			// The plus one is important so that we upload the highest vector written.
+			uint32 NumDirtyVectors = PackedGlobalUniformDirty[ArrayIndex].HighVector - StartVector + 1;
 			NumDirtyVectors = FMath::Min(NumDirtyVectors, NumVectors - StartVector);
 			check(NumDirtyVectors);
 			UniformData = (uint8*)UniformData + PackedGlobalUniformDirty[ArrayIndex].LowVector * SizeOfFloat4;
@@ -1879,6 +1886,8 @@ void FOpenGLShaderParameterCache::CommitPackedUniformBuffers(FOpenGLLinkedProgra
 				const FOpenGLUniformBufferCopyInfo& Info = UniformBuffersCopyInfo[InfoIndex];
 				if (Info.SourceUBIndex == BufferIndex)
 				{
+					check((Info.SourceOffsetInFloats + Info.SizeInFloats) * sizeof(float) <= (uint32)GlobalUniformArraySize);
+
 					float* RESTRICT ScratchMem = (float*)PackedGlobalUniforms[Info.DestUBTypeIndex];
 					ScratchMem += Info.DestOffsetInFloats;
 					FMemory::Memcpy(ScratchMem, SourceData + Info.SourceOffsetInFloats, Info.SizeInFloats * sizeof(float));
diff --git a/Engine/Source/Runtime/Renderer/Private/SceneRendering.cpp b/Engine/Source/Runtime/Renderer/Private/SceneRendering.cpp
--- a/Engine/Source/Runtime/Renderer/Private/SceneRendering.cpp
+++ b/Engine/Source/Runtime/Renderer/Private/SceneRendering.cpp
@@ -412,7 +412,10 @@ TUniformBufferRef<FViewUniformShaderParameters> FViewInfo::CreateUniformBuffer(
 	}
 	
 	ViewUniformShaderParameters.RenderTargetSize = BufferSize;
-	ViewUniformShaderParameters.ExposureScale = FRCPassPostProcessEyeAdaptation::ComputeExposureScaleValue( *this );
+	// The exposure scale is just a scalar but needs to be a float4 to workaround a driver bug on IOS.
+	// After 4.2 we can put the workaround in the cross compiler.
+	float ExposureScale = FRCPassPostProcessEyeAdaptation::ComputeExposureScaleValue( *this );
+	ViewUniformShaderParameters.ExposureScale = FVector4(ExposureScale,ExposureScale,ExposureScale,ExposureScale);
 	ViewUniformShaderParameters.DepthOfFieldFocalDistance = FinalPostProcessSettings.DepthOfFieldFocalDistance;
 	ViewUniformShaderParameters.DepthOfFieldFocalRegion = FinalPostProcessSettings.DepthOfFieldFocalRegion;
 	ViewUniformShaderParameters.DepthOfFieldNearTransitionRegion = FinalPostProcessSettings.DepthOfFieldNearTransitionRegion;
