Fix GPU profiling on OS X by amending the GL_TIMESTAMP workaround so that it shouldn't crash the AMD drivers. Also enable in all non-shipping builds and add a commandline option -DisableMacGPUTimestamp to disable the timestamp emulation in debug or development if not required.

diff --git a/Engine/Source/Runtime/OpenGLDrv/Private/Mac/OpenGLMac.cpp b/Engine/Source/Runtime/OpenGLDrv/Private/Mac/OpenGLMac.cpp
--- a/Engine/Source/Runtime/OpenGLDrv/Private/Mac/OpenGLMac.cpp
+++ b/Engine/Source/Runtime/OpenGLDrv/Private/Mac/OpenGLMac.cpp
@@ -9,9 +9,6 @@
 	OpenGL context management.
 ------------------------------------------------------------------------------*/
 
-#define UE_EMULATE_TIMESTAMP 0 // (WITH_EDITOR && !UE_BUILD_SHIPPING) @todo: Crashes on AMD on window close
-
-
 char const* const CompositedBlitVertexShader = "#version 150\n"
 "const int VertexCount = 6;\n"
 "uniform int TextureDirection;\n"
@@ -306,7 +303,6 @@ bool GMacEnableCocoaScreenUpdates = true;
 
 @end
 
-#if UE_EMULATE_TIMESTAMP
 class FMacOpenGLTimer
 {
 public:
@@ -316,15 +312,21 @@ public:
 	void Begin(void);
 	void End(void);
 	uint64 GetResult(void);
+	uint64 GetAccumulatedResult(void);
 	int32 GetResultAvailable(void);
 	
 public:
 	TSharedPtr<FMacOpenGLTimer> Next;
+	TSharedPtr<FMacOpenGLTimer> Previous;
+	
+private:
+	void CacheResult(void);
 	
 private:
 	FPlatformOpenGLContext* Context;
 	uint32 Name;
 	uint64 Result;
+	uint64 Accumulated;
 	uint32 Available;
 	bool Cached;
 	bool Running;
@@ -352,9 +354,7 @@ private:
 	bool Cached;
 	TSharedPtr<FMacOpenGLTimer> Start;
 	TSharedPtr<FMacOpenGLTimer> Finish;
-	TWeakPtr<FMacOpenGLQuery> LastTimer;
 };
-#endif
 
 /** Platform specific OpenGL context. */
 struct FPlatformOpenGLContext
@@ -365,15 +365,9 @@ struct FPlatformOpenGLContext
 	GLuint				ViewportFramebuffer;
 	int32				SyncInterval;
 	GLuint				VertexArrayObject;	// one has to be generated and set for each context (OpenGL 3.2 Core requirements)
-	
-#if UE_EMULATE_TIMESTAMP
-	uint64 TimeElapsed;
-	TSharedPtr<FMacOpenGLTimer> FirstTimer;
-	TArray<TWeakPtr<FMacOpenGLTimer>> Timers;
 	TArray<TSharedPtr<FMacOpenGLQuery>> Queries;
 	TMap<GLenum, TSharedPtr<FMacOpenGLQuery>> RunningQueries;
-	TWeakPtr<FMacOpenGLQuery> LastTimer;
-#endif
+	TSharedPtr<FMacOpenGLTimer> LastTimer;
 };
 
 struct OpenGLContextInfo
@@ -393,6 +387,7 @@ static TAutoConsoleVariable<int32> CVarMacUseFrameBufferSRGB(
 		);
 
 bool GIsRunningOnIntelCard = false; // @todo: remove once Apple fixes radr://16223045 Changes to the GL separate blend state aren't always respected on Intel cards
+static bool GIsEmulatingTimestamp = (!UE_BUILD_SHIPPING); // @todo: Crashes on AMD on window close
 
 struct FPlatformOpenGLDevice
 {
@@ -430,14 +425,13 @@ struct FPlatformOpenGLDevice
 		glGenVertexArrays(1,&RenderingContext.VertexArrayObject);
 		glBindVertexArray(RenderingContext.VertexArrayObject);
 		
-#if UE_EMULATE_TIMESTAMP
-		RenderingContext.TimeElapsed = 0;
-		intptr_t RenderPtr = (intptr_t)&RenderingContext;
-		CGLSetParameter((CGLContextObj)[RenderingContext.OpenGLContext CGLContextObj], kCGLCPClientStorage, (GLint*)&RenderPtr);
-		RenderingContext.FirstTimer = TSharedPtr<FMacOpenGLTimer>(new FMacOpenGLTimer(&RenderingContext));
-		RenderingContext.Timers.Add(RenderingContext.FirstTimer);
-		RenderingContext.FirstTimer->Begin();
-#endif
+		if(GIsEmulatingTimestamp)
+		{
+			intptr_t RenderPtr = (intptr_t)&RenderingContext;
+			CGLSetParameter((CGLContextObj)[RenderingContext.OpenGLContext CGLContextObj], kCGLCPClientStorage, (GLint*)&RenderPtr);
+			RenderingContext.LastTimer = TSharedPtr<FMacOpenGLTimer>(new FMacOpenGLTimer(&RenderingContext));
+			RenderingContext.LastTimer->Begin();
+		}
 		InitDefaultGLContextState();
 		
 		glFlushRenderAPPLE();
@@ -447,14 +441,13 @@ struct FPlatformOpenGLDevice
 		glGenVertexArrays(1,&SharedContext.VertexArrayObject);
 		glBindVertexArray(SharedContext.VertexArrayObject);
 		
-#if UE_EMULATE_TIMESTAMP
-		SharedContext.TimeElapsed = 0;
-		intptr_t SharedPtr = (intptr_t)&SharedContext;
-		CGLSetParameter((CGLContextObj)[SharedContext.OpenGLContext CGLContextObj], kCGLCPClientStorage, (GLint*)&SharedPtr);
-		SharedContext.FirstTimer = TSharedPtr<FMacOpenGLTimer>(new FMacOpenGLTimer(&SharedContext));
-		SharedContext.Timers.Add(SharedContext.FirstTimer);
-		SharedContext.FirstTimer->Begin();
-#endif
+		if(GIsEmulatingTimestamp)
+		{
+			intptr_t SharedPtr = (intptr_t)&SharedContext;
+			CGLSetParameter((CGLContextObj)[SharedContext.OpenGLContext CGLContextObj], kCGLCPClientStorage, (GLint*)&SharedPtr);
+			SharedContext.LastTimer = TSharedPtr<FMacOpenGLTimer>(new FMacOpenGLTimer(&SharedContext));
+			SharedContext.LastTimer->Begin();
+		}
 		
 		SharedContextCompositeVertexShader = glCreateShader(GL_VERTEX_SHADER);
 		glShaderSource(SharedContextCompositeVertexShader, 1, &CompositedBlitVertexShader, NULL);
@@ -510,12 +503,12 @@ struct FPlatformOpenGLDevice
 			glBindVertexArray(0);
 			glDeleteVertexArrays(1,&RenderingContext.VertexArrayObject);
 			
-#if UE_EMULATE_TIMESTAMP
-			RenderingContext.RunningQueries.Empty();
-			RenderingContext.Queries.Empty();
-			RenderingContext.FirstTimer.Reset();
-			RenderingContext.Timers.Empty();
-#endif
+			if(GIsEmulatingTimestamp)
+			{
+				RenderingContext.RunningQueries.Empty();
+				RenderingContext.Queries.Empty();
+				RenderingContext.LastTimer.Reset();
+			}
 			
 			// Unbind the platform context from the CGL context
 			intptr_t Val = 0;
@@ -534,12 +527,12 @@ struct FPlatformOpenGLDevice
 			
 			glDeleteTextures(1, &SharedContextCompositeTexture);
 			
-#if UE_EMULATE_TIMESTAMP
-			SharedContext.RunningQueries.Empty();
-			SharedContext.Queries.Empty();
-			SharedContext.FirstTimer.Reset();
-			SharedContext.Timers.Empty();
-#endif
+			if(GIsEmulatingTimestamp)
+			{
+				SharedContext.RunningQueries.Empty();
+				SharedContext.Queries.Empty();
+				SharedContext.LastTimer.Reset();
+			}
 			
 			// Unbind the platform context from the CGL context
 			intptr_t Val = 0;
@@ -646,18 +639,15 @@ FPlatformOpenGLContext* PlatformCreateOpenGLContext(FPlatformOpenGLDevice* Devic
 		}
     }
 	
-#if UE_EMULATE_TIMESTAMP
-	// Bind the platform context into the CGL context
-	Context->TimeElapsed = 0;
-	CGLSetParameter((CGLContextObj)[Context->OpenGLContext CGLContextObj], kCGLCPClientStorage, (GLint*)&Context);
-	
+	if(GIsEmulatingTimestamp)
 	{
+		// Bind the platform context into the CGL context
+		CGLSetParameter((CGLContextObj)[Context->OpenGLContext CGLContextObj], kCGLCPClientStorage, (GLint*)&Context);
+		
 		FScopeContext ScopeContext(Context->OpenGLContext);
-		Context->FirstTimer = TSharedPtr<FMacOpenGLTimer>(new FMacOpenGLTimer(Context));
-		Context->Timers.Add(Context->FirstTimer);
-		Context->FirstTimer->Begin();
+		Context->LastTimer = TSharedPtr<FMacOpenGLTimer>(new FMacOpenGLTimer(Context));
+		Context->LastTimer->Begin();
 	}
-#endif
 	
 	return Context;
 }
@@ -667,19 +657,17 @@ void PlatformReleaseOpenGLContext(FPlatformOpenGLDevice* Device, FPlatformOpenGL
 	SCOPED_AUTORELEASE_POOL;
 
 	check(Context && Context->OpenGLView && Context->OpenGLContext);
-#if UE_EMULATE_TIMESTAMP
+	if(GIsEmulatingTimestamp)
 	{
 		FScopeContext ScopeContext(Context->OpenGLContext);
 		Context->RunningQueries.Empty();
 		Context->Queries.Empty();
-		Context->FirstTimer.Reset();
-		Context->Timers.Empty();
+		Context->LastTimer.Reset();
 		
 		// Unbind the platform context from the CGL context
 		intptr_t Val = 0;
 		CGLSetParameter((CGLContextObj)[Context->OpenGLContext CGLContextObj], kCGLCPClientStorage, (GLint*)&Val);
 	}
-#endif
 	{
 		FScopeLock ScopeLock(Device->ContextUsageGuard);
 
@@ -1215,62 +1203,65 @@ void PlatformGetNewRenderQuery( GLuint* OutQuery, uint64* OutQueryContext )
 
 		GLuint NewQuery = 0;
 		
-#if UE_EMULATE_TIMESTAMP
-		CGLContextObj Current = CGLGetCurrentContext();
-		FPlatformOpenGLContext* PlatformContext = nullptr;
-		if(Current && (CGLGetParameter(Current, kCGLCPClientStorage, (GLint *)&PlatformContext) == kCGLNoError) && PlatformContext)
+		if(GIsEmulatingTimestamp)
 		{
-			for( int32 Index = 0; Index < ReleasedQueries.Num(); ++Index )
+			CGLContextObj Current = CGLGetCurrentContext();
+			FPlatformOpenGLContext* PlatformContext = nullptr;
+			if(Current && (CGLGetParameter(Current, kCGLCPClientStorage, (GLint *)&PlatformContext) == kCGLNoError) && PlatformContext)
 			{
-				if( ReleasedQueries[Index].Context == Context )
+				for( int32 Index = 0; Index < ReleasedQueries.Num(); ++Index )
 				{
-					if(PlatformContext->Queries.Num() >= ReleasedQueries[Index].Query)
+					if( ReleasedQueries[Index].Context == Context )
 					{
-						PlatformContext->Queries[ReleasedQueries[Index].Query - 1].Reset();
+						if(PlatformContext->Queries.Num() >= ReleasedQueries[Index].Query)
+						{
+							PlatformContext->Queries[ReleasedQueries[Index].Query - 1].Reset();
+						}
+						ReleasedQueries.RemoveAt(Index);
+						Index--;
 					}
-					ReleasedQueries.RemoveAt(Index);
-					Index--;
 				}
-			}
-			
-			if(PlatformContext->Queries.Num() < UINT_MAX)
-			{
-				TSharedPtr<FMacOpenGLQuery> Query(new FMacOpenGLQuery(PlatformContext));
-				for(int32 Index = 0; Index < PlatformContext->Queries.Num(); ++Index)
+				
+				if(PlatformContext->Queries.Num() < UINT_MAX)
 				{
-					if(!PlatformContext->Queries[Index].IsValid())
+					TSharedPtr<FMacOpenGLQuery> Query(new FMacOpenGLQuery(PlatformContext));
+					for(int32 Index = 0; Index < PlatformContext->Queries.Num(); ++Index)
 					{
-						PlatformContext->Queries[Index] = Query;
-						NewQuery = Index + 1;
-						break;
+						if(!PlatformContext->Queries[Index].IsValid())
+						{
+							PlatformContext->Queries[Index] = Query;
+							NewQuery = Index + 1;
+							break;
+						}
+					}
+					
+					if(!NewQuery)
+					{
+						PlatformContext->Queries.Add(Query);
+						NewQuery = PlatformContext->Queries.Num();
 					}
-				}
-				
-				if(!NewQuery)
-				{
-					PlatformContext->Queries.Add(Query);
-					NewQuery = PlatformContext->Queries.Num();
 				}
 			}
 		}
-#else
-		// Check for possible query reuse
-		const int32 ArraySize = ReleasedQueries.Num();
-		for( int32 Index = 0; Index < ArraySize; ++Index )
+		else
 		{
-			if( ReleasedQueries[Index].Context == Context )
+			// Check for possible query reuse
+			const int32 ArraySize = ReleasedQueries.Num();
+			for( int32 Index = 0; Index < ArraySize; ++Index )
 			{
-				NewQuery = ReleasedQueries[Index].Query;
-				ReleasedQueries.RemoveAtSwap(Index);
-				break;
+				if( ReleasedQueries[Index].Context == Context )
+				{
+					NewQuery = ReleasedQueries[Index].Query;
+					ReleasedQueries.RemoveAtSwap(Index);
+					break;
+				}
 			}
-		}
 
-		if( !NewQuery )
-		{
-			glGenQueries( 1, &NewQuery );
+			if( !NewQuery )
+			{
+				glGenQueries( 1, &NewQuery );
+			}
 		}
-#endif
 
 		*OutQuery = NewQuery;
 		*OutQueryContext = (uint64)Context;
@@ -1284,19 +1275,22 @@ void PlatformReleaseRenderQuery( GLuint Query, uint64 QueryContext )
 	NSOpenGLContext* Context = [NSOpenGLContext currentContext];
 	if( (uint64)Context == QueryContext )
 	{
-#if UE_EMULATE_TIMESTAMP
-		CGLContextObj Current = CGLGetCurrentContext();
-		FPlatformOpenGLContext* PlatformContext = nullptr;
-		if(Current && (CGLGetParameter(Current, kCGLCPClientStorage, (GLint *)&PlatformContext) == kCGLNoError) && PlatformContext)
+		if(GIsEmulatingTimestamp)
 		{
-			if(PlatformContext->Queries.Num() >= Query)
+			CGLContextObj Current = CGLGetCurrentContext();
+			FPlatformOpenGLContext* PlatformContext = nullptr;
+			if(Current && (CGLGetParameter(Current, kCGLCPClientStorage, (GLint *)&PlatformContext) == kCGLNoError) && PlatformContext)
 			{
-				PlatformContext->Queries[Query - 1].Reset();
+				if(PlatformContext->Queries.Num() >= Query)
+				{
+					PlatformContext->Queries[Query - 1].Reset();
+				}
 			}
 		}
-#else
-		glDeleteQueries(1, &Query );
-#endif
+		else
+		{
+			glDeleteQueries(1, &Query );
+		}
 	}
 	else
 	{
@@ -1318,26 +1312,28 @@ void DeleteQueriesForCurrentContext( NSOpenGLContext* Context )
 		ReleasedQueriesGuard = new FCriticalSection;
 	}
 	
-#if UE_EMULATE_TIMESTAMP
-	CGLContextObj Current = (CGLContextObj)[Context CGLContextObj];
-	FPlatformOpenGLContext* PlatformContext = nullptr;
-	if(Current && (CGLGetParameter(Current, kCGLCPClientStorage, (GLint *)&PlatformContext) == kCGLNoError) && PlatformContext)
+	if(GIsEmulatingTimestamp)
 	{
-		FScopeLock Lock(ReleasedQueriesGuard);
-		for( int32 Index = 0; Index < ReleasedQueries.Num(); ++Index )
+		CGLContextObj Current = (CGLContextObj)[Context CGLContextObj];
+		FPlatformOpenGLContext* PlatformContext = nullptr;
+		if(Current && (CGLGetParameter(Current, kCGLCPClientStorage, (GLint *)&PlatformContext) == kCGLNoError) && PlatformContext)
 		{
-			if( ReleasedQueries[Index].Context == Context )
+			FScopeLock Lock(ReleasedQueriesGuard);
+			for( int32 Index = 0; Index < ReleasedQueries.Num(); ++Index )
 			{
-				if(PlatformContext->Queries.Num() >= ReleasedQueries[Index].Query)
+				if( ReleasedQueries[Index].Context == Context )
 				{
-					PlatformContext->Queries[ReleasedQueries[Index].Query - 1].Reset();
+					if(PlatformContext->Queries.Num() >= ReleasedQueries[Index].Query)
+					{
+						PlatformContext->Queries[ReleasedQueries[Index].Query - 1].Reset();
+					}
+					ReleasedQueries.RemoveAtSwap(Index);
+					--Index;
 				}
-				ReleasedQueries.RemoveAtSwap(Index);
-				--Index;
 			}
 		}
 	}
-#else
+	else
 	{
 		FScopeLock Lock(ReleasedQueriesGuard);
 		for( int32 Index = 0; Index < ReleasedQueries.Num(); ++Index )
@@ -1350,7 +1346,6 @@ void DeleteQueriesForCurrentContext( NSOpenGLContext* Context )
 			}
 		}
 	}
-#endif
 }
 
 bool PlatformContextIsCurrent( uint64 QueryContext )
@@ -1389,10 +1384,16 @@ void FMacOpenGL::ProcessExtensions(const FString& ExtensionsString)
 	ProcessQueryGLInt();
 	FOpenGL3::ProcessExtensions(ExtensionsString);
 	
-#if UE_EMULATE_TIMESTAMP
-	TimestampQueryBits = 64;
+	// Only use the timestamp emulation in a non-shipping build - end-users shouldn't care about this profiling feature.
+#if (!UE_BUILD_SHIPPING)
+	GIsEmulatingTimestamp = !FParse::Param(FCommandLine::Get(), TEXT("DisableMacGPUTimestamp"));
 #endif
 	
+	if(GIsEmulatingTimestamp)
+	{
+		TimestampQueryBits = 64;
+	}
+	
 	// Not all GPUs support the new 4.1 Core Profile required for GL_ARB_texture_storage
 	// Those that are stuck with 3.2 don't have this extension.
 	if(ExtensionsString.Contains(TEXT("GL_ARB_texture_storage")))
@@ -1436,11 +1437,11 @@ void FMacOpenGL::ProcessExtensions(const FString& ExtensionsString)
 	}
 }
 
-#if UE_EMULATE_TIMESTAMP
 FMacOpenGLTimer::FMacOpenGLTimer(FPlatformOpenGLContext* InContext)
 : Context(InContext)
 , Name(0)
 , Result(0)
+, Accumulated(0)
 , Available(0)
 , Cached(false)
 , Running(false)
@@ -1452,23 +1453,8 @@ FMacOpenGLTimer::FMacOpenGLTimer(FPlatformOpenGLContext* InContext)
 FMacOpenGLTimer::~FMacOpenGLTimer()
 {
 	End();
-	Context->TimeElapsed += GetResult();
-	for(uint32 i = 0; i < Context->Timers.Num(); ++i)
-	{
-		if(!Context->Timers[i].IsValid())
-		{
-			Context->Timers.RemoveAt(i);
-			--i;
-		}
-	}
+	Next.Reset();
 	glDeleteQueries(1, &Name);
-	
-	while(Next.IsValid() && Next.IsUnique())
-	{
-		TSharedPtr<FMacOpenGLTimer> NewNext = Next->Next;
-		Next->Next.Reset();
-		Next = NewNext;
-	}
 }
 
 void FMacOpenGLTimer::Begin(void)
@@ -1492,21 +1478,17 @@ void FMacOpenGLTimer::End(void)
 uint64 FMacOpenGLTimer::GetResult(void)
 {
 	check(!Running);
-	if(!Cached)
-	{
-		while(!GetResultAvailable())
-		{
-			usleep(1000);
-		}
-		if(Available)
-		{
-			glGetQueryObjectui64v(Name, GL_QUERY_RESULT, &Result);
-			Cached = true;
-		}
-	}
+	CacheResult();
 	return Result;
 }
 
+uint64 FMacOpenGLTimer::GetAccumulatedResult(void)
+{
+	check(!Running);
+	CacheResult();
+	return Accumulated;
+}
+
 int32 FMacOpenGLTimer::GetResultAvailable(void)
 {
 	check(!Running);
@@ -1517,6 +1499,27 @@ int32 FMacOpenGLTimer::GetResultAvailable(void)
 	return Available;
 }
 
+void FMacOpenGLTimer::CacheResult(void)
+{
+	check(!Running);
+	if(!Cached)
+	{
+		glGetQueryObjectui64v(Name, GL_QUERY_RESULT, &Result);
+		
+		if(Previous.IsValid())
+		{
+			Accumulated += Result;
+			Accumulated += Previous->GetAccumulatedResult();
+			
+			Previous->Next.Reset();
+			Previous.Reset();
+		}
+		
+		Available = true;
+		Cached = true;
+	}
+}
+
 FMacOpenGLQuery::FMacOpenGLQuery(FPlatformOpenGLContext* InContext)
 : Context(InContext)
 , Name(0)
@@ -1552,67 +1555,60 @@ void FMacOpenGLQuery::Begin(GLenum InTarget)
 		Cached = false;
 		Start.Reset();
 		Finish.Reset();
-		if(Target == GL_TIMESTAMP)
+		switch(Target)
 		{
-			for( auto It : Context->Queries )
+			case GL_TIMESTAMP:
 			{
-				if(It.IsValid() && It->Running == false && (It->Target == GL_TIMESTAMP || It->Target == GL_TIME_ELAPSED))
+				// There can be a lot of timestamps, clear emulated queries out to avoid problems.
+				for( auto It : Context->Queries )
 				{
-					if(It->GetResultAvailable())
+					if(It.IsValid() && It->Running == false && (It->Target == GL_TIMESTAMP || It->Target == GL_TIME_ELAPSED))
 					{
-						It->GetResult();
-					}
-					else
-					{
-						break;
+						if(It->GetResultAvailable())
+						{
+							It->GetResult();
+						}
+						else
+						{
+							break;
+						}
 					}
 				}
+				
+				check(Context->LastTimer.IsValid());
+				Context->LastTimer->End();
+				Start = Context->LastTimer;
+				
+				TSharedPtr<FMacOpenGLTimer> Current(new FMacOpenGLTimer(Context));
+				Current->Previous = Start;
+				Start->Next = Current;
+				Current->Begin();
+				Context->LastTimer = Current;
+				
+				break;
 			}
-			
-			LastTimer = Context->LastTimer;
-			Context->LastTimer = AsShared();
-			while(Context->Timers.Num() > 0 && !Context->Timers[0].IsValid())
-			{
-				Context->Timers.RemoveAt(0);
-			}
-			Start = Context->Timers[0].Pin();
-			if(Start == Context->FirstTimer)
-			{
-				Context->FirstTimer.Reset();
-			}
-			Finish = Context->Timers.Last().Pin();
-			TSharedPtr<FMacOpenGLTimer> Current(new FMacOpenGLTimer(Context));
-			Finish->Next = Current;
-			Finish->End();
-			Current->Begin();
-			Context->Timers.Add(Current);
-			if(!Context->RunningQueries.FindOrAdd(GL_TIME_ELAPSED).IsValid())
+			case GL_TIME_ELAPSED:
 			{
-				Context->FirstTimer = Current;
-			}
-		}
-		else if(Target == GL_TIME_ELAPSED)
-		{
-			LastTimer = Context->LastTimer;
-			Context->LastTimer = AsShared();
-			TSharedPtr<FMacOpenGLTimer> Current = Context->Timers.Last().Pin();
-			Start = TSharedPtr<FMacOpenGLTimer>(new FMacOpenGLTimer(Context));
-			Current->Next = Start;
-			Current->End();
-			Start->Begin();
-			if(Current == Context->FirstTimer)
-			{
-				Context->FirstTimer.Reset();
+				check(Context->LastTimer.IsValid());
+				Context->LastTimer->End();
+				
+				Start = TSharedPtr<FMacOpenGLTimer>(new FMacOpenGLTimer(Context));
+				Start->Previous = Context->LastTimer;
+				Context->LastTimer->Next = Start;
+				Start->Begin();
+				
+				Context->LastTimer = Start;
+				break;
 			}
-			Context->Timers.Add(Start);
-		}
-		else
-		{
-			if(!Name)
+			default:
 			{
-				glGenQueries(1, &Name);
+				if(!Name)
+				{
+					glGenQueries(1, &Name);
+				}
+				glBeginQuery(Target, Name);
+				break;
 			}
-			glBeginQuery(Target, Name);
 		}
 	}
 }
@@ -1622,13 +1618,13 @@ void FMacOpenGLQuery::End()
 	Running = false;
 	if(Target == GL_TIME_ELAPSED)
 	{
-		Finish = Context->Timers.Last().Pin();
+		Finish = Context->LastTimer;
 		TSharedPtr<FMacOpenGLTimer> Current(new FMacOpenGLTimer(Context));
+		Current->Previous = Finish;
 		Finish->Next = Current;
 		Finish->End();
 		Current->Begin();
-		Context->Timers.Add(Current);
-		Context->FirstTimer = Current;
+		Context->LastTimer = Current;
 	}
 	else if(Target != GL_TIMESTAMP)
 	{
@@ -1640,56 +1636,24 @@ uint64 FMacOpenGLQuery::GetResult(void)
 {
 	if(!Cached)
 	{
-		while(!GetResultAvailable())
+		switch(Target)
 		{
-			usleep(1000);
-		}
-		if(Available)
-		{
-			if(Target == GL_TIMESTAMP || Target == GL_TIME_ELAPSED)
-			{
-				TWeakPtr<FMacOpenGLQuery> PreviousQuery = LastTimer;
-				while(PreviousQuery.IsValid())
-				{
-					TSharedPtr<FMacOpenGLQuery> Query = PreviousQuery.Pin();
-					PreviousQuery = Query->LastTimer;
-					
-					// Don't recurse - it may blow the stack!
-					Query->LastTimer.Reset();
-					
-					// But do prefetch the previous results.
-					Query->GetResult();
-				}
-				
-				Result = 0;
-				TSharedPtr<FMacOpenGLTimer> Current = Start;
-				do
-				{
-					Result += Current->GetResult();
-					if(Current != Finish)
-					{
-						Current = Current->Next;
-					}
-					else
-					{
-						Current.Reset();
-					}
-				} while(Current.IsValid());
-				
-				if(Target == GL_TIMESTAMP)
-				{
-					Result += Context->TimeElapsed;
-				}
-				
+			case GL_TIMESTAMP:
+				Result = Start->GetAccumulatedResult();
 				Start.Reset();
 				Finish.Reset();
-			}
-			else
-			{
+				break;
+			case GL_TIME_ELAPSED:
+				Result = Finish->GetAccumulatedResult() - Start->GetAccumulatedResult();
+				Start.Reset();
+				Finish.Reset();
+				break;
+			default:
 				glGetQueryObjectui64v(Name, GL_QUERY_RESULT, &Result);
-			}
-			Cached = true;
+				break;
 		}
+		Available = true;
+		Cached = true;
 	}
 	return Result;
 }
@@ -1698,13 +1662,17 @@ int32 FMacOpenGLQuery::GetResultAvailable(void)
 {
 	if(!Available)
 	{
-		if(Target == GL_TIMESTAMP || Target == GL_TIME_ELAPSED)
-		{
-			Available = Finish->GetResultAvailable();
-		}
-		else
+		switch(Target)
 		{
-			glGetQueryObjectuiv(Name, GL_QUERY_RESULT_AVAILABLE, &Available);
+			case GL_TIMESTAMP:
+				Available = Start->GetResultAvailable();
+				break;
+			case GL_TIME_ELAPSED:
+				Available = Finish->GetResultAvailable();
+				break;
+			default:
+				glGetQueryObjectuiv(Name, GL_QUERY_RESULT_AVAILABLE, &Available);
+				break;
 		}
 	}
 	return Available;
@@ -1729,111 +1697,125 @@ TSharedPtr<FMacOpenGLQuery> GetQuery(FPlatformOpenGLContext* PlatformContext, GL
 	}
 	return Query;
 }
-#endif
 
 void FMacOpenGL::MacQueryTimestampCounter(GLuint QueryID)
 {
-#if UE_EMULATE_TIMESTAMP
-	CGLContextObj Current = CGLGetCurrentContext();
-	FPlatformOpenGLContext* PlatformContext = nullptr;
-	if(Current && (CGLGetParameter(Current, kCGLCPClientStorage, (GLint *)&PlatformContext) == kCGLNoError) && PlatformContext)
+	if(GIsEmulatingTimestamp)
 	{
-		TSharedPtr<FMacOpenGLQuery> Query = GetQuery(PlatformContext, QueryID);
-		if(Query.IsValid())
+		CGLContextObj Current = CGLGetCurrentContext();
+		FPlatformOpenGLContext* PlatformContext = nullptr;
+		if(Current && (CGLGetParameter(Current, kCGLCPClientStorage, (GLint *)&PlatformContext) == kCGLNoError) && PlatformContext)
 		{
-			Query->Begin(GL_TIMESTAMP);
-			Query->End();
+			TSharedPtr<FMacOpenGLQuery> Query = GetQuery(PlatformContext, QueryID);
+			if(Query.IsValid())
+			{
+				Query->Begin(GL_TIMESTAMP);
+				Query->End();
+			}
 		}
 	}
-#else
-	FOpenGL3::QueryTimestampCounter(QueryID);
-#endif
+	else
+	{
+		FOpenGL3::QueryTimestampCounter(QueryID);
+	}
 }
 
 void FMacOpenGL::MacBeginQuery(GLenum QueryType, GLuint QueryId)
 {
-#if UE_EMULATE_TIMESTAMP
-	CGLContextObj Current = CGLGetCurrentContext();
-	FPlatformOpenGLContext* PlatformContext = nullptr;
-	if(Current && (CGLGetParameter(Current, kCGLCPClientStorage, (GLint *)&PlatformContext) == kCGLNoError) && PlatformContext)
+	if(GIsEmulatingTimestamp)
 	{
-		TSharedPtr<FMacOpenGLQuery> Query = GetQuery(PlatformContext, QueryId);
-		if(Query.IsValid())
+		CGLContextObj Current = CGLGetCurrentContext();
+		FPlatformOpenGLContext* PlatformContext = nullptr;
+		if(Current && (CGLGetParameter(Current, kCGLCPClientStorage, (GLint *)&PlatformContext) == kCGLNoError) && PlatformContext)
 		{
-			Query->Begin(QueryType);
-			PlatformContext->RunningQueries.Add(QueryType, Query);
+			TSharedPtr<FMacOpenGLQuery> Query = GetQuery(PlatformContext, QueryId);
+			if(Query.IsValid())
+			{
+				Query->Begin(QueryType);
+				PlatformContext->RunningQueries.Add(QueryType, Query);
+			}
 		}
 	}
-#else
-	FOpenGL3::BeginQuery(QueryType, QueryId);
-#endif
+	else
+	{
+		FOpenGL3::BeginQuery(QueryType, QueryId);
+	}
 }
 
 void FMacOpenGL::MacEndQuery(GLenum QueryType)
 {
-#if UE_EMULATE_TIMESTAMP
-	CGLContextObj Current = CGLGetCurrentContext();
-	FPlatformOpenGLContext* PlatformContext = nullptr;
-	if(Current && (CGLGetParameter(Current, kCGLCPClientStorage, (GLint *)&PlatformContext) == kCGLNoError) && PlatformContext)
+	if(GIsEmulatingTimestamp)
 	{
-		TSharedPtr<FMacOpenGLQuery> Query = PlatformContext->RunningQueries.FindOrAdd(QueryType);
-		if(Query.IsValid())
+		CGLContextObj Current = CGLGetCurrentContext();
+		FPlatformOpenGLContext* PlatformContext = nullptr;
+		if(Current && (CGLGetParameter(Current, kCGLCPClientStorage, (GLint *)&PlatformContext) == kCGLNoError) && PlatformContext)
 		{
-			Query->End();
-			PlatformContext->RunningQueries[QueryType].Reset();
+			TSharedPtr<FMacOpenGLQuery> Query = PlatformContext->RunningQueries.FindOrAdd(QueryType);
+			if(Query.IsValid())
+			{
+				Query->End();
+				PlatformContext->RunningQueries[QueryType].Reset();
+			}
 		}
 	}
-#else
-	FOpenGL3::EndQuery(QueryType);
-#endif
+	else
+	{
+		FOpenGL3::EndQuery(QueryType);
+	}
 }
 
 void FMacOpenGL::MacGetQueryObject(GLuint QueryId, EQueryMode QueryMode, uint64 *OutResult)
 {
-#if UE_EMULATE_TIMESTAMP
-	CGLContextObj Current = CGLGetCurrentContext();
-	FPlatformOpenGLContext* PlatformContext = nullptr;
-	if(Current && (CGLGetParameter(Current, kCGLCPClientStorage, (GLint *)&PlatformContext) == kCGLNoError) && PlatformContext)
+	if(GIsEmulatingTimestamp)
 	{
-		TSharedPtr<FMacOpenGLQuery> Query = GetQuery(PlatformContext, QueryId);
-		if(Query.IsValid())
+		CGLContextObj Current = CGLGetCurrentContext();
+		FPlatformOpenGLContext* PlatformContext = nullptr;
+		if(Current && (CGLGetParameter(Current, kCGLCPClientStorage, (GLint *)&PlatformContext) == kCGLNoError) && PlatformContext)
 		{
-			if(QueryMode == QM_Result)
-			{
-				*OutResult = Query->GetResult();
-			}
-			else
+			TSharedPtr<FMacOpenGLQuery> Query = GetQuery(PlatformContext, QueryId);
+			if(Query.IsValid())
 			{
-				*OutResult = Query->GetResultAvailable();
+				if(QueryMode == QM_Result)
+				{
+					*OutResult = Query->GetResult();
+				}
+				else
+				{
+					*OutResult = Query->GetResultAvailable();
+				}
 			}
 		}
 	}
-#else
-	FOpenGL3::GetQueryObject(QueryId, QueryMode, OutResult);
-#endif
+	else
+	{
+		FOpenGL3::GetQueryObject(QueryId, QueryMode, OutResult);
+	}
 }
 
 void FMacOpenGL::MacGetQueryObject(GLuint QueryId, EQueryMode QueryMode, GLuint *OutResult)
 {
-#if UE_EMULATE_TIMESTAMP
-	CGLContextObj Current = CGLGetCurrentContext();
-	FPlatformOpenGLContext* PlatformContext = nullptr;
-	if(Current && (CGLGetParameter(Current, kCGLCPClientStorage, (GLint *)&PlatformContext) == kCGLNoError) && PlatformContext)
+	if(GIsEmulatingTimestamp)
 	{
-		TSharedPtr<FMacOpenGLQuery> Query = GetQuery(PlatformContext, QueryId);
-		if(Query.IsValid())
+		CGLContextObj Current = CGLGetCurrentContext();
+		FPlatformOpenGLContext* PlatformContext = nullptr;
+		if(Current && (CGLGetParameter(Current, kCGLCPClientStorage, (GLint *)&PlatformContext) == kCGLNoError) && PlatformContext)
 		{
-			if(QueryMode == QM_Result)
+			TSharedPtr<FMacOpenGLQuery> Query = GetQuery(PlatformContext, QueryId);
+			if(Query.IsValid())
 			{
-				*OutResult = (GLuint)Query->GetResult();
-			}
-			else
-			{
-				*OutResult = (GLuint)Query->GetResultAvailable();
+				if(QueryMode == QM_Result)
+				{
+					*OutResult = (GLuint)Query->GetResult();
+				}
+				else
+				{
+					*OutResult = (GLuint)Query->GetResultAvailable();
+				}
 			}
 		}
 	}
-#else
-	FOpenGL3::GetQueryObject(QueryId, QueryMode, OutResult);
-#endif
+	else
+	{
+		FOpenGL3::GetQueryObject(QueryId, QueryMode, OutResult);
+	}
 }
