UE4 - merge first pass at parallel rendering to main

diff --git a/Engine/Source/Runtime/Engine/Private/SimpleElementShaders.cpp b/Engine/Source/Runtime/Engine/Private/SimpleElementShaders.cpp
--- a/Engine/Source/Runtime/Engine/Private/SimpleElementShaders.cpp
+++ b/Engine/Source/Runtime/Engine/Private/SimpleElementShaders.cpp
@@ -5,6 +5,7 @@
 ==============================================================================*/
 
 #include "EnginePrivate.h"
+#include "ShaderParameterUtils.h"
 #include "SimpleElementShaders.h"
 
 /*------------------------------------------------------------------------------
@@ -18,10 +19,10 @@ FSimpleElementVS::FSimpleElementVS(const ShaderMetaType::CompiledShaderInitializ
 	SwitchVerticalAxis.Bind(Initializer.ParameterMap,TEXT("SwitchVerticalAxis"), SPF_Optional);
 }
 
-void FSimpleElementVS::SetParameters(const FMatrix& TransformValue, bool bSwitchVerticalAxis)
+void FSimpleElementVS::SetParameters(FRHICommandList* RHICmdList, const FMatrix& TransformValue, bool bSwitchVerticalAxis)
 {
-	SetShaderValue(GetVertexShader(),Transform,TransformValue);
-	SetShaderValue(GetVertexShader(),SwitchVerticalAxis, bSwitchVerticalAxis ? -1.0f : 1.0f);
+	SetShaderValue(RHICmdList, GetVertexShader(),Transform,TransformValue);
+	SetShaderValue(RHICmdList, GetVertexShader(),SwitchVerticalAxis, bSwitchVerticalAxis ? -1.0f : 1.0f);
 }
 
 bool FSimpleElementVS::Serialize(FArchive& Ar)
@@ -53,11 +54,11 @@ FSimpleElementPS::FSimpleElementPS(const ShaderMetaType::CompiledShaderInitializ
 	ScreenToPixel.Bind(Initializer.ParameterMap,TEXT("ScreenToPixel"));
 }
 
-void FSimpleElementPS::SetEditorCompositingParameters( const FSceneView* View, FTexture2DRHIRef DepthTexture )
+void FSimpleElementPS::SetEditorCompositingParameters(FRHICommandList* RHICmdList, const FSceneView* View, FTexture2DRHIRef DepthTexture )
 {
 	if( View )
 	{
-		FGlobalShader::SetParameters( GetPixelShader(), *View );
+		FGlobalShader::SetParameters(RHICmdList, GetPixelShader(), *View );
 
 		FIntRect DestRect = View->ViewRect;
 		FIntPoint ViewportOffset = DestRect.Min;
@@ -69,27 +70,27 @@ void FSimpleElementPS::SetEditorCompositingParameters( const FSceneView* View, F
 			ViewportExtent.X * 0.5f - 0.5f + ViewportOffset.X,
 			ViewportExtent.Y * 0.5f - 0.5f + ViewportOffset.Y);
 
-		SetShaderValue(GetPixelShader(), ScreenToPixel, ScreenPosToPixelValue);
+		SetShaderValue(RHICmdList, GetPixelShader(), ScreenToPixel, ScreenPosToPixelValue);
 
-		SetShaderValue(GetPixelShader(),EditorCompositeDepthTestParameter,IsValidRef(DepthTexture) );
+		SetShaderValue(RHICmdList, GetPixelShader(),EditorCompositeDepthTestParameter,IsValidRef(DepthTexture) );
 	}
 	else
 	{
 		// Unset the view uniform buffer since we don't have a view
-		SetUniformBufferParameter(GetPixelShader(), GetUniformBufferParameter<FViewUniformShaderParameters>(), NULL);
-		SetShaderValue(GetPixelShader(),EditorCompositeDepthTestParameter,false );
+		SetUniformBufferParameter(RHICmdList, GetPixelShader(), GetUniformBufferParameter<FViewUniformShaderParameters>(), NULL);
+		SetShaderValue(RHICmdList, GetPixelShader(),EditorCompositeDepthTestParameter,false );
 	}
 
 	// Bind the zbuffer as a texture if depth textures are supported
-	SetTextureParameter(GetPixelShader(), SceneDepthTextureNonMS, IsValidRef(DepthTexture) ? (FTextureRHIRef)DepthTexture : GWhiteTexture->TextureRHI);
+	SetTextureParameter(RHICmdList, GetPixelShader(), SceneDepthTextureNonMS, IsValidRef(DepthTexture) ? (FTextureRHIRef)DepthTexture : GWhiteTexture->TextureRHI);
 }
 
-void FSimpleElementPS::SetParameters(const FTexture* TextureValue)
+void FSimpleElementPS::SetParameters(FRHICommandList* RHICmdList, const FTexture* TextureValue)
 {
-	SetTextureParameter(GetPixelShader(),InTexture,InTextureSampler,TextureValue);
+	SetTextureParameter(RHICmdList, GetPixelShader(),InTexture,InTextureSampler,TextureValue);
 	
-	SetShaderValue(GetPixelShader(),TextureComponentReplicate,TextureValue->bGreyScaleFormat ? FLinearColor(1,0,0,0) : FLinearColor(0,0,0,0));
-	SetShaderValue(GetPixelShader(),TextureComponentReplicateAlpha,TextureValue->bGreyScaleFormat ? FLinearColor(1,0,0,0) : FLinearColor(0,0,0,1));
+	SetShaderValue(RHICmdList, GetPixelShader(),TextureComponentReplicate,TextureValue->bGreyScaleFormat ? FLinearColor(1,0,0,0) : FLinearColor(0,0,0,0));
+	SetShaderValue(RHICmdList, GetPixelShader(),TextureComponentReplicateAlpha,TextureValue->bGreyScaleFormat ? FLinearColor(1,0,0,0) : FLinearColor(0,0,0,1));
 }
 
 bool FSimpleElementPS::Serialize(FArchive& Ar)
@@ -111,10 +112,10 @@ FSimpleElementGammaPS::FSimpleElementGammaPS(const ShaderMetaType::CompiledShade
 	Gamma.Bind(Initializer.ParameterMap,TEXT("Gamma"));
 }
 
-void FSimpleElementGammaPS::SetParameters(const FTexture* Texture,float GammaValue,ESimpleElementBlendMode BlendMode)
+void FSimpleElementGammaPS::SetParameters(FRHICommandList* RHICmdList, const FTexture* Texture,float GammaValue,ESimpleElementBlendMode BlendMode)
 {
-	FSimpleElementPS::SetParameters(Texture);
-	SetShaderValue(GetPixelShader(),Gamma,GammaValue);
+	FSimpleElementPS::SetParameters(RHICmdList, Texture);
+	SetShaderValue(RHICmdList, GetPixelShader(),Gamma,GammaValue);
 }
 
 bool FSimpleElementGammaPS::Serialize(FArchive& Ar)
@@ -130,10 +131,10 @@ FSimpleElementMaskedGammaPS::FSimpleElementMaskedGammaPS(const ShaderMetaType::C
 	ClipRef.Bind(Initializer.ParameterMap,TEXT("ClipRef"), SPF_Mandatory);
 }
 
-void FSimpleElementMaskedGammaPS::SetParameters(const FTexture* Texture,float Gamma,float ClipRefValue,ESimpleElementBlendMode BlendMode)
+void FSimpleElementMaskedGammaPS::SetParameters(FRHICommandList* RHICmdList, const FTexture* Texture,float Gamma,float ClipRefValue,ESimpleElementBlendMode BlendMode)
 {
-	FSimpleElementGammaPS::SetParameters(Texture,Gamma,BlendMode);
-	SetShaderValue(GetPixelShader(),ClipRef,ClipRefValue);
+	FSimpleElementGammaPS::SetParameters(RHICmdList, Texture,Gamma,BlendMode);
+	SetShaderValue(RHICmdList, GetPixelShader(),ClipRef,ClipRefValue);
 }
 
 bool FSimpleElementMaskedGammaPS::Serialize(FArchive& Ar)
@@ -176,6 +177,7 @@ FSimpleElementDistanceFieldGammaPS::FSimpleElementDistanceFieldGammaPS(const Sha
 * @param BlendMode - current batched element blend mode being rendered
 */
 void FSimpleElementDistanceFieldGammaPS::SetParameters(
+	FRHICommandList* RHICmdList, 
 	const FTexture* Texture,
 	float Gamma,
 	float ClipRef,
@@ -188,25 +190,25 @@ void FSimpleElementDistanceFieldGammaPS::SetParameters(
 	ESimpleElementBlendMode BlendMode
 	)
 {
-	FSimpleElementMaskedGammaPS::SetParameters(Texture,Gamma,ClipRef,BlendMode);
-	SetShaderValue(GetPixelShader(),SmoothWidth,SmoothWidthValue);		
-	SetPixelShaderBool(GetPixelShader(),EnableShadow,EnableShadowValue);
+	FSimpleElementMaskedGammaPS::SetParameters(RHICmdList, Texture,Gamma,ClipRef,BlendMode);
+	SetShaderValue(RHICmdList, GetPixelShader(),SmoothWidth,SmoothWidthValue);		
+	SetPixelShaderBool(RHICmdList, GetPixelShader(),EnableShadow,EnableShadowValue);
 	if (EnableShadowValue)
 	{
-		SetShaderValue(GetPixelShader(),ShadowDirection,ShadowDirectionValue);
-		SetShaderValue(GetPixelShader(),ShadowColor,ShadowColorValue);
-		SetShaderValue(GetPixelShader(),ShadowSmoothWidth,ShadowSmoothWidthValue);
+		SetShaderValue(RHICmdList, GetPixelShader(),ShadowDirection,ShadowDirectionValue);
+		SetShaderValue(RHICmdList, GetPixelShader(),ShadowColor,ShadowColorValue);
+		SetShaderValue(RHICmdList, GetPixelShader(),ShadowSmoothWidth,ShadowSmoothWidthValue);
 	}
-	SetPixelShaderBool(GetPixelShader(),EnableGlow,GlowInfo.bEnableGlow);
+	SetPixelShaderBool(RHICmdList, GetPixelShader(),EnableGlow,GlowInfo.bEnableGlow);
 	if (GlowInfo.bEnableGlow)
 	{
-		SetShaderValue(GetPixelShader(),GlowColor,GlowInfo.GlowColor);
-		SetShaderValue(GetPixelShader(),GlowOuterRadius,GlowInfo.GlowOuterRadius);
-		SetShaderValue(GetPixelShader(),GlowInnerRadius,GlowInfo.GlowInnerRadius);
+		SetShaderValue(RHICmdList, GetPixelShader(),GlowColor,GlowInfo.GlowColor);
+		SetShaderValue(RHICmdList, GetPixelShader(),GlowOuterRadius,GlowInfo.GlowOuterRadius);
+		SetShaderValue(RHICmdList, GetPixelShader(),GlowInnerRadius,GlowInfo.GlowInnerRadius);
 	}
 
 	// This shader does not use editor compositing
-	SetEditorCompositingParameters( NULL, FTexture2DRHIRef() );
+	SetEditorCompositingParameters(RHICmdList, NULL, FTexture2DRHIRef() );
 }
 
 /**
@@ -237,9 +239,9 @@ FSimpleElementHitProxyPS::FSimpleElementHitProxyPS(const ShaderMetaType::Compile
 	InTextureSampler.Bind(Initializer.ParameterMap,TEXT("InTextureSampler"));
 }
 
-void FSimpleElementHitProxyPS::SetParameters(const FTexture* TextureValue)
+void FSimpleElementHitProxyPS::SetParameters(FRHICommandList* RHICmdList, const FTexture* TextureValue)
 {
-	SetTextureParameter(GetPixelShader(),InTexture,InTextureSampler,TextureValue);
+	SetTextureParameter(RHICmdList, GetPixelShader(),InTexture,InTextureSampler,TextureValue);
 }
 
 bool FSimpleElementHitProxyPS::Serialize(FArchive& Ar)
@@ -267,11 +269,11 @@ FSimpleElementColorChannelMaskPS::FSimpleElementColorChannelMaskPS(const ShaderM
 * @param ColorWeights - reference value to compare with alpha for killing pixels
 * @param Gamma - if gamma != 1.0 then a pow(color,Gamma) is applied
 */
-void FSimpleElementColorChannelMaskPS::SetParameters(const FTexture* TextureValue, const FMatrix& ColorWeightsValue, float GammaValue)
+void FSimpleElementColorChannelMaskPS::SetParameters(FRHICommandList* RHICmdList, const FTexture* TextureValue, const FMatrix& ColorWeightsValue, float GammaValue)
 {
-	SetTextureParameter(GetPixelShader(),InTexture,InTextureSampler,TextureValue);
-	SetShaderValue(GetPixelShader(),ColorWeights,ColorWeightsValue);
-	SetShaderValue(GetPixelShader(),Gamma,GammaValue);
+	SetTextureParameter(RHICmdList, GetPixelShader(),InTexture,InTextureSampler,TextureValue);
+	SetShaderValue(RHICmdList, GetPixelShader(),ColorWeights,ColorWeightsValue);
+	SetShaderValue(RHICmdList, GetPixelShader(),Gamma,GammaValue);
 }
 
 bool FSimpleElementColorChannelMaskPS::Serialize(FArchive& Ar)
diff --git a/Engine/Source/Runtime/Engine/Private/Slate/SceneViewport.cpp b/Engine/Source/Runtime/Engine/Private/Slate/SceneViewport.cpp
--- a/Engine/Source/Runtime/Engine/Private/Slate/SceneViewport.cpp
+++ b/Engine/Source/Runtime/Engine/Private/Slate/SceneViewport.cpp
@@ -1014,8 +1014,6 @@ void FSceneViewport::BeginRenderFrame()
 {
 	check( IsInRenderingThread() );
 	
-	RHIBeginScene();
-
 	if( bUseSeparateRenderTarget )
 	{
 		RHISetRenderTarget( RenderTargetTextureRHI,  FTexture2DRHIRef() );
@@ -1047,8 +1045,6 @@ void FSceneViewport::EndRenderFrame( bool bPresent, bool bLockToVsync )
 		// Note: this releases our reference but does not release the resource as it is owned by slate (this is intended)
 		RenderTargetTextureRHI.SafeRelease();
 	}
-
-	RHIEndScene();
 }
 
 void FSceneViewport::Tick( float DeltaTime )
diff --git a/Engine/Source/Runtime/Engine/Private/SplineMeshComponent.cpp b/Engine/Source/Runtime/Engine/Private/SplineMeshComponent.cpp
--- a/Engine/Source/Runtime/Engine/Private/SplineMeshComponent.cpp
+++ b/Engine/Source/Runtime/Engine/Private/SplineMeshComponent.cpp
@@ -3,6 +3,7 @@
 #include "EnginePrivate.h"
 #include "Components/SplineMeshComponent.h"
 #include "SplineMeshSceneProxy.h"
+#include "ShaderParameterUtils.h"
 
 
 //////////////////////////////////////////////////////////////////////////
@@ -33,7 +34,7 @@ void FSplineMeshVertexFactoryShaderParameters::Bind(const FShaderParameterMap& P
 	SplineMeshYParam.Bind(ParameterMap, TEXT("SplineMeshY"), SPF_Mandatory);
 }
 
-void FSplineMeshVertexFactoryShaderParameters::SetMesh(FShader* Shader, const FVertexFactory* VertexFactory, const FSceneView& View, const FMeshBatchElement& BatchElement, uint32 DataFlags) const
+void FSplineMeshVertexFactoryShaderParameters::SetMesh(FRHICommandList* RHICmdList, FShader* Shader, const FVertexFactory* VertexFactory, const FSceneView& View, const FMeshBatchElement& BatchElement, uint32 DataFlags) const
 {
 	if (Shader->GetVertexShader())
 	{
@@ -41,33 +42,33 @@ void FSplineMeshVertexFactoryShaderParameters::SetMesh(FShader* Shader, const FV
 		FSplineMeshSceneProxy* SplineProxy = SplineVertexFactory->SplineSceneProxy;
 		FSplineMeshParams& SplineParams = SplineProxy->SplineParams;
 
-		SetShaderValue(Shader->GetVertexShader(), SplineStartPosParam, SplineParams.StartPos);
-		SetShaderValue(Shader->GetVertexShader(), SplineStartTangentParam, SplineParams.StartTangent);
-		SetShaderValue(Shader->GetVertexShader(), SplineStartRollParam, SplineParams.StartRoll);
-		SetShaderValue(Shader->GetVertexShader(), SplineStartScaleParam, SplineParams.StartScale);
-		SetShaderValue(Shader->GetVertexShader(), SplineStartOffsetParam, SplineParams.StartOffset);
+		SetShaderValue(RHICmdList, Shader->GetVertexShader(), SplineStartPosParam, SplineParams.StartPos);
+		SetShaderValue(RHICmdList, Shader->GetVertexShader(), SplineStartTangentParam, SplineParams.StartTangent);
+		SetShaderValue(RHICmdList, Shader->GetVertexShader(), SplineStartRollParam, SplineParams.StartRoll);
+		SetShaderValue(RHICmdList, Shader->GetVertexShader(), SplineStartScaleParam, SplineParams.StartScale);
+		SetShaderValue(RHICmdList, Shader->GetVertexShader(), SplineStartOffsetParam, SplineParams.StartOffset);
 
-		SetShaderValue(Shader->GetVertexShader(), SplineEndPosParam, SplineParams.EndPos);
-		SetShaderValue(Shader->GetVertexShader(), SplineEndTangentParam, SplineParams.EndTangent);
-		SetShaderValue(Shader->GetVertexShader(), SplineEndRollParam, SplineParams.EndRoll);
-		SetShaderValue(Shader->GetVertexShader(), SplineEndScaleParam, SplineParams.EndScale);
-		SetShaderValue(Shader->GetVertexShader(), SplineEndOffsetParam, SplineParams.EndOffset);
+		SetShaderValue(RHICmdList, Shader->GetVertexShader(), SplineEndPosParam, SplineParams.EndPos);
+		SetShaderValue(RHICmdList, Shader->GetVertexShader(), SplineEndTangentParam, SplineParams.EndTangent);
+		SetShaderValue(RHICmdList, Shader->GetVertexShader(), SplineEndRollParam, SplineParams.EndRoll);
+		SetShaderValue(RHICmdList, Shader->GetVertexShader(), SplineEndScaleParam, SplineParams.EndScale);
+		SetShaderValue(RHICmdList, Shader->GetVertexShader(), SplineEndOffsetParam, SplineParams.EndOffset);
 
-		SetShaderValue(Shader->GetVertexShader(), SplineUpDirParam, SplineProxy->SplineUpDir);
-		SetShaderValue(Shader->GetVertexShader(), SmoothInterpRollScaleParam, SplineProxy->bSmoothInterpRollScale);
+		SetShaderValue(RHICmdList, Shader->GetVertexShader(), SplineUpDirParam, SplineProxy->SplineUpDir);
+		SetShaderValue(RHICmdList, Shader->GetVertexShader(), SmoothInterpRollScaleParam, SplineProxy->bSmoothInterpRollScale);
 
-		SetShaderValue(Shader->GetVertexShader(), SplineMeshMinZParam, SplineProxy->SplineMeshMinZ);
-		SetShaderValue(Shader->GetVertexShader(), SplineMeshScaleZParam, SplineProxy->SplineMeshScaleZ);
+		SetShaderValue(RHICmdList, Shader->GetVertexShader(), SplineMeshMinZParam, SplineProxy->SplineMeshMinZ);
+		SetShaderValue(RHICmdList, Shader->GetVertexShader(), SplineMeshScaleZParam, SplineProxy->SplineMeshScaleZ);
 
 		FVector DirMask(0, 0, 0);
 		DirMask[SplineProxy->ForwardAxis] = 1;
-		SetShaderValue(Shader->GetVertexShader(), SplineMeshDirParam, DirMask);
+		SetShaderValue(RHICmdList, Shader->GetVertexShader(), SplineMeshDirParam, DirMask);
 		DirMask = FVector::ZeroVector;
 		DirMask[(SplineProxy->ForwardAxis + 1) % 3] = 1;
-		SetShaderValue(Shader->GetVertexShader(), SplineMeshXParam, DirMask);
+		SetShaderValue(RHICmdList, Shader->GetVertexShader(), SplineMeshXParam, DirMask);
 		DirMask = FVector::ZeroVector;
 		DirMask[(SplineProxy->ForwardAxis + 2) % 3] = 1;
-		SetShaderValue(Shader->GetVertexShader(), SplineMeshYParam, DirMask);
+		SetShaderValue(RHICmdList, Shader->GetVertexShader(), SplineMeshYParam, DirMask);
 	}
 }
 
diff --git a/Engine/Source/Runtime/Engine/Private/SplineMeshSceneProxy.h b/Engine/Source/Runtime/Engine/Private/SplineMeshSceneProxy.h
--- a/Engine/Source/Runtime/Engine/Private/SplineMeshSceneProxy.h
+++ b/Engine/Source/Runtime/Engine/Private/SplineMeshSceneProxy.h
@@ -69,7 +69,7 @@ class FSplineMeshVertexFactoryShaderParameters : public FVertexFactoryShaderPara
 {
 	void Bind(const FShaderParameterMap& ParameterMap) OVERRIDE;
 
-	void SetMesh(FShader* Shader, const FVertexFactory* VertexFactory, const FSceneView& View, const FMeshBatchElement& BatchElement, uint32 DataFlags) const OVERRIDE;
+	void SetMesh(FRHICommandList* RHICmdList, FShader* Shader, const FVertexFactory* VertexFactory, const FSceneView& View, const FMeshBatchElement& BatchElement, uint32 DataFlags) const OVERRIDE;
 
 	void Serialize(FArchive& Ar) OVERRIDE
 	{
diff --git a/Engine/Source/Runtime/Engine/Private/Texture.cpp b/Engine/Source/Runtime/Engine/Private/Texture.cpp
--- a/Engine/Source/Runtime/Engine/Private/Texture.cpp
+++ b/Engine/Source/Runtime/Engine/Private/Texture.cpp
@@ -53,6 +53,18 @@ UTexture::UTexture(const class FPostConstructInitializeProperties& PCIP)
 	CompositeTextureMode = CTM_NormalRoughnessToAlpha;
 	CompositePower = 1.0f;
 #endif // #if WITH_EDITORONLY_DATA
+
+	if (FApp::CanEverRender() && !IsTemplate())
+	{
+		ENQUEUE_UNIQUE_RENDER_COMMAND_ONEPARAMETER(
+			FCreateTextureRefCommand,
+			UTexture*,Texture,this,
+		{
+			FTextureReferenceRHIRef TextureRef = RHICreateTextureReference();
+			TextureRef->AddRef();
+			Texture->TextureRefRHI = TextureRef;
+		});
+	}
 }
 
 void UTexture::ReleaseResource()
@@ -305,11 +317,24 @@ void UTexture::PostLoad()
 void UTexture::BeginDestroy()
 {
 	Super::BeginDestroy();
-	if( !UpdateStreamingStatus() && Resource )
+	if( !UpdateStreamingStatus() && (Resource || TextureRefRHI) )
 	{
 		// Send the rendering thread a release message for the texture's resource.
-		BeginReleaseResource(Resource);
-		Resource->ReleaseFence.BeginFence();
+		if (Resource)
+		{
+			BeginReleaseResource(Resource);
+		}
+		if (TextureRefRHI)
+		{
+			ENQUEUE_UNIQUE_RENDER_COMMAND_ONEPARAMETER(
+				FDestroyTextureRefCommand,
+				UTexture*,Texture,this,
+			{
+				Texture->TextureRefRHI->Release();
+				Texture->TextureRefRHI = NULL;
+			});
+		}
+		ReleaseFence.BeginFence();
 		// Keep track that we already kicked off the async release.
 		bAsyncResourceReleaseHasBeenStarted = true;
 	}
@@ -322,16 +347,29 @@ bool UTexture::IsReadyForFinishDestroy()
 	if( Super::IsReadyForFinishDestroy() && !UpdateStreamingStatus() )
 	{
 		// Kick off async resource release if we haven't already.
-		if( !bAsyncResourceReleaseHasBeenStarted && Resource )
+		if( !bAsyncResourceReleaseHasBeenStarted && (Resource || TextureRefRHI) )
 		{
 			// Send the rendering thread a release message for the texture's resource.
-			BeginReleaseResource(Resource);
-			Resource->ReleaseFence.BeginFence();
+			if (Resource)
+			{
+				BeginReleaseResource(Resource);
+			}
+			if (TextureRefRHI)
+			{
+				ENQUEUE_UNIQUE_RENDER_COMMAND_ONEPARAMETER(
+					FDestroyTextureRefCommand,
+					UTexture*,Texture,this,
+				{
+					Texture->TextureRefRHI->Release();
+					Texture->TextureRefRHI = NULL;
+				});
+			}
+			ReleaseFence.BeginFence();
 			// Keep track that we already kicked off the async release.
 			bAsyncResourceReleaseHasBeenStarted = true;
 		}
 		// Only allow FinishDestroy to be called once the texture resource has finished its rendering thread cleanup.
-		else if( !Resource || Resource->ReleaseFence.IsFenceComplete() )
+		else if( (!Resource && !TextureRefRHI) || ReleaseFence.IsFenceComplete() )
 		{
 			bReadyForFinishDestroy = true;
 		}
@@ -345,7 +383,7 @@ void UTexture::FinishDestroy()
 
 	if(Resource)
 	{
-		check(Resource->ReleaseFence.IsFenceComplete());
+		check(ReleaseFence.IsFenceComplete());
 
 		// Free the resource.
 		delete Resource;
@@ -355,6 +393,7 @@ void UTexture::FinishDestroy()
 	CleanupCachedRunningPlatformData();
 	CleanupCachedCookedPlatformData();
 
+	check(TextureRefRHI == NULL);
 }
 
 void UTexture::PreSave()
diff --git a/Engine/Source/Runtime/Engine/Private/Texture2D.cpp b/Engine/Source/Runtime/Engine/Private/Texture2D.cpp
--- a/Engine/Source/Runtime/Engine/Private/Texture2D.cpp
+++ b/Engine/Source/Runtime/Engine/Private/Texture2D.cpp
@@ -301,6 +301,32 @@ void UTexture2D::Serialize(FArchive& Ar)
 #endif // #if WITH_EDITORONLY_DATA
 }
 
+float UTexture2D::GetLastRenderTimeForStreaming()
+{
+	float LastRenderTime = -FLT_MAX;
+	if (Resource)
+	{
+		LastRenderTime = Resource->LastRenderTime;
+		if (TextureRefRHI)
+		{
+			LastRenderTime = FMath::Max(LastRenderTime,TextureRefRHI->GetLastCachedTime());
+		}
+	}
+	return LastRenderTime;
+}
+
+void UTexture2D::InvalidateLastRenderTimeForStreaming()
+{
+	if (Resource)
+	{
+		Resource->LastRenderTime = -FLT_MAX;
+	}
+	if (TextureRefRHI)
+	{
+		TextureRefRHI->SetLastCachedTime(-FLT_MAX);
+	}
+}
+
 #if WITH_EDITOR
 void UTexture2D::PostEditUndo()
 {
@@ -1299,6 +1325,7 @@ void FTexture2DResource::InitRHI()
 
 				TextureRHI = Texture2DRHI;
 				RHIBindDebugLabelName(TextureRHI, *Owner->GetName());
+				RHIUpdateTextureReference(Owner->TextureRefRHI,TextureRHI);
 
 				return;
 			}
@@ -1308,6 +1335,7 @@ void FTexture2DResource::InitRHI()
 			Texture2DRHI	= RHICreateTexture2D( SizeX, SizeY, EffectiveFormat, Owner->RequestedMips, 1, TexCreateFlags, CreateInfo);
 			TextureRHI		= Texture2DRHI;
 			RHIBindDebugLabelName(TextureRHI, *Owner->GetName());
+			RHIUpdateTextureReference(Owner->TextureRefRHI,TextureRHI);
 
 			check(Owner->RequestedMips == Texture2DRHI->GetNumMips());
 			check(Owner->PlatformData->Mips.Num() == CurrentFirstMip + Owner->RequestedMips);
@@ -1357,6 +1385,7 @@ void FTexture2DResource::InitRHI()
 			Texture2DRHI	= RHICreateTexture2D( SizeX, SizeY, EffectiveFormat, Owner->RequestedMips, 1, TexCreateFlags, CreateInfo );
 			TextureRHI		= Texture2DRHI;
 			RHIBindDebugLabelName(TextureRHI, *Owner->GetName());
+			RHIUpdateTextureReference(Owner->TextureRefRHI,TextureRHI);
 			for( int32 MipIndex=CurrentFirstMip; MipIndex<OwnerMips.Num(); MipIndex++ )
 			{
 				if( MipData[MipIndex] != NULL )
@@ -1397,6 +1426,7 @@ void FTexture2DResource::ReleaseRHI()
 	check( Owner->PendingMipChangeRequestStatus.GetValue() == TexState_ReadyFor_Requests );	
 	FTextureResource::ReleaseRHI();
 	Texture2DRHI.SafeRelease();
+	RHIUpdateTextureReference(Owner->TextureRefRHI,FTextureRHIParamRef());
 }
 
 void FTexture2DResource::CreateSamplerStates(float MipMapBias)
@@ -1666,6 +1696,7 @@ void FTexture2DResource::UpdateMipCount()
 			}
 
 			TextureRHI = Texture2DRHI;
+			RHIUpdateTextureReference(Owner->TextureRefRHI,TextureRHI);
 		}
 	}
 
@@ -1986,6 +2017,7 @@ void FCreateTextureTask::DoWork()
 	// NOTE: Leaving scope here to drop the reference held by AsyncTexture.
 	// The only remaining reference is now Texture2DResource->IntermediateTextureRHI.
 	// This is important otherwise two threads could be ref counting simultaneously!
+	FPlatformMisc::MemoryBarrier();
 
 	// Decrement the counter so the texture streamer knows that texture creation has completed.
 	Args.ThreadSafeCounter->Decrement();
@@ -2170,6 +2202,7 @@ void FTexture2DResource::FinalizeMipCount()
 			TextureRHI		= IntermediateTextureRHI;
 			Texture2DRHI	= IntermediateTextureRHI;
 			CurrentFirstMip = PendingFirstMip;
+			RHIUpdateTextureReference(Owner->TextureRefRHI,TextureRHI);
 
 			// Update mip-level fading.
 			EMipFadeSettings MipFadeSetting = (Owner->LODGroup == TEXTUREGROUP_Lightmap || Owner->LODGroup == TEXTUREGROUP_Shadowmap) ? MipFade_Slow : MipFade_Normal;
@@ -2269,6 +2302,7 @@ bool FTexture2DResource::TryReallocate( int32 OldMipCount, int32 NewMipCount )
 
 	Texture2DRHI = NewTextureRHI;
 	TextureRHI = NewTextureRHI;
+	RHIUpdateTextureReference(Owner->TextureRefRHI,TextureRHI);
 
 	PendingFirstMip = CurrentFirstMip = MipIndex;
 
diff --git a/Engine/Source/Runtime/Engine/Private/Texture2DDynamic.cpp b/Engine/Source/Runtime/Engine/Private/Texture2DDynamic.cpp
--- a/Engine/Source/Runtime/Engine/Private/Texture2DDynamic.cpp
+++ b/Engine/Source/Runtime/Engine/Private/Texture2DDynamic.cpp
@@ -59,11 +59,13 @@ void FTexture2DDynamicResource::InitRHI()
 	FRHIResourceCreateInfo CreateInfo;
 	Texture2DRHI = RHICreateTexture2D( Owner->SizeX, Owner->SizeY, Owner->Format, Owner->NumMips, 1, Flags, CreateInfo );
 	TextureRHI = Texture2DRHI;
+	RHIUpdateTextureReference(Owner->TextureRefRHI,TextureRHI);
 }
 
 /** Called when the resource is released. This is only called by the rendering thread. */
 void FTexture2DDynamicResource::ReleaseRHI()
 {
+	RHIUpdateTextureReference(Owner->TextureRefRHI,FTextureRHIParamRef());
 	FTextureResource::ReleaseRHI();
 	Texture2DRHI.SafeRelease();
 }
diff --git a/Engine/Source/Runtime/Engine/Private/TextureCube.cpp b/Engine/Source/Runtime/Engine/Private/TextureCube.cpp
--- a/Engine/Source/Runtime/Engine/Private/TextureCube.cpp
+++ b/Engine/Source/Runtime/Engine/Private/TextureCube.cpp
@@ -228,6 +228,7 @@ public:
 		TextureCubeRHI = RHICreateTextureCube( Owner->GetSizeX(), Owner->GetPixelFormat(), Owner->GetNumMips(), TexCreateFlags, CreateInfo );
 		TextureRHI = TextureCubeRHI;
 		RHIBindDebugLabelName(TextureRHI, *Owner->GetName());
+		RHIUpdateTextureReference(Owner->TextureRefRHI,TextureRHI);
 
 		// Read the mip-levels into the RHI texture.
 		int32 NumMips = Owner->GetNumMips();
@@ -264,6 +265,7 @@ public:
 	{
 		DEC_DWORD_STAT_BY( STAT_TextureMemory, TextureSize );
 		DEC_DWORD_STAT_FNAME_BY( LODGroupStatName, TextureSize );
+		RHIUpdateTextureReference(Owner->TextureRefRHI,FTextureRHIParamRef());
 		TextureCubeRHI.SafeRelease();
 		FTextureResource::ReleaseRHI();
 	}
diff --git a/Engine/Source/Runtime/Engine/Private/UnrealClient.cpp b/Engine/Source/Runtime/Engine/Private/UnrealClient.cpp
--- a/Engine/Source/Runtime/Engine/Private/UnrealClient.cpp
+++ b/Engine/Source/Runtime/Engine/Private/UnrealClient.cpp
@@ -696,14 +696,12 @@ public:
 	virtual void BeginRenderFrame() 
 	{
 		check( IsInRenderingThread() );
-		RHIBeginScene();
 		RHISetRenderTarget( RenderTargetTextureRHI,  FTexture2DRHIRef() );
 	};
 
 	void EndRenderFrame( bool bPresent, bool bLockToVsync )
 	{
 		check( IsInRenderingThread() );
-		RHIEndScene();
 	}
 
 	virtual void*	GetWindow() { return 0; }
diff --git a/Engine/Source/Runtime/Engine/Private/UnrealEngine.cpp b/Engine/Source/Runtime/Engine/Private/UnrealEngine.cpp
--- a/Engine/Source/Runtime/Engine/Private/UnrealEngine.cpp
+++ b/Engine/Source/Runtime/Engine/Private/UnrealEngine.cpp
@@ -346,6 +346,7 @@ void RefreshSamplerStatesCallback()
 			UTexture2D* Texture = *It;
 			Texture->RefreshSamplerStates();
 		}
+		UMaterialInterface::RecacheAllMaterialUniformExpressions();
 	}
 }
 
@@ -667,16 +668,6 @@ void UEngine::Init(IEngineLoop* InEngineLoop)
 	}
 
 #if !(UE_BUILD_SHIPPING || UE_BUILD_TEST)
-	{
-		// if the console variable FreezeAtPosition is set, we right away freeze at the given position
-		static IConsoleVariable* CVar = IConsoleManager::Get().FindConsoleVariable(TEXT("FreezeAtPosition")); 
-
-		if(!CVar->GetString().IsEmpty())
-		{
-			new(GEngine->DeferredCommands) FString(TEXT("FreezeAt"));
-		}
-	}
-
 	// Optionally Exec an exec file
 	FString Temp;
 	if( FParse::Value(FCommandLine::Get(), TEXT("EXEC="), Temp) )
diff --git a/Engine/Source/Runtime/Engine/Private/VectorField.cpp b/Engine/Source/Runtime/Engine/Private/VectorField.cpp
--- a/Engine/Source/Runtime/Engine/Private/VectorField.cpp
+++ b/Engine/Source/Runtime/Engine/Private/VectorField.cpp
@@ -10,6 +10,8 @@
 #include "RenderResource.h"
 #include "UniformBuffer.h"
 #include "ShaderParameters.h"
+#include "ShaderParameterUtils.h"
+#include "RHIStaticStates.h"
 #include "GlobalShader.h"
 #include "FXSystem.h"
 
@@ -621,15 +623,16 @@ public:
 	 * @param NoiseVolumeTextureRHI - The volume texture to use to add noise to the vector field.
 	 */
 	void SetParameters(
+		FRHICommandList* RHICmdList, 
 		const FCompositeAnimatedVectorFieldUniformBufferRef& UniformBuffer,
 		FTextureRHIParamRef AtlasTextureRHI,
 		FTextureRHIParamRef NoiseVolumeTextureRHI )
 	{
 		FComputeShaderRHIParamRef ComputeShaderRHI = GetComputeShader();
 		FSamplerStateRHIParamRef SamplerStateLinear = TStaticSamplerState<SF_Bilinear,AM_Clamp,AM_Clamp,AM_Clamp>::GetRHI();
-		SetUniformBufferParameter( ComputeShaderRHI, GetUniformBufferParameter<FCompositeAnimatedVectorFieldUniformParameters>(), UniformBuffer );
-		SetTextureParameter( ComputeShaderRHI, AtlasTexture, AtlasTextureSampler, SamplerStateLinear, AtlasTextureRHI );
-		SetTextureParameter( ComputeShaderRHI, NoiseVolumeTexture, NoiseVolumeTextureSampler, SamplerStateLinear, NoiseVolumeTextureRHI );
+		SetUniformBufferParameter(RHICmdList, ComputeShaderRHI, GetUniformBufferParameter<FCompositeAnimatedVectorFieldUniformParameters>(), UniformBuffer );
+		SetTextureParameter(RHICmdList, ComputeShaderRHI, AtlasTexture, AtlasTextureSampler, SamplerStateLinear, AtlasTextureRHI );
+		SetTextureParameter(RHICmdList, ComputeShaderRHI, NoiseVolumeTexture, NoiseVolumeTextureSampler, SamplerStateLinear, NoiseVolumeTextureRHI );
 	}
 
 	/**
@@ -799,7 +802,7 @@ public:
 			Parameters.Op = (uint32)AnimatedVectorField->ConstructionOp;
 
 			FCompositeAnimatedVectorFieldUniformBufferRef UniformBuffer = 
-				FCompositeAnimatedVectorFieldUniformBufferRef::CreateUniformBufferImmediate(Parameters, UniformBuffer_SingleUse);
+				FCompositeAnimatedVectorFieldUniformBufferRef::CreateUniformBufferImmediate(Parameters, UniformBuffer_SingleDraw);
 
 			TShaderMapRef<FCompositeAnimatedVectorFieldCS> CompositeCS(GetGlobalShaderMap());
 			FTextureRHIParamRef NoiseVolumeTextureRHI = GBlackVolumeTexture->TextureRHI;
@@ -810,11 +813,14 @@ public:
 
 			RHISetComputeShader(CompositeCS->GetComputeShader());
 			CompositeCS->SetOutput(VolumeTextureUAV);
+			/// ?
 			CompositeCS->SetParameters(
+				nullptr,
 				UniformBuffer,
 				AnimatedVectorField->Texture->Resource->TextureRHI,
 				NoiseVolumeTextureRHI );
 			DispatchComputeShader(
+				nullptr,
 				*CompositeCS,
 				SizeX / THREADS_PER_AXIS,
 				SizeY / THREADS_PER_AXIS,
diff --git a/Engine/Source/Runtime/Engine/Private/VectorFieldVisualization.cpp b/Engine/Source/Runtime/Engine/Private/VectorFieldVisualization.cpp
--- a/Engine/Source/Runtime/Engine/Private/VectorFieldVisualization.cpp
+++ b/Engine/Source/Runtime/Engine/Private/VectorFieldVisualization.cpp
@@ -10,6 +10,8 @@
 #include "RenderResource.h"
 #include "UniformBuffer.h"
 #include "ShaderParameters.h"
+#include "ShaderParameterUtils.h"
+#include "RHIStaticStates.h"
 #include "GlobalShader.h"
 #include "SceneManagement.h"
 #include "FXSystem.h"
@@ -38,7 +40,7 @@ public:
 		Ar << VectorFieldTextureSampler;
 	}
 
-	virtual void SetMesh(FShader* Shader,const FVertexFactory* VertexFactory,const FSceneView& View,const FMeshBatchElement& BatchElement,uint32 DataFlags) const OVERRIDE;
+	virtual void SetMesh(FRHICommandList* RHICmdList, FShader* Shader,const FVertexFactory* VertexFactory,const FSceneView& View,const FMeshBatchElement& BatchElement,uint32 DataFlags) const OVERRIDE;
 
 	virtual uint32 GetSize() const { return sizeof(*this); }
 
@@ -155,20 +157,20 @@ void FVectorFieldVisualizationVertexFactory::SetParameters(
 	const FVectorFieldVisualizationParameters& InUniformParameters,
 	FTexture3DRHIParamRef InVectorFieldTextureRHI )
 {
-	UniformBuffer = FVectorFieldVisualizationBufferRef::CreateUniformBufferImmediate(InUniformParameters, UniformBuffer_SingleUse);
+	UniformBuffer = FVectorFieldVisualizationBufferRef::CreateUniformBufferImmediate(InUniformParameters, UniformBuffer_SingleFrame);
 	VectorFieldTextureRHI = InVectorFieldTextureRHI;
 }
 
 /**
  * Set vertex factory shader parameters.
  */
-void FVectorFieldVisualizationVertexFactoryShaderParameters::SetMesh(FShader* Shader,const FVertexFactory* InVertexFactory,const FSceneView& View,const FMeshBatchElement& BatchElement,uint32 DataFlags) const
+void FVectorFieldVisualizationVertexFactoryShaderParameters::SetMesh(FRHICommandList* RHICmdList, FShader* Shader,const FVertexFactory* InVertexFactory,const FSceneView& View,const FMeshBatchElement& BatchElement,uint32 DataFlags) const
 {
 	FVectorFieldVisualizationVertexFactory* VertexFactory = (FVectorFieldVisualizationVertexFactory*)InVertexFactory;
 	FVertexShaderRHIParamRef VertexShaderRHI = Shader->GetVertexShader();
 	FSamplerStateRHIParamRef SamplerStatePoint = TStaticSamplerState<SF_Point>::GetRHI();
-	SetUniformBufferParameter(VertexShaderRHI, Shader->GetUniformBufferParameter<FVectorFieldVisualizationParameters>(), VertexFactory->UniformBuffer);
-	SetTextureParameter(VertexShaderRHI, VectorFieldTexture, VectorFieldTextureSampler, SamplerStatePoint, VertexFactory->VectorFieldTextureRHI);
+	SetUniformBufferParameter(RHICmdList, VertexShaderRHI, Shader->GetUniformBufferParameter<FVectorFieldVisualizationParameters>(), VertexFactory->UniformBuffer);
+	SetTextureParameter(RHICmdList, VertexShaderRHI, VectorFieldTexture, VectorFieldTextureSampler, SamplerStatePoint, VertexFactory->VectorFieldTextureRHI);
 }
 
 IMPLEMENT_VERTEX_FACTORY_TYPE(FVectorFieldVisualizationVertexFactory,"VectorFieldVisualizationVertexFactory",true,false,true,false,false);
diff --git a/Engine/Source/Runtime/Engine/Public/BatchedElements.h b/Engine/Source/Runtime/Engine/Public/BatchedElements.h
--- a/Engine/Source/Runtime/Engine/Public/BatchedElements.h
+++ b/Engine/Source/Runtime/Engine/Public/BatchedElements.h
@@ -88,7 +88,7 @@ class FBatchedElementParameters
 public:
 
 	/** Binds vertex and pixel shaders for this element */
-	virtual void BindShaders_RenderThread( const FMatrix& InTransform, const float InGamma, const FMatrix& ColorWeights, const FTexture* Texture ) = 0;
+	virtual void BindShaders_RenderThread(FRHICommandList* RHICmdList, const FMatrix& InTransform, const float InGamma, const FMatrix& ColorWeights, const FTexture* Texture ) = 0;
 
 };
 
diff --git a/Engine/Source/Runtime/Engine/Public/ClearReplacementShaders.h b/Engine/Source/Runtime/Engine/Public/ClearReplacementShaders.h
--- a/Engine/Source/Runtime/Engine/Public/ClearReplacementShaders.h
+++ b/Engine/Source/Runtime/Engine/Public/ClearReplacementShaders.h
@@ -83,12 +83,7 @@ public:
 		return bShaderHasOutdatedParameters;
 	}
 
-	void SetParameters(FUnorderedAccessViewRHIParamRef TextureRW, FLinearColor Value)
-	{
-		FComputeShaderRHIParamRef ComputeShaderRHI = GetComputeShader();
-		SetShaderValue(ComputeShaderRHI, ClearColor, Value);
-		RHISetUAVParameter(ComputeShaderRHI, ClearTextureRW.GetBaseIndex(), TextureRW);
-	}
+	void SetParameters(FRHICommandList* RHICmdList, FUnorderedAccessViewRHIParamRef TextureRW, FLinearColor Value);
 
 	static bool ShouldCache(EShaderPlatform Platform)
 	{
@@ -112,12 +107,7 @@ public:
 		ClearBufferRW.Bind(Initializer.ParameterMap, TEXT("ClearBufferRW"), SPF_Mandatory);
 	}
 
-	void SetParameters(FUnorderedAccessViewRHIParamRef BufferRW, uint32 Dword)
-	{
-		FComputeShaderRHIParamRef ComputeShaderRHI = GetComputeShader();
-		SetShaderValue(ComputeShaderRHI, ClearDword, Dword);
-		RHISetUAVParameter(ComputeShaderRHI, ClearBufferRW.GetBaseIndex(), BufferRW);
-	}
+	void SetParameters(FRHICommandList* RHICmdList, FUnorderedAccessViewRHIParamRef BufferRW, uint32 Dword);
 
 	// FShader interface.
 	virtual bool Serialize(FArchive& Ar)
diff --git a/Engine/Source/Runtime/Engine/Public/FogRendering.h b/Engine/Source/Runtime/Engine/Public/FogRendering.h
--- a/Engine/Source/Runtime/Engine/Public/FogRendering.h
+++ b/Engine/Source/Runtime/Engine/Public/FogRendering.h
@@ -18,14 +18,14 @@ public:
 	 * Sets exponential height fog shader parameters.
 	 */
 	template<typename ShaderRHIParamRef>
-	void Set(ShaderRHIParamRef Shader, const FSceneView* View) const
+	void Set(FRHICommandList* RHICmdList, ShaderRHIParamRef Shader, const FSceneView* View) const
 	{
 		const FViewInfo* ViewInfo = static_cast<const FViewInfo*>(View);
-		SetShaderValue(Shader, ExponentialFogParameters, ViewInfo->ExponentialFogParameters);
-		SetShaderValue(Shader, ExponentialFogColorParameter, FVector4(ViewInfo->ExponentialFogColor, 1.0f - ViewInfo->FogMaxOpacity));
-		SetShaderValue(Shader, InscatteringLightDirection, FVector4(ViewInfo->InscatteringLightDirection, ViewInfo->bUseDirectionalInscattering ? 1 : 0));
-		SetShaderValue(Shader, DirectionalInscatteringColor, FVector4(FVector(ViewInfo->DirectionalInscatteringColor), ViewInfo->DirectionalInscatteringExponent));
-		SetShaderValue(Shader, DirectionalInscatteringStartDistance, ViewInfo->DirectionalInscatteringStartDistance);
+		SetShaderValue(RHICmdList, Shader, ExponentialFogParameters, ViewInfo->ExponentialFogParameters);
+		SetShaderValue(RHICmdList, Shader, ExponentialFogColorParameter, FVector4(ViewInfo->ExponentialFogColor, 1.0f - ViewInfo->FogMaxOpacity));
+		SetShaderValue(RHICmdList, Shader, InscatteringLightDirection, FVector4(ViewInfo->InscatteringLightDirection, ViewInfo->bUseDirectionalInscattering ? 1 : 0));
+		SetShaderValue(RHICmdList, Shader, DirectionalInscatteringColor, FVector4(FVector(ViewInfo->DirectionalInscatteringColor), ViewInfo->DirectionalInscatteringExponent));
+		SetShaderValue(RHICmdList, Shader, DirectionalInscatteringStartDistance, ViewInfo->DirectionalInscatteringStartDistance);
 	}
 
 	/** Serializer. */
@@ -50,10 +50,10 @@ public:
 	 * Sets height fog shader parameters.
 	 */
 	template<typename ShaderRHIParamRef>
-	void Set(ShaderRHIParamRef Shader, const FSceneView* View) const
+	void Set(FRHICommandList* RHICmdList, ShaderRHIParamRef Shader, const FSceneView* View) const
 	{
 		// Set the fog constants.
-		ExponentialParameters.Set(Shader, View);
+		ExponentialParameters.Set(RHICmdList, Shader, View);
 	}
 
 	/** Serializer. */
diff --git a/Engine/Source/Runtime/Engine/Public/GPUSkinVertexFactory.h b/Engine/Source/Runtime/Engine/Public/GPUSkinVertexFactory.h
--- a/Engine/Source/Runtime/Engine/Public/GPUSkinVertexFactory.h
+++ b/Engine/Source/Runtime/Engine/Public/GPUSkinVertexFactory.h
@@ -87,6 +87,7 @@ public:
 	enum { NumColumns = 4 };
 	enum { NumElements = 0 };
 	enum { Alignment = 16 };
+	enum { IsResource = 0 };
 	static const FUniformBufferStruct* GetStruct() { return NULL; }
 };
 
diff --git a/Engine/Source/Runtime/Engine/Public/GlobalShader.h b/Engine/Source/Runtime/Engine/Public/GlobalShader.h
--- a/Engine/Source/Runtime/Engine/Public/GlobalShader.h
+++ b/Engine/Source/Runtime/Engine/Public/GlobalShader.h
@@ -118,11 +118,11 @@ public:
 	ENGINE_API FGlobalShader(const ShaderMetaType::CompiledShaderInitializerType& Initializer);
 	
 	template<typename ShaderRHIParamRef>
-	void SetParameters(const ShaderRHIParamRef ShaderRHI,const FSceneView& View)
+	void SetParameters(FRHICommandList* RHICmdList, const ShaderRHIParamRef ShaderRHI,const FSceneView& View)
 	{
 		check(GetUniformBufferParameter<FViewUniformShaderParameters>().IsInitialized());
 		CheckShaderIsValid();
-		SetUniformBufferParameter(ShaderRHI,GetUniformBufferParameter<FViewUniformShaderParameters>(),View.UniformBuffer);
+		SetUniformBufferParameter(RHICmdList, ShaderRHI,GetUniformBufferParameter<FViewUniformShaderParameters>(),View.UniformBuffer);
 	}
 
 	typedef void (*ModifyCompilationEnvironmentType)(EShaderPlatform, FShaderCompilerEnvironment&);
diff --git a/Engine/Source/Runtime/Engine/Public/Landscape/LandscapeRender.h b/Engine/Source/Runtime/Engine/Public/Landscape/LandscapeRender.h
--- a/Engine/Source/Runtime/Engine/Public/Landscape/LandscapeRender.h
+++ b/Engine/Source/Runtime/Engine/Public/Landscape/LandscapeRender.h
@@ -112,7 +112,7 @@ public:
 	/**
 	* Set any shader data specific to this vertex factory
 	*/
-	virtual void SetMesh(FShader* PixelShader,const class FVertexFactory* VertexFactory,const class FSceneView& View,const struct FMeshBatchElement& BatchElement,uint32 DataFlags) const OVERRIDE;
+	virtual void SetMesh(FRHICommandList* RHICmdList, FShader* PixelShader,const class FVertexFactory* VertexFactory,const class FSceneView& View,const struct FMeshBatchElement& BatchElement,uint32 DataFlags) const OVERRIDE;
 	
 	virtual uint32 GetSize() const OVERRIDE
 	{
diff --git a/Engine/Source/Runtime/Engine/Public/MaterialShared.h b/Engine/Source/Runtime/Engine/Public/MaterialShared.h
--- a/Engine/Source/Runtime/Engine/Public/MaterialShared.h
+++ b/Engine/Source/Runtime/Engine/Public/MaterialShared.h
@@ -216,7 +216,6 @@ public:
 	bool IsEmpty() const;
 	bool operator==(const FUniformExpressionSet& ReferenceSet) const;
 	FString GetSummaryString() const;
-	void GetResourcesString(EShaderPlatform ShaderPlatform,FString& InputsString) const;
 
 	void SetParameterCollections(const TArray<class UMaterialParameterCollection*>& Collections);
 	void CreateBufferStruct();
diff --git a/Engine/Source/Runtime/Engine/Public/ParameterCollection.h b/Engine/Source/Runtime/Engine/Public/ParameterCollection.h
--- a/Engine/Source/Runtime/Engine/Public/ParameterCollection.h
+++ b/Engine/Source/Runtime/Engine/Public/ParameterCollection.h
@@ -47,5 +47,7 @@ private:
 	/** Uniform buffer containing the UMaterialParameterCollection default parameter values and UMaterialParameterCollectionInstance instance overrides. */
 	FUniformBufferRHIRef UniformBuffer;
 
+	FRHIUniformBufferLayout UniformBufferLayout;
+
 	void UpdateContents(const FGuid& InId, const TArray<FVector4>& Data);
 };
\ No newline at end of file
diff --git a/Engine/Source/Runtime/Engine/Public/SceneManagement.h b/Engine/Source/Runtime/Engine/Public/SceneManagement.h
--- a/Engine/Source/Runtime/Engine/Public/SceneManagement.h
+++ b/Engine/Source/Runtime/Engine/Public/SceneManagement.h
@@ -1201,7 +1201,7 @@ inline TUniformBufferRef<FPrimitiveUniformShaderParameters> CreatePrimitiveUnifo
 	check(IsInRenderingThread());
 	return TUniformBufferRef<FPrimitiveUniformShaderParameters>::CreateUniformBufferImmediate(
 		GetPrimitiveUniformShaderParameters(LocalToWorld, WorldBounds.Origin, WorldBounds, LocalBounds, bReceivesDecals, false, LpvBiasMultiplier ),
-		UniformBuffer_MultiUse
+		UniformBuffer_MultiFrame
 		);
 }
 
@@ -2645,6 +2645,9 @@ public:
 
 	/** Whether this view is being used to render a reflection capture. */
 	bool bIsReflectionCapture;
+	
+	/** Whether this view was created from a locked viewpoint. */
+	bool bIsLocked;
 
 	/** 
 	 * Whether to only render static lights and objects.  
diff --git a/Engine/Source/Runtime/Engine/Public/ScreenRendering.h b/Engine/Source/Runtime/Engine/Public/ScreenRendering.h
--- a/Engine/Source/Runtime/Engine/Public/ScreenRendering.h
+++ b/Engine/Source/Runtime/Engine/Public/ScreenRendering.h
@@ -1,7 +1,7 @@
 // Copyright 1998-2014 Epic Games, Inc. All Rights Reserved.
 
 /*=============================================================================
-	ScreenRendering.h: D3D render target implementation.
+	ScreenRendering.h: Screen rendering definitions.
 =============================================================================*/
 
 #pragma once
@@ -59,14 +59,14 @@ public:
 	}
 	FScreenPS() {}
 
-	void SetParameters(const FTexture* Texture)
+	void SetParameters(FRHICommandList* RHICmdList, const FTexture* Texture)
 	{
-		SetTextureParameter(GetPixelShader(),InTexture,InTextureSampler,Texture);
+		SetTextureParameter(RHICmdList, GetPixelShader(),InTexture,InTextureSampler,Texture);
 	}
 
-	void SetParameters(FSamplerStateRHIParamRef SamplerStateRHI, FTextureRHIParamRef TextureRHI)
+	void SetParameters(FRHICommandList* RHICmdList, FSamplerStateRHIParamRef SamplerStateRHI, FTextureRHIParamRef TextureRHI)
 	{
-		SetTextureParameter(GetPixelShader(),InTexture,InTextureSampler,SamplerStateRHI,TextureRHI);
+		SetTextureParameter(RHICmdList, GetPixelShader(),InTexture,InTextureSampler,SamplerStateRHI,TextureRHI);
 	}
 
 	virtual bool Serialize(FArchive& Ar)
@@ -99,9 +99,9 @@ public:
 	FScreenVS() {}
 
 
-	void SetParameters(const FSceneView& View)
+	void SetParameters(FRHICommandList* RHICmdList, const FSceneView& View)
 	{
-		FGlobalShader::SetParameters(GetVertexShader(), View);
+		FGlobalShader::SetParameters(RHICmdList, GetVertexShader(), View);
 	}
 
 
diff --git a/Engine/Source/Runtime/Engine/Public/SimpleElementShaders.h b/Engine/Source/Runtime/Engine/Public/SimpleElementShaders.h
--- a/Engine/Source/Runtime/Engine/Public/SimpleElementShaders.h
+++ b/Engine/Source/Runtime/Engine/Public/SimpleElementShaders.h
@@ -22,7 +22,7 @@ public:
 	FSimpleElementVS(const ShaderMetaType::CompiledShaderInitializerType& Initializer);
 	FSimpleElementVS() {}
 
-	/*ENGINE_API */void SetParameters(const FMatrix& TransformValue, bool bSwitchVerticalAxis = false);
+	/*ENGINE_API */void SetParameters(FRHICommandList* RHICmdList, const FMatrix& TransformValue, bool bSwitchVerticalAxis = false);
 
 	virtual bool Serialize(FArchive& Ar);
 
@@ -53,9 +53,9 @@ public:
 	 * @param View			SceneView for view constants when compositing
 	 * @param DepthTexture	Depth texture to read from when depth testing for compositing.  If not set no compositing will occur
 	 */
-	void SetEditorCompositingParameters( const FSceneView* View, FTexture2DRHIRef DepthTexture );
+	void SetEditorCompositingParameters(FRHICommandList* RHICmdList, const FSceneView* View, FTexture2DRHIRef DepthTexture );
 
-	void SetParameters(const FTexture* TextureValue );
+	void SetParameters(FRHICommandList* RHICmdList, const FTexture* TextureValue );
 
 	virtual bool Serialize(FArchive& Ar);
 
@@ -82,7 +82,7 @@ public:
 	FSimpleElementGammaPS(const ShaderMetaType::CompiledShaderInitializerType& Initializer);
 	FSimpleElementGammaPS() {}
 
-	void SetParameters(const FTexture* Texture,float GammaValue,ESimpleElementBlendMode BlendMode);
+	void SetParameters(FRHICommandList* RHICmdList, const FTexture* Texture,float GammaValue,ESimpleElementBlendMode BlendMode);
 
 	virtual bool Serialize(FArchive& Ar);
 
@@ -103,7 +103,7 @@ public:
 	FSimpleElementMaskedGammaPS(const ShaderMetaType::CompiledShaderInitializerType& Initializer);
 	FSimpleElementMaskedGammaPS() {}
 
-	void SetParameters(const FTexture* Texture,float Gamma,float ClipRefValue,ESimpleElementBlendMode BlendMode);
+	void SetParameters(FRHICommandList* RHICmdList, const FTexture* Texture,float Gamma,float ClipRefValue,ESimpleElementBlendMode BlendMode);
 
 	virtual bool Serialize(FArchive& Ar);
 
@@ -156,6 +156,7 @@ public:
 	* @param BlendMode - current batched element blend mode being rendered
 	*/
 	void SetParameters(
+		FRHICommandList* RHICmdList, 
 		const FTexture* Texture,
 		float Gamma,
 		float ClipRef,
@@ -214,7 +215,7 @@ public:
 	FSimpleElementHitProxyPS(const ShaderMetaType::CompiledShaderInitializerType& Initializer);
 	FSimpleElementHitProxyPS() {}
 
-	void SetParameters(const FTexture* TextureValue);
+	void SetParameters(FRHICommandList* RHICmdList, const FTexture* TextureValue);
 
 	virtual bool Serialize(FArchive& Ar);
 
@@ -250,7 +251,7 @@ public:
 	* @param ColorWeights - reference value to compare with alpha for killing pixels
 	* @param Gamma - if gamma != 1.0 then a pow(color,Gamma) is applied
 	*/
-	void SetParameters(const FTexture* TextureValue, const FMatrix& ColorWeightsValue, float GammaValue);
+	void SetParameters(FRHICommandList* RHICmdList, const FTexture* TextureValue, const FMatrix& ColorWeightsValue, float GammaValue);
 
 	virtual bool Serialize(FArchive& Ar);
 
diff --git a/Engine/Source/Runtime/Engine/Public/SpeedTreeWind.h b/Engine/Source/Runtime/Engine/Public/SpeedTreeWind.h
--- a/Engine/Source/Runtime/Engine/Public/SpeedTreeWind.h
+++ b/Engine/Source/Runtime/Engine/Public/SpeedTreeWind.h
@@ -323,7 +323,7 @@ protected:
  * Uniform buffer setup for SpeedTrees.
  */
 
-BEGIN_UNIFORM_BUFFER_STRUCT(FSpeedTreeUniformParameters, )
+BEGIN_UNIFORM_BUFFER_STRUCT(FSpeedTreeUniformParameters, ENGINE_API)
 	DECLARE_UNIFORM_BUFFER_STRUCT_MEMBER(FVector4, WindVector)
 	DECLARE_UNIFORM_BUFFER_STRUCT_MEMBER(FVector4, WindGlobal)
 	DECLARE_UNIFORM_BUFFER_STRUCT_MEMBER(FVector4, WindBranch)
diff --git a/Engine/Source/Runtime/Engine/Public/Texture.h b/Engine/Source/Runtime/Engine/Public/Texture.h
--- a/Engine/Source/Runtime/Engine/Public/Texture.h
+++ b/Engine/Source/Runtime/Engine/Public/Texture.h
@@ -187,8 +187,6 @@ class FTextureResource : public FTexture
 {
 public:
 
-	FRenderCommandFence ReleaseFence;
-
 	FTextureResource()
 	{}
 	virtual ~FTextureResource() {}
diff --git a/Engine/Source/Runtime/Launch/Launch.Build.cs b/Engine/Source/Runtime/Launch/Launch.Build.cs
--- a/Engine/Source/Runtime/Launch/Launch.Build.cs
+++ b/Engine/Source/Runtime/Launch/Launch.Build.cs
@@ -162,6 +162,7 @@ public class Launch : ModuleRules
 			(Target.Platform == UnrealTargetPlatform.Mac) ||
 			(Target.Platform == UnrealTargetPlatform.Linux && !UEBuildConfiguration.bBuildDedicatedServer))
 		{
+            // TODO: re-enable after implementing resource tables for OpenGL.
 			DynamicallyLoadedModuleNames.Add("OpenGLDrv");
 		}
 
diff --git a/Engine/Source/Runtime/OpenGLDrv/Private/OpenGLCommands.cpp b/Engine/Source/Runtime/OpenGLDrv/Private/OpenGLCommands.cpp
--- a/Engine/Source/Runtime/OpenGLDrv/Private/OpenGLCommands.cpp
+++ b/Engine/Source/Runtime/OpenGLDrv/Private/OpenGLCommands.cpp
@@ -986,39 +986,31 @@ void FOpenGLDynamicRHI::RHISetShaderSampler(FPixelShaderRHIParamRef PixelShaderR
 void FOpenGLDynamicRHI::RHISetShaderUniformBuffer(FVertexShaderRHIParamRef VertexShaderRHI,uint32 BufferIndex,FUniformBufferRHIParamRef BufferRHI)
 {
 	VALIDATE_BOUND_SHADER(VertexShaderRHI);
-
-	VERIFY_GL_SCOPE();
-	DYNAMIC_CAST_OPENGLRESOURCE(VertexShader,VertexShader);
-	VertexShader->BoundUniformBuffers[BufferIndex] = BufferRHI;
+	PendingState.BoundUniformBuffers[SF_Vertex][BufferIndex] = BufferRHI;
+	PendingState.DirtyUniformBuffers[SF_Vertex] |= 1 << BufferIndex;
 }
 
 void FOpenGLDynamicRHI::RHISetShaderUniformBuffer(FHullShaderRHIParamRef HullShaderRHI,uint32 BufferIndex,FUniformBufferRHIParamRef BufferRHI)
 {
 	VALIDATE_BOUND_SHADER(HullShaderRHI);
-
 	check(GRHIFeatureLevel >= ERHIFeatureLevel::SM5);
-	VERIFY_GL_SCOPE();
-	DYNAMIC_CAST_OPENGLRESOURCE(HullShader,HullShader);
-	HullShader->BoundUniformBuffers[BufferIndex] = BufferRHI;
+	PendingState.BoundUniformBuffers[SF_Hull][BufferIndex] = BufferRHI;
+	PendingState.DirtyUniformBuffers[SF_Hull] |= 1 << BufferIndex;
 }
 
 void FOpenGLDynamicRHI::RHISetShaderUniformBuffer(FDomainShaderRHIParamRef DomainShaderRHI,uint32 BufferIndex,FUniformBufferRHIParamRef BufferRHI)
 {
 	VALIDATE_BOUND_SHADER(DomainShaderRHI);
-
 	check(GRHIFeatureLevel >= ERHIFeatureLevel::SM5);
-	VERIFY_GL_SCOPE();
-	DYNAMIC_CAST_OPENGLRESOURCE(DomainShader,DomainShader);
-	DomainShader->BoundUniformBuffers[BufferIndex] = BufferRHI;
+	PendingState.BoundUniformBuffers[SF_Domain][BufferIndex] = BufferRHI;
+	PendingState.DirtyUniformBuffers[SF_Domain] |= 1 << BufferIndex;
 }
 
 void FOpenGLDynamicRHI::RHISetShaderUniformBuffer(FGeometryShaderRHIParamRef GeometryShaderRHI,uint32 BufferIndex,FUniformBufferRHIParamRef BufferRHI)
 {
 	VALIDATE_BOUND_SHADER(GeometryShaderRHI);
-
-	VERIFY_GL_SCOPE();
-	DYNAMIC_CAST_OPENGLRESOURCE(GeometryShader,GeometryShader);
-	GeometryShader->BoundUniformBuffers[BufferIndex] = BufferRHI;
+	PendingState.BoundUniformBuffers[SF_Geometry][BufferIndex] = BufferRHI;
+	PendingState.DirtyUniformBuffers[SF_Geometry] |= 1 << BufferIndex;
 }
 
 void FOpenGLDynamicRHI::RHISetShaderSampler(FComputeShaderRHIParamRef ComputeShaderRHI,uint32 SamplerIndex,FSamplerStateRHIParamRef NewStateRHI)
@@ -1027,24 +1019,20 @@ void FOpenGLDynamicRHI::RHISetShaderSampler(FComputeShaderRHIParamRef ComputeSha
 	VERIFY_GL_SCOPE();
 	DYNAMIC_CAST_OPENGLRESOURCE(SamplerState,NewState);
 	FOpenGL::BindSampler(FOpenGL::GetFirstComputeTextureUnit() + SamplerIndex, NewState->Resource);
-
 }
 
 void FOpenGLDynamicRHI::RHISetShaderUniformBuffer(FPixelShaderRHIParamRef PixelShaderRHI,uint32 BufferIndex,FUniformBufferRHIParamRef BufferRHI)
 {
 	VALIDATE_BOUND_SHADER(PixelShaderRHI);
-
-	VERIFY_GL_SCOPE();
-	DYNAMIC_CAST_OPENGLRESOURCE(PixelShader,PixelShader);
-	PixelShader->BoundUniformBuffers[BufferIndex] = BufferRHI;
+	PendingState.BoundUniformBuffers[SF_Pixel][BufferIndex] = BufferRHI;
+	PendingState.DirtyUniformBuffers[SF_Pixel] |= 1 << BufferIndex;
 }
 
 void FOpenGLDynamicRHI::RHISetShaderUniformBuffer(FComputeShaderRHIParamRef ComputeShaderRHI,uint32 BufferIndex,FUniformBufferRHIParamRef BufferRHI)
 {
 	check(GRHIFeatureLevel >= ERHIFeatureLevel::SM5);
-	VERIFY_GL_SCOPE();
-	DYNAMIC_CAST_OPENGLRESOURCE(ComputeShader, ComputeShader);
-	ComputeShader->BoundUniformBuffers[BufferIndex] = BufferRHI;
+	PendingState.BoundUniformBuffers[SF_Compute][BufferIndex] = BufferRHI;
+	PendingState.DirtyUniformBuffers[SF_Compute] |= 1 << BufferIndex;
 }
 
 void FOpenGLDynamicRHI::RHISetShaderParameter(FVertexShaderRHIParamRef VertexShaderRHI,uint32 BufferIndex,uint32 BaseIndex,uint32 NumBytes,const void* NewValue)
@@ -1726,7 +1714,7 @@ void FOpenGLDynamicRHI::FreeZeroStrideBuffers()
 void FOpenGLDynamicRHI::SetupVertexArrays(FOpenGLContextState& ContextState, uint32 BaseVertexIndex, FOpenGLStream* Streams, uint32 NumStreams, uint32 MaxVertices)
 {
 	if (FOpenGL::SupportsVertexAttribBinding() && OpenGLConsoleVariables::bUseVAB)
-	{
+{
 		SetupVertexArraysVAB(ContextState, BaseVertexIndex, Streams, NumStreams, MaxVertices);
 		return;
 	}
@@ -2055,9 +2043,7 @@ void FOpenGLDynamicRHI::OnPixelBufferDeletion( GLuint PixelBufferResource )
 	}
 }
 
-extern void AddNewlyFreedBufferToUniformBufferPool( GLuint Buffer, uint32 BufferSize, bool bStreamDraw, uint32 Offset, uint8* Pointer );
-
-void FOpenGLDynamicRHI::OnUniformBufferDeletion( GLuint UniformBufferResource, uint32 AllocatedSize, bool bStreamDraw, uint32 Offset, uint8* Pointer )
+void FOpenGLDynamicRHI::OnUniformBufferDeletion( GLuint UniformBufferResource, uint32 AllocatedSize, bool bStreamDraw )
 {
 	if (SharedContextState.UniformBufferBound == UniformBufferResource)
 	{
@@ -2081,8 +2067,6 @@ void FOpenGLDynamicRHI::OnUniformBufferDeletion( GLuint UniformBufferResource, u
 			RenderingContextState.UniformBuffers[UniformBufferIndex] = FOpenGLCachedUniformBuffer_Invalid;	// that'll enforce state update on next cache test
 		}
 	}
-
-	AddNewlyFreedBufferToUniformBufferPool( UniformBufferResource, AllocatedSize, bStreamDraw, Offset, Pointer );
 }
 
 void FOpenGLDynamicRHI::CommitNonComputeShaderConstants()
@@ -2092,13 +2076,13 @@ void FOpenGLDynamicRHI::CommitNonComputeShaderConstants()
 	FOpenGLLinkedProgram* LinkedProgram = PendingState.BoundShaderState->LinkedProgram;
 	if (IsES2Platform(GRHIShaderPlatform))
 	{
-		PendingState.ShaderParameters[OGL_SHADER_STAGE_VERTEX].CommitPackedUniformBuffers(LinkedProgram, OGL_SHADER_STAGE_VERTEX, PendingState.BoundShaderState->VertexShader->BoundUniformBuffers, PendingState.BoundShaderState->VertexShader->UniformBuffersCopyInfo);
+		PendingState.ShaderParameters[OGL_SHADER_STAGE_VERTEX].CommitPackedUniformBuffers(LinkedProgram, OGL_SHADER_STAGE_VERTEX, PendingState.BoundUniformBuffers[SF_Vertex], PendingState.BoundShaderState->VertexShader->UniformBuffersCopyInfo);
 	}
 	PendingState.ShaderParameters[OGL_SHADER_STAGE_VERTEX].CommitPackedGlobals(LinkedProgram, OGL_SHADER_STAGE_VERTEX);
 
 	if (IsES2Platform(GRHIShaderPlatform))
 	{
-		PendingState.ShaderParameters[OGL_SHADER_STAGE_PIXEL].CommitPackedUniformBuffers(LinkedProgram, OGL_SHADER_STAGE_PIXEL, PendingState.BoundShaderState->PixelShader->BoundUniformBuffers, PendingState.BoundShaderState->PixelShader->UniformBuffersCopyInfo);
+		PendingState.ShaderParameters[OGL_SHADER_STAGE_PIXEL].CommitPackedUniformBuffers(LinkedProgram, OGL_SHADER_STAGE_PIXEL, PendingState.BoundUniformBuffers[SF_Pixel], PendingState.BoundShaderState->PixelShader->UniformBuffersCopyInfo);
 	}
 	PendingState.ShaderParameters[OGL_SHADER_STAGE_PIXEL].CommitPackedGlobals(LinkedProgram, OGL_SHADER_STAGE_PIXEL);
 
@@ -2106,13 +2090,12 @@ void FOpenGLDynamicRHI::CommitNonComputeShaderConstants()
 	{
 		if (IsES2Platform(GRHIShaderPlatform))
 		{
-			PendingState.ShaderParameters[OGL_SHADER_STAGE_GEOMETRY].CommitPackedUniformBuffers(LinkedProgram, OGL_SHADER_STAGE_GEOMETRY, PendingState.BoundShaderState->GeometryShader->BoundUniformBuffers, PendingState.BoundShaderState->GeometryShader->UniformBuffersCopyInfo);
+			PendingState.ShaderParameters[OGL_SHADER_STAGE_GEOMETRY].CommitPackedUniformBuffers(LinkedProgram, OGL_SHADER_STAGE_GEOMETRY, PendingState.BoundUniformBuffers[SF_Geometry], PendingState.BoundShaderState->GeometryShader->UniformBuffersCopyInfo);
 		}
 		PendingState.ShaderParameters[OGL_SHADER_STAGE_GEOMETRY].CommitPackedGlobals(LinkedProgram, OGL_SHADER_STAGE_GEOMETRY);
 	}
 }
 
-
 void FOpenGLDynamicRHI::CommitComputeShaderConstants(FComputeShaderRHIParamRef ComputeShaderRHI)
 {
 	VERIFY_GL_SCOPE();
@@ -2125,12 +2108,138 @@ void FOpenGLDynamicRHI::CommitComputeShaderConstants(FComputeShaderRHIParamRef C
 	StageShaderParameters.CommitPackedGlobals(ComputeShader->LinkedProgram, OGL_SHADER_STAGE_COMPUTE);
 }
 
+template <EShaderFrequency Frequency>
+FORCEINLINE uint32 GetFirstTextureUnit()
+{
+	switch (Frequency)
+	{
+	case SF_Vertex: return FOpenGL::GetFirstVertexTextureUnit();
+	case SF_Hull: return FOpenGL::GetFirstHullTextureUnit();
+	case SF_Domain: return FOpenGL::GetFirstDomainTextureUnit();
+	case SF_Pixel: return FOpenGL::GetFirstPixelTextureUnit();
+	case SF_Geometry: return FOpenGL::GetFirstGeometryTextureUnit();
+	case SF_Compute: return FOpenGL::GetFirstComputeTextureUnit();
+	}
+	return INDEX_NONE;
+}
+
+template <EShaderFrequency Frequency>
+FORCEINLINE void SetResource(FOpenGLDynamicRHI* RESTRICT OpenGLRHI, uint32 BindIndex, FOpenGLTextureBase* RESTRICT Texture)
+{
+	OpenGLRHI->InternalSetShaderTexture(Texture, GetFirstTextureUnit<Frequency>() + BindIndex, Texture->Target, Texture->Resource, Texture->NumMips, -1);
+}
+
+template <EShaderFrequency Frequency>
+FORCEINLINE void SetResource(FOpenGLDynamicRHI* RESTRICT OpenGLRHI, uint32 BindIndex, FOpenGLSamplerState* RESTRICT SamplerState)
+{
+	if (FOpenGL::SupportsSamplerObjects())
+	{
+		PTRINT SamplerStateAsInt = (PTRINT)SamplerState;
+		FOpenGL::BindSampler(GetFirstTextureUnit<Frequency>() + BindIndex, (GLuint)SamplerStateAsInt);
+	}
+	else
+	{
+		OpenGLRHI->InternalSetSamplerStates(GetFirstTextureUnit<Frequency>() + BindIndex, SamplerState);
+	}
+}
+
+template <EShaderFrequency Frequency>
+FORCEINLINE void SetResource(FOpenGLDynamicRHI* RESTRICT OpenGLRHI, uint32 BindIndex, FOpenGLShaderResourceView* RESTRICT SRV)
+{
+	OpenGLRHI->InternalSetShaderTexture(NULL, GetFirstTextureUnit<Frequency>() + BindIndex, SRV->Target, SRV->Resource, 0, SRV->LimitMip);
+	SetResource<Frequency>(OpenGLRHI,BindIndex,OpenGLRHI->GetPointSamplerState());
+}
+
+template <class GLResourceType, EShaderFrequency ShaderFrequency>
+inline int32 SetShaderResourcesFromBuffer(FOpenGLDynamicRHI* RESTRICT OpenGLRHI, FOpenGLUniformBuffer* RESTRICT Buffer, const uint32 * RESTRICT ResourceMap, int32 BufferIndex)
+{
+	int32 NumSetCalls = 0;
+	uint32 BufferOffset = ResourceMap[BufferIndex];
+	if (BufferOffset > 0)
+	{
+		const uint32* RESTRICT ResourceInfos = &ResourceMap[BufferOffset];
+		uint32 ResourceInfo = *ResourceInfos++;
+		do 
+		{
+			checkSlow(FRHIResourceTableEntry::GetUniformBufferIndex(ResourceInfo) == BufferIndex);
+			const uint16 ResourceIndex = FRHIResourceTableEntry::GetResourceIndex(ResourceInfo);
+			const uint8 BindIndex = FRHIResourceTableEntry::GetBindIndex(ResourceInfo);
+
+			GLResourceType* ResourcePtr = (GLResourceType*)Buffer->RawResourceTable[ResourceIndex];
+			SetResource<ShaderFrequency>(OpenGLRHI, BindIndex, ResourcePtr);
+
+			NumSetCalls++;
+			ResourceInfo = *ResourceInfos++;
+		} while (FRHIResourceTableEntry::GetUniformBufferIndex(ResourceInfo) == BufferIndex);
+	}
+	return NumSetCalls;
+}
+
+template <class ShaderType>
+void FOpenGLDynamicRHI::SetResourcesFromTables(const ShaderType* RESTRICT Shader)
+{
+	checkSlow(Shader);
+	const FOpenGLShaderResourceTable* RESTRICT SRT = &Shader->Bindings.ShaderResourceTable;
+
+	// Mask the dirty bits by those buffers from which the shader has bound resources.
+	uint32 DirtyBits = SRT->ResourceTableBits & PendingState.DirtyUniformBuffers[ShaderType::StaticFrequency];
+	uint32 NumSetCalls = 0;
+	while (DirtyBits)
+	{
+		// Scan for the lowest set bit, compute its index, clear it in the set of dirty bits.
+		const uint32 LowestBitMask = (DirtyBits) & (-(int32)DirtyBits);
+		const int32 BufferIndex = FMath::FloorLog2(LowestBitMask); // todo: This has a branch on zero, we know it could never be zero...
+		DirtyBits ^= LowestBitMask;
+
+		FOpenGLUniformBuffer* Buffer = (FOpenGLUniformBuffer*)PendingState.BoundUniformBuffers[ShaderType::StaticFrequency][BufferIndex].GetReference();
+		check(Buffer);
+		check(BufferIndex < SRT->ResourceTableLayoutHashes.Num());
+		check(Buffer->GetLayout().GetHash() == SRT->ResourceTableLayoutHashes[BufferIndex]);
+		Buffer->CacheResources(ResourceTableFrameCounter);
+
+		// todo: could make this two pass: gather then set
+		NumSetCalls += SetShaderResourcesFromBuffer<FOpenGLTextureBase,(EShaderFrequency)ShaderType::StaticFrequency>(this,Buffer,SRT->TextureMap.GetData(),BufferIndex);
+		NumSetCalls += SetShaderResourcesFromBuffer<FOpenGLShaderResourceView,(EShaderFrequency)ShaderType::StaticFrequency>(this,Buffer,SRT->ShaderResourceViewMap.GetData(),BufferIndex);
+		SetShaderResourcesFromBuffer<FOpenGLSamplerState,(EShaderFrequency)ShaderType::StaticFrequency>(this,Buffer,SRT->SamplerMap.GetData(),BufferIndex);
+	}
+	PendingState.DirtyUniformBuffers[ShaderType::StaticFrequency] = 0;
+	//SetTextureInTableCalls += NumSetCalls;
+}
+
+void FOpenGLDynamicRHI::CommitGraphicsResourceTables()
+{
+	if (auto* Shader = PendingState.BoundShaderState->VertexShader.GetReference())
+	{
+		SetResourcesFromTables(Shader);
+	}
+	if (auto* Shader = PendingState.BoundShaderState->PixelShader.GetReference())
+	{
+		SetResourcesFromTables(Shader);
+	}
+	if (auto* Shader = PendingState.BoundShaderState->HullShader.GetReference())
+	{
+		SetResourcesFromTables(Shader);
+	}
+	if (auto* Shader = PendingState.BoundShaderState->DomainShader.GetReference())
+	{
+		SetResourcesFromTables(Shader);
+	}
+	if (auto* Shader = PendingState.BoundShaderState->GeometryShader.GetReference())
+	{
+		SetResourcesFromTables(Shader);
+	}
+}
+
+void FOpenGLDynamicRHI::CommitComputeResourceTables(FOpenGLComputeShader* ComputeShader)
+{
+	check(ComputeShader);
+	SetResourcesFromTables(ComputeShader);
+}
 
 void FOpenGLDynamicRHI::RHIDrawPrimitive(uint32 PrimitiveType,uint32 BaseVertexIndex,uint32 NumPrimitives,uint32 NumInstances)
 {
 	SCOPE_CYCLE_COUNTER_DETAILED(STAT_OpenGLDrawPrimitiveTime);
 	VERIFY_GL_SCOPE();
-
 	RHI_DRAW_CALL_STATS(PrimitiveType,NumPrimitives*NumInstances);
 
 	FOpenGLContextState& ContextState = GetContextStateForCurrentContext();
@@ -2141,6 +2250,7 @@ void FOpenGLDynamicRHI::RHIDrawPrimitive(uint32 PrimitiveType,uint32 BaseVertexI
 	UpdateRasterizerStateInOpenGLContext(ContextState);
 	UpdateDepthStencilStateInOpenGLContext(ContextState);
 	BindPendingShaderState(ContextState);
+	CommitGraphicsResourceTables();
 	SetupTexturesForDraw(ContextState);
 	CommitNonComputeShaderConstants();
 	CachedBindElementArrayBuffer(ContextState,0);
@@ -2178,7 +2288,7 @@ void FOpenGLDynamicRHI::RHIDrawPrimitiveIndirect(uint32 PrimitiveType,FVertexBuf
 		VERIFY_GL_SCOPE();
 
 		check(ArgumentBufferRHI);
-		GPUProfilingData.RegisterGPUWork(0);
+	GPUProfilingData.RegisterGPUWork(0);
 
 		FOpenGLContextState& ContextState = GetContextStateForCurrentContext();
 		BindPendingFramebuffer(ContextState);
@@ -2228,7 +2338,7 @@ void FOpenGLDynamicRHI::RHIDrawIndexedIndirect(FIndexBufferRHIParamRef IndexBuff
 		VERIFY_GL_SCOPE();
 
 		DYNAMIC_CAST_OPENGLRESOURCE(IndexBuffer,IndexBuffer);
-		GPUProfilingData.RegisterGPUWork(1);
+	GPUProfilingData.RegisterGPUWork(1);
 
 		check(ArgumentsBufferRHI);
 
@@ -2295,6 +2405,7 @@ void FOpenGLDynamicRHI::RHIDrawIndexedPrimitive(FIndexBufferRHIParamRef IndexBuf
 	UpdateRasterizerStateInOpenGLContext(ContextState);
 	UpdateDepthStencilStateInOpenGLContext(ContextState);
 	BindPendingShaderState(ContextState);
+	CommitGraphicsResourceTables();
 	SetupTexturesForDraw(ContextState);
 	CommitNonComputeShaderConstants();
 	CachedBindElementArrayBuffer(ContextState,IndexBuffer->Resource);
@@ -2383,8 +2494,8 @@ void FOpenGLDynamicRHI::RHIDrawIndexedPrimitiveIndirect(uint32 PrimitiveType,FIn
 		glBindBuffer( GL_DRAW_INDIRECT_BUFFER, 0);
 	}
 	else
-	{
-		UE_LOG(LogRHI, Fatal,TEXT("OpenGL RHI does not yet support indirect draw calls."));
+{
+	UE_LOG(LogRHI, Fatal,TEXT("OpenGL RHI does not yet support indirect draw calls."));
 	}
 }
 
@@ -2464,6 +2575,7 @@ void FOpenGLDynamicRHI::RHIEndDrawPrimitiveUP()
 	UpdateRasterizerStateInOpenGLContext(ContextState);
 	UpdateDepthStencilStateInOpenGLContext(ContextState);
 	BindPendingShaderState(ContextState);
+	CommitGraphicsResourceTables();
 	SetupTexturesForDraw(ContextState);
 	CommitNonComputeShaderConstants();
 	CachedBindElementArrayBuffer(ContextState,0);
@@ -2589,6 +2701,7 @@ void FOpenGLDynamicRHI::RHIEndDrawIndexedPrimitiveUP()
 	UpdateRasterizerStateInOpenGLContext(ContextState);
 	UpdateDepthStencilStateInOpenGLContext(ContextState);
 	BindPendingShaderState(ContextState);
+	CommitGraphicsResourceTables();
 	SetupTexturesForDraw(ContextState);
 	CommitNonComputeShaderConstants();
 	if(FOpenGL::SupportsFastBufferData()) 
@@ -2773,7 +2886,7 @@ void FOpenGLDynamicRHI::RHIClearMRT(bool bClearColor,int32 NumClearColors,const
 {
 	VERIFY_GL_SCOPE();
 
-	check((GRHIFeatureLevel >= ERHIFeatureLevel::SM5) || !PendingState.bFramebufferSetupInvalid);
+	check((GRHIFeatureLevel >= ERHIFeatureLevel::SM5 ) || !PendingState.bFramebufferSetupInvalid);
 
 	if (bClearColor)
 	{
@@ -3036,17 +3149,13 @@ void FOpenGLDynamicRHI::RHIDispatchComputeShader(uint32 ThreadGroupCountX, uint3
 		FOpenGLContextState& ContextState = GetContextStateForCurrentContext();
 
 		GPUProfilingData.RegisterGPUWork(1);		
-
+		CommitComputeResourceTables(ComputeShader);
 		SetupTexturesForDraw(ContextState, ComputeShader, FOpenGL::GetMaxComputeTextureImageUnits());
-
 		SetupUAVsForDraw(ContextState, ComputeShader, OGL_MAX_COMPUTE_STAGE_UAV_UNITS);
-	
 		CommitComputeShaderConstants(ComputeShader);
 	
 		FOpenGL::MemoryBarrier(GL_ALL_BARRIER_BITS);
-	
 		FOpenGL::DispatchCompute(ThreadGroupCountX, ThreadGroupCountY, ThreadGroupCountZ);
-	
 		FOpenGL::MemoryBarrier(GL_ALL_BARRIER_BITS);
 	}
 	else
@@ -3070,7 +3179,7 @@ void FOpenGLDynamicRHI::RHIDispatchIndirectComputeShader(FVertexBufferRHIParamRe
 
 		FOpenGLContextState& ContextState = GetContextStateForCurrentContext();
 
-		GPUProfilingData.RegisterGPUWork(1);		
+		GPUProfilingData.RegisterGPUWork(1);
 
 		SetupTexturesForDraw(ContextState, ComputeShader, FOpenGL::GetMaxComputeTextureImageUnits());
 
@@ -3101,6 +3210,11 @@ void FOpenGLDynamicRHI::RHISetMultipleViewports(uint32 Count, const FViewportBou
 	UE_LOG(LogRHI, Fatal,TEXT("OpenGL Render path does not support multiple Viewports!"));
 }
 
+void FOpenGLDynamicRHI::RHIExecuteCommandList(FRHICommandList*)
+{
+	check(0);
+}
+
 void FOpenGLDynamicRHI::RHIEnableDepthBoundsTest(bool bEnable,float MinDepth,float MaxDepth)
 {
 	if (FOpenGL::SupportsDepthBoundsTest())
diff --git a/Engine/Source/Runtime/OpenGLDrv/Private/OpenGLDevice.cpp b/Engine/Source/Runtime/OpenGLDrv/Private/OpenGLDevice.cpp
--- a/Engine/Source/Runtime/OpenGLDrv/Private/OpenGLDevice.cpp
+++ b/Engine/Source/Runtime/OpenGLDrv/Private/OpenGLDevice.cpp
@@ -65,10 +65,11 @@ void OnPixelBufferDeletion( GLuint PixelBufferResource )
 	PrivateOpenGLDevicePtr->OnPixelBufferDeletion( PixelBufferResource );
 }
 
-void OnUniformBufferDeletion( GLuint UniformBufferResource, uint32 AllocatedSize, bool bStreamDraw, uint32 Offset, uint8* Pointer )
+//gilmerge lose more of these args?
+void OnUniformBufferDeletion( GLuint UniformBufferResource, uint32 AllocatedSize, bool bStreamDraw, uint32 , uint8*  )
 {
 	check(PrivateOpenGLDevicePtr);
-	PrivateOpenGLDevicePtr->OnUniformBufferDeletion( UniformBufferResource, AllocatedSize, bStreamDraw, Offset, Pointer );
+	PrivateOpenGLDevicePtr->OnUniformBufferDeletion( UniformBufferResource, AllocatedSize, bStreamDraw );
 }
 
 void CachedBindArrayBuffer( GLuint Buffer )
@@ -145,10 +146,24 @@ void FOpenGLDynamicRHI::RHIEndFrame()
 
 void FOpenGLDynamicRHI::RHIBeginScene()
 {
+	// Increment the frame counter. INDEX_NONE is a special value meaning "uninitialized", so if
+	// we hit it just wrap around to zero.
+	SceneFrameCounter++;
+	if (SceneFrameCounter == INDEX_NONE)
+	{
+		SceneFrameCounter++;
+	}
+
+	static auto* ResourceTableCachingCvar = IConsoleManager::Get().FindTConsoleVariableDataInt(TEXT("rhi.ResourceTableCaching"));
+	if (ResourceTableCachingCvar == NULL || ResourceTableCachingCvar->GetValueOnAnyThread() == 1)
+	{
+		ResourceTableFrameCounter = SceneFrameCounter;
+	}
 }
 
 void FOpenGLDynamicRHI::RHIEndScene()
 {
+	ResourceTableFrameCounter = INDEX_NONE;
 }
 
 bool GDisableOpenGLDebugOutput = false;
@@ -787,7 +802,9 @@ static void InitRHICapabilitiesForGL()
 }
 
 FOpenGLDynamicRHI::FOpenGLDynamicRHI()
-:	bRevertToSharedContextAfterDrawingViewport(false)
+:	SceneFrameCounter(0)
+,	ResourceTableFrameCounter(INDEX_NONE)
+,	bRevertToSharedContextAfterDrawingViewport(false)
 ,	bIsRenderingContextAcquired(false)
 ,	PlatformDevice(NULL)
 ,	GPUProfilingData(this)
diff --git a/Engine/Source/Runtime/OpenGLDrv/Private/OpenGLDrv.cpp b/Engine/Source/Runtime/OpenGLDrv/Private/OpenGLDrv.cpp
--- a/Engine/Source/Runtime/OpenGLDrv/Private/OpenGLDrv.cpp
+++ b/Engine/Source/Runtime/OpenGLDrv/Private/OpenGLDrv.cpp
@@ -8,6 +8,8 @@
 
 IMPLEMENT_MODULE(FOpenGLDynamicRHIModule, OpenGLDrv);
 
+#include "ShaderParameterUtils.h"
+#include "RHIStaticStates.h"
 #include "OneColorShader.h"
 
 #if !UE_BUILD_SHIPPING
diff --git a/Engine/Source/Runtime/OpenGLDrv/Private/OpenGLShaders.cpp b/Engine/Source/Runtime/OpenGLDrv/Private/OpenGLShaders.cpp
--- a/Engine/Source/Runtime/OpenGLDrv/Private/OpenGLShaders.cpp
+++ b/Engine/Source/Runtime/OpenGLDrv/Private/OpenGLShaders.cpp
@@ -328,7 +328,7 @@ ShaderType* CompileOpenGLShader(const TArray<uint8>& Code)
 			strcpy(VersionString, "#version 100");
 			strcat(VersionString, "\n");
 		}
-
+		
 		//remove "#version 100" line,  if found in existing GlslCode. 
 		ANSICHAR* VersionCodePointer = const_cast<ANSICHAR*>(GlslCode);
 		VersionCodePointer = strstr(VersionCodePointer, "#version 100");
@@ -433,7 +433,6 @@ ShaderType* CompileOpenGLShader(const TArray<uint8>& Code)
 	Shader = new ShaderType();
 	Shader->Resource = Resource;
 	Shader->Bindings = Header.Bindings;
-	Shader->BoundUniformBuffers.AddZeroed( Header.Bindings.NumUniformBuffers );
 	Shader->UniformBuffersCopyInfo = Header.UniformBuffersCopyInfo;
 
 #if DEBUG_GL_SHADERS
@@ -517,18 +516,18 @@ static void MarkShaderParameterCachesDirty(FOpenGLShaderParameterCache* ShaderPa
 	}
 }
 
-void FOpenGLDynamicRHI::BindUniformBufferBase(FOpenGLContextState& ContextState, int32 NumUniformBuffers, const TArray< TRefCountPtr<FRHIUniformBuffer> >& BoundUniformBuffers, uint32 FirstUniformBuffer, bool ForceUpdate)
+void FOpenGLDynamicRHI::BindUniformBufferBase(FOpenGLContextState& ContextState, int32 NumUniformBuffers, FUniformBufferRHIRef* BoundUniformBuffers, uint32 FirstUniformBuffer, bool ForceUpdate)
 {
 	check(IsInRenderingThread());
 	for (int32 BufferIndex = 0; BufferIndex < NumUniformBuffers; ++BufferIndex)
 	{
-		GLuint Buffer;
+		GLuint Buffer = 0;
 		uint32 Offset = 0;
 		uint32 Size = ZERO_FILLED_DUMMY_UNIFORM_BUFFER_SIZE;
 		int32 BindIndex = FirstUniformBuffer + BufferIndex;
-		if (BoundUniformBuffers[BufferIndex])
+		if (IsValidRef(BoundUniformBuffers[BufferIndex]))
 		{
-			FRHIUniformBuffer* UB = BoundUniformBuffers[BufferIndex];
+			FRHIUniformBuffer* UB = BoundUniformBuffers[BufferIndex].GetReference();
 			Buffer = ((FOpenGLUniformBuffer*)UB)->Resource;
 			Size = ((FOpenGLUniformBuffer*)UB)->GetSize();
 #if SUBALLOCATED_CONSTANT_BUFFER
@@ -552,7 +551,7 @@ void FOpenGLDynamicRHI::BindUniformBufferBase(FOpenGLContextState& ContextState,
 			Buffer = PendingState.ZeroFilledDummyUniformBuffer;
 		}
 
-		if (ForceUpdate || ContextState.UniformBuffers[BindIndex] != Buffer || ContextState.UniformBufferOffsets[BindIndex] != Offset )
+		if (ForceUpdate || (Buffer != 0 && ContextState.UniformBuffers[BindIndex] != Buffer)|| ContextState.UniformBufferOffsets[BindIndex] != Offset)
 		{
 			FOpenGL::BindBufferRange(GL_UNIFORM_BUFFER, BindIndex, Buffer, Offset, Size);
 			ContextState.UniformBuffers[BindIndex] = Buffer;
@@ -965,11 +964,13 @@ static void VerifyUniformLayout(const TCHAR* UniformName, const UniformData& GLS
 				FString BaseTypeName;
 				switch(Member.GetBaseType())
 				{
-					case UBMT_STRUCT:  BaseTypeName = TEXT("struct"); break;
+					case UBMT_STRUCT:  BaseTypeName = TEXT("struct");  break;
 					case UBMT_BOOL:    BaseTypeName = TEXT("bool"); break;
 					case UBMT_INT32:   BaseTypeName = TEXT("int"); break;
 					case UBMT_UINT32:  BaseTypeName = TEXT("uint"); break;
 					case UBMT_FLOAT32: BaseTypeName = TEXT("float"); break;
+					case UBMT_TEXTURE: BaseTypeName = TEXT("texture"); break;
+					case UBMT_SAMPLER: BaseTypeName = TEXT("sampler"); break;
 					default:           UE_LOG(LogShaders, Fatal,TEXT("Unrecognized uniform buffer struct member base type."));
 				};
 #if ENABLE_UNIFORM_BUFFER_LAYOUT_DUMP
@@ -982,7 +983,7 @@ static void VerifyUniformLayout(const TCHAR* UniformName, const UniformData& GLS
 					Member.GetNumElements()
 					);
 #endif // #if ENABLE_UNIFORM_BUFFER_LAYOUT_DUMP
-				FString CompositeName = FString(StructIt->GetShaderVariableName()) + TEXT(".") + Member.GetName();
+				FString CompositeName = FString(StructIt->GetShaderVariableName()) + TEXT("_") + Member.GetName();
 
 				// GLSL returns array members with a "[0]" suffix
 				if(Member.GetNumElements())
@@ -1017,7 +1018,7 @@ static void VerifyUniformLayout(const TCHAR* UniformName, const UniformData& GLS
 	const UniformData* FoundUniform = Uniforms.Find(RequestedUniformName);
 
 	// MaterialTemplate uniform buffer does not have an entry in the FUniformBufferStructs list, so skipping it here
-	if(!(RequestedUniformName.StartsWith("Material.") || RequestedUniformName.StartsWith("MaterialCollection")))
+	if(!(RequestedUniformName.StartsWith("Material_") || RequestedUniformName.StartsWith("MaterialCollection")))
 	{
 		if(!FoundUniform || (*FoundUniform != GLSLUniform))
 		{
@@ -1577,51 +1578,58 @@ void FOpenGLDynamicRHI::BindPendingShaderState( FOpenGLContextState& ContextStat
 
 	if (!IsES2Platform(GRHIShaderPlatform))
 	{
-		const int32 VertexShaderNumUniformBuffers   = PendingState.BoundShaderState->VertexShader   ? PendingState.BoundShaderState->VertexShader->Bindings.NumUniformBuffers : 0;
-		const int32 PixelShaderNumUniformBuffers    = PendingState.BoundShaderState->PixelShader    ? PendingState.BoundShaderState->PixelShader->Bindings.NumUniformBuffers : 0;
-		const int32 GeometryShaderNumUniformBuffers = PendingState.BoundShaderState->GeometryShader ? PendingState.BoundShaderState->GeometryShader->Bindings.NumUniformBuffers : 0;
-		const int32 HullShaderNumUniformBuffers     = PendingState.BoundShaderState->HullShader     ? PendingState.BoundShaderState->HullShader->Bindings.NumUniformBuffers : 0;
-		const int32 DomainShaderNumUniformBuffers   = PendingState.BoundShaderState->DomainShader   ? PendingState.BoundShaderState->DomainShader->Bindings.NumUniformBuffers : 0;
-
-		BindUniformBufferBase(ContextState, VertexShaderNumUniformBuffers,
-			PendingState.BoundShaderState->VertexShader->BoundUniformBuffers, OGL_FIRST_UNIFORM_BUFFER, ForceUniformBindingUpdate);
+		int32 NextUniformBufferIndex = OGL_FIRST_UNIFORM_BUFFER;
 
-		BindUniformBufferBase(ContextState, PixelShaderNumUniformBuffers,
-			PendingState.BoundShaderState->PixelShader->BoundUniformBuffers,
-			OGL_FIRST_UNIFORM_BUFFER + VertexShaderNumUniformBuffers,
+		int32 NumVertexUniformBuffers = PendingState.BoundShaderState->VertexShader->Bindings.NumUniformBuffers;
+		BindUniformBufferBase(
+			ContextState,
+			NumVertexUniformBuffers,
+			PendingState.BoundUniformBuffers[SF_Vertex],
+			NextUniformBufferIndex,
 			ForceUniformBindingUpdate);
+		NextUniformBufferIndex += NumVertexUniformBuffers;
+
+		int32 NumPixelUniformBuffers = PendingState.BoundShaderState->PixelShader->Bindings.NumUniformBuffers;
+		BindUniformBufferBase(
+			ContextState,
+			NumPixelUniformBuffers,
+			PendingState.BoundUniformBuffers[SF_Pixel],
+			NextUniformBufferIndex,
+			ForceUniformBindingUpdate);
+		NextUniformBufferIndex += NumPixelUniformBuffers;
 
 		if (PendingState.BoundShaderState->GeometryShader)
 		{
-			BindUniformBufferBase(ContextState,
-				GeometryShaderNumUniformBuffers,
-				PendingState.BoundShaderState->GeometryShader->BoundUniformBuffers,
-				OGL_FIRST_UNIFORM_BUFFER + VertexShaderNumUniformBuffers +
-				PixelShaderNumUniformBuffers,
+			int32 NumGeometryUniformBuffers = PendingState.BoundShaderState->GeometryShader->Bindings.NumUniformBuffers;
+			BindUniformBufferBase(
+				ContextState,
+				NumGeometryUniformBuffers,
+				PendingState.BoundUniformBuffers[SF_Geometry],
+				NextUniformBufferIndex,
 				ForceUniformBindingUpdate);
+			NextUniformBufferIndex += NumGeometryUniformBuffers;
 		}
 
 		if (PendingState.BoundShaderState->HullShader)
 		{
+			int32 NumHullUniformBuffers = PendingState.BoundShaderState->HullShader->Bindings.NumUniformBuffers;
 			BindUniformBufferBase(ContextState,
-				HullShaderNumUniformBuffers,
-				PendingState.BoundShaderState->HullShader->BoundUniformBuffers,
-				OGL_FIRST_UNIFORM_BUFFER + VertexShaderNumUniformBuffers +
-				PixelShaderNumUniformBuffers +
-				GeometryShaderNumUniformBuffers,
+				NumHullUniformBuffers,
+				PendingState.BoundUniformBuffers[SF_Hull],
+				NextUniformBufferIndex,
 				ForceUniformBindingUpdate);
+			NextUniformBufferIndex += NumHullUniformBuffers;
 		}
 
 		if (PendingState.BoundShaderState->DomainShader)
 		{
+			int32 NumDomainUniformBuffers = PendingState.BoundShaderState->DomainShader->Bindings.NumUniformBuffers;
 			BindUniformBufferBase(ContextState,
-				DomainShaderNumUniformBuffers,
-				PendingState.BoundShaderState->DomainShader->BoundUniformBuffers,
-				OGL_FIRST_UNIFORM_BUFFER + VertexShaderNumUniformBuffers +
-				PixelShaderNumUniformBuffers + 
-				GeometryShaderNumUniformBuffers +
-				HullShaderNumUniformBuffers,
+				NumDomainUniformBuffers,
+				PendingState.BoundUniformBuffers[SF_Domain],
+				NextUniformBufferIndex,
 				ForceUniformBindingUpdate);
+			NextUniformBufferIndex += NumDomainUniformBuffers;
 		}
 	}
 }
@@ -1710,10 +1718,12 @@ void FOpenGLDynamicRHI::BindPendingComputeShaderState(FOpenGLContextState& Conte
 
 	if (!IsES2Platform(GRHIShaderPlatform))
 	{
-		BindUniformBufferBase(ContextState,
+		BindUniformBufferBase(
+			ContextState,
 			ComputeShader->Bindings.NumUniformBuffers,
-			ComputeShader->BoundUniformBuffers,
-			OGL_FIRST_UNIFORM_BUFFER, ForceUniformBindingUpdate);
+			PendingState.BoundUniformBuffers[SF_Compute],
+			OGL_FIRST_UNIFORM_BUFFER,
+			ForceUniformBindingUpdate);
 	}
 }
 
@@ -1862,7 +1872,7 @@ void FOpenGLShaderParameterCache::CommitPackedGlobals(const FOpenGLLinkedProgram
 	}
 }
 
-void FOpenGLShaderParameterCache::CommitPackedUniformBuffers(FOpenGLLinkedProgram* LinkedProgram, int32 Stage, const TArray< TRefCountPtr<FRHIUniformBuffer> >& RHIUniformBuffers, const TArray<FOpenGLUniformBufferCopyInfo>& UniformBuffersCopyInfo)
+void FOpenGLShaderParameterCache::CommitPackedUniformBuffers(FOpenGLLinkedProgram* LinkedProgram, int32 Stage, FUniformBufferRHIRef* RHIUniformBuffers, const TArray<FOpenGLUniformBufferCopyInfo>& UniformBuffersCopyInfo)
 {
 	SCOPE_CYCLE_COUNTER(STAT_OpenGLConstantBufferUpdateTime);
 	VERIFY_GL_SCOPE();
@@ -1870,17 +1880,16 @@ void FOpenGLShaderParameterCache::CommitPackedUniformBuffers(FOpenGLLinkedProgra
 	// Uniform Buffers are split into precision/type; the list of RHI UBs is traversed and if a new one was set, its
 	// contents are copied per precision/type into corresponding scratch buffers which are then uploaded to the program
 	const FOpenGLShaderBindings& Bindings = LinkedProgram->Config.Shaders[Stage].Bindings;
-	check(Bindings.NumUniformBuffers == RHIUniformBuffers.Num());
+	check(Bindings.NumUniformBuffers <= FOpenGLRHIState::MAX_UNIFORM_BUFFERS_PER_SHADER_STAGE);
 
 	if (Bindings.bFlattenUB)
 	{
 		int32 LastInfoIndex = 0;
-		for (int32 BufferIndex = 0; BufferIndex < RHIUniformBuffers.Num(); ++BufferIndex)
+		for (int32 BufferIndex = 0; BufferIndex < Bindings.NumUniformBuffers; ++BufferIndex)
 		{
-			const FRHIUniformBuffer* RHIUniformBuffer = RHIUniformBuffers[BufferIndex];
-			check(RHIUniformBuffer);
-			FOpenGLEUniformBuffer* EmulatedUniformBuffer = (FOpenGLEUniformBuffer*)RHIUniformBuffer;
-			const uint32* RESTRICT SourceData = EmulatedUniformBuffer->Buffer->Data.GetTypedData();
+			const FOpenGLUniformBuffer* UniformBuffer = (FOpenGLUniformBuffer*)RHIUniformBuffers[BufferIndex].GetReference();
+			check(UniformBuffer);
+			const uint32* RESTRICT SourceData = UniformBuffer->EmulatedBufferData->Data.GetTypedData();
 			for (int32 InfoIndex = LastInfoIndex; InfoIndex < UniformBuffersCopyInfo.Num(); ++InfoIndex)
 			{
 				const FOpenGLUniformBufferCopyInfo& Info = UniformBuffersCopyInfo[InfoIndex];
@@ -1907,14 +1916,13 @@ void FOpenGLShaderParameterCache::CommitPackedUniformBuffers(FOpenGLLinkedProgra
 		const auto& PackedUniformBufferInfos = LinkedProgram->StagePackedUniformInfo[Stage].PackedUniformBufferInfos;
 		int32 LastCopyInfoIndex = 0;
 		auto& EmulatedUniformBufferSet = LinkedProgram->StagePackedUniformInfo[Stage].LastEmulatedUniformBufferSet;
-		for (int32 BufferIndex = 0; BufferIndex < RHIUniformBuffers.Num(); ++BufferIndex)
+		for (int32 BufferIndex = 0; BufferIndex < Bindings.NumUniformBuffers; ++BufferIndex)
 		{
-			const FRHIUniformBuffer* RHIUniformBuffer = RHIUniformBuffers[BufferIndex];
-			check(RHIUniformBuffer);
-			FOpenGLEUniformBuffer* EmulatedUniformBuffer = (FOpenGLEUniformBuffer*)RHIUniformBuffer;
-			if (EmulatedUniformBufferSet[BufferIndex] != EmulatedUniformBuffer->UniqueID)
+			const FOpenGLUniformBuffer* UniformBuffer = (FOpenGLUniformBuffer*)RHIUniformBuffers[BufferIndex].GetReference();
+			check(UniformBuffer);
+			if (EmulatedUniformBufferSet[BufferIndex] != UniformBuffer->UniqueID)
 			{
-				EmulatedUniformBufferSet[BufferIndex] = EmulatedUniformBuffer->UniqueID;
+				EmulatedUniformBufferSet[BufferIndex] = UniformBuffer->UniqueID;
 
 				// Go through the list of copy commands and perform the appropriate copy into the scratch buffer
 				for (int32 InfoIndex = LastCopyInfoIndex; InfoIndex < UniformBuffersCopyInfo.Num(); ++InfoIndex)
@@ -1922,7 +1930,7 @@ void FOpenGLShaderParameterCache::CommitPackedUniformBuffers(FOpenGLLinkedProgra
 					const FOpenGLUniformBufferCopyInfo& Info = UniformBuffersCopyInfo[InfoIndex];
 					if (Info.SourceUBIndex == BufferIndex)
 					{
-						const uint32* RESTRICT SourceData = EmulatedUniformBuffer->Buffer->Data.GetTypedData();
+						const uint32* RESTRICT SourceData = UniformBuffer->EmulatedBufferData->Data.GetTypedData();
 						SourceData += Info.SourceOffsetInFloats;
 						float* RESTRICT ScratchMem = (float*)PackedUniformsScratch[Info.DestUBTypeIndex];
 						ScratchMem += Info.DestOffsetInFloats;
diff --git a/Engine/Source/Runtime/OpenGLDrv/Private/OpenGLTexture.cpp b/Engine/Source/Runtime/OpenGLDrv/Private/OpenGLTexture.cpp
--- a/Engine/Source/Runtime/OpenGLDrv/Private/OpenGLTexture.cpp
+++ b/Engine/Source/Runtime/OpenGLDrv/Private/OpenGLTexture.cpp
@@ -1790,3 +1790,23 @@ void FOpenGLDynamicRHI::RHIVirtualTextureSetFirstMipInMemory(FTexture2DRHIParamR
 void FOpenGLDynamicRHI::RHIVirtualTextureSetFirstMipVisible(FTexture2DRHIParamRef TextureRHI, uint32 FirstMip)
 {
 }
+
+FTextureReferenceRHIRef FOpenGLDynamicRHI::RHICreateTextureReference()
+{
+	return new FOpenGLTextureReference();
+}
+
+void FOpenGLTextureReference::SetReferencedTexture(FRHITexture* InTexture)
+{
+	FRHITextureReference::SetReferencedTexture(InTexture);
+	TexturePtr = GetOpenGLTextureFromRHITexture(InTexture);
+}
+
+void FOpenGLDynamicRHI::RHIUpdateTextureReference(FTextureReferenceRHIParamRef TextureRefRHI, FTextureRHIParamRef NewTextureRHI)
+{
+	auto* TextureRef = (FOpenGLTextureReference*)TextureRefRHI;
+	if (TextureRef)
+	{
+		TextureRef->SetReferencedTexture(NewTextureRHI);
+	}
+}
