#UE4  #HTML
    #TTP 335966  Image not centered/cut off in "Fullscreen" mode on Tappy Chicken page.

diff --git a/Engine/Source/Runtime/Core/Private/HTML5/HTML5Application.cpp b/Engine/Source/Runtime/Core/Private/HTML5/HTML5Application.cpp
--- a/Engine/Source/Runtime/Core/Private/HTML5/HTML5Application.cpp
+++ b/Engine/Source/Runtime/Core/Private/HTML5/HTML5Application.cpp
@@ -50,12 +50,9 @@ FPlatformRect FHTML5Application::GetWorkArea( const FPlatformRect& CurrentWindow
 void FHTML5Application::GetDisplayMetrics( FDisplayMetrics& OutDisplayMetrics ) const
 {
 	OutDisplayMetrics.PrimaryDisplayWorkAreaRect = FHTML5Window::GetScreenRect();
-
-	OutDisplayMetrics.VirtualDisplayRect = OutDisplayMetrics.PrimaryDisplayWorkAreaRect;
-	// Total screen size of the primary monitor - or in HTML5's case- total size of the webpage. Make it arbitrary large. 
-	OutDisplayMetrics.PrimaryDisplayWidth = 4096;
-	OutDisplayMetrics.PrimaryDisplayHeight = 4096; 
-
+	OutDisplayMetrics.VirtualDisplayRect    =	OutDisplayMetrics.PrimaryDisplayWorkAreaRect;
+	OutDisplayMetrics.PrimaryDisplayWidth   =	OutDisplayMetrics.PrimaryDisplayWorkAreaRect.Right;
+	OutDisplayMetrics.PrimaryDisplayHeight  =	OutDisplayMetrics.PrimaryDisplayWorkAreaRect.Bottom; 
 }
 
 TSharedRef< FGenericWindow > FHTML5Application::MakeWindow()
diff --git a/Engine/Source/Runtime/Core/Private/HTML5/HTML5Window.cpp b/Engine/Source/Runtime/Core/Private/HTML5/HTML5Window.cpp
--- a/Engine/Source/Runtime/Core/Private/HTML5/HTML5Window.cpp
+++ b/Engine/Source/Runtime/Core/Private/HTML5/HTML5Window.cpp
@@ -25,14 +25,11 @@ FHTML5Window::FHTML5Window()
 
 bool FHTML5Window::GetFullScreenInfo( int32& X, int32& Y, int32& Width, int32& Height ) const
 {
-	// fix me. 
-
 	FPlatformRect ScreenRect = GetScreenRect();
 	X = ScreenRect.Left;
 	Y = ScreenRect.Top;
 	Width = ScreenRect.Right - ScreenRect.Left;
 	Height = ScreenRect.Bottom - ScreenRect.Top;
-
 	return true;
 }
 
@@ -80,7 +77,7 @@ EWindowMode::Type FHTML5Window::GetWindowMode() const
 	return FullScreen ? EWindowMode::Fullscreen : EWindowMode::Windowed;
 #else
 	return EWindowMode::Windowed; 
-#endif
+#endif 
 
 }
 
diff --git a/Engine/Source/Runtime/OpenGLDrv/Private/HTML5/HTML5OpenGL.cpp b/Engine/Source/Runtime/OpenGLDrv/Private/HTML5/HTML5OpenGL.cpp
--- a/Engine/Source/Runtime/OpenGLDrv/Private/HTML5/HTML5OpenGL.cpp
+++ b/Engine/Source/Runtime/OpenGLDrv/Private/HTML5/HTML5OpenGL.cpp
@@ -164,10 +164,10 @@ void PlatformDestroyOpenGLDevice(FPlatformOpenGLDevice* Device)
 
 FPlatformOpenGLContext* PlatformCreateOpenGLContext(FPlatformOpenGLDevice* Device, void* InWindowHandle)
 {
-	int Width = 768 , Height = 1024;
+	int Width = 800 , Height = 600;
 #if !PLATFORM_HTML5_WIN32
 	 int fs;
-	emscripten_get_canvas_size(&Width, &Height, &fs);
+     emscripten_get_canvas_size(&Width, &Height, &fs);
 #endif
 	Device->SharedContext->Context = SDL_SetVideoMode(Width,Height, 16, SDL_OPENGL| SDL_RESIZABLE | SDL_DOUBLEBUF);
 	return Device->SharedContext; 
@@ -217,8 +217,9 @@ void PlatformResizeGLContext( FPlatformOpenGLDevice* Device, FPlatformOpenGLCont
 	// we can't resize with the win32 emu right now -- we have no way of resizing the window.
 #if !PLATFORM_HTML5_WIN32
 	emscripten_set_canvas_size(SizeX,SizeY);
-#endif
+#else
 	Device->SharedContext->Context = SDL_SetVideoMode(SizeX,SizeY, 16, SDL_OPENGL | SDL_RESIZABLE | SDL_DOUBLEBUF );
+#endif
 
 	glViewport(0, 0, SizeX, SizeY);
 	//@todo-mobile Do we need to clear here?
diff --git a/Engine/Source/Runtime/SlateCore/Private/Widgets/SWindow.cpp b/Engine/Source/Runtime/SlateCore/Private/Widgets/SWindow.cpp
--- a/Engine/Source/Runtime/SlateCore/Private/Widgets/SWindow.cpp
+++ b/Engine/Source/Runtime/SlateCore/Private/Widgets/SWindow.cpp
@@ -685,6 +685,7 @@ void SWindow::Resize( FVector2D NewSize )
 			InitialDesiredSize = NewSize;
 		}
 	}
+	SetCachedSize(NewSize); 
 }
 
 FSlateRect SWindow::GetFullScreenInfo() const
