Fix building on Mac using distcc which seperates pre-processing and compilation.

diff --git a/Engine/Source/Runtime/OpenGLDrv/Public/OpenGLUtil.h b/Engine/Source/Runtime/OpenGLDrv/Public/OpenGLUtil.h
--- a/Engine/Source/Runtime/OpenGLDrv/Public/OpenGLUtil.h
+++ b/Engine/Source/Runtime/OpenGLDrv/Public/OpenGLUtil.h
@@ -14,7 +14,7 @@
 #define ENABLE_VERIFY_GL_TRACE 0
 
 /** Set to 1 to verify that the the engine side uniform buffer layout matches the driver side of the GLSL shader*/
-#define ENABLE_UNIFORM_BUFFER_LAYOUT_VERIFICATION (0 && UE_BUILD_DEBUG && OPENGL_GL3)
+#define ENABLE_UNIFORM_BUFFER_LAYOUT_VERIFICATION (0 & UE_BUILD_DEBUG & OPENGL_GL3)
 
 /** Set to 1 to additinally dump uniform buffer layout at shader link time, this assumes ENABLE_UNIFORM_BUFFER_LAYOUT_VERIFICATION == 1 */
 #define ENABLE_UNIFORM_BUFFER_LAYOUT_DUMP 0
