#android Chosing EGL config at startup is now more robust

diff --git a/Engine/Source/Runtime/OpenGLDrv/Private/Android/AndroidEGL.cpp b/Engine/Source/Runtime/OpenGLDrv/Private/Android/AndroidEGL.cpp
--- a/Engine/Source/Runtime/OpenGLDrv/Private/Android/AndroidEGL.cpp
+++ b/Engine/Source/Runtime/OpenGLDrv/Private/Android/AndroidEGL.cpp
@@ -16,7 +16,6 @@ DEFINE_LOG_CATEGORY(LogEGL);
 #define ENABLE_CONFIG_FILTER 1
 #define ENABLE_EGL_DEBUG 0
 
-const uint32_t MAX_EGL_CONFIGS = 20;
 const  int EGLMinRedBits		= 5;
 const  int EGLMinGreenBits		= 6;
 const  int EGLMinBlueBits		= 5;
@@ -277,8 +276,15 @@ void AndroidEGL::InitEGL()
 
 #if ENABLE_CONFIG_FILTER
 
-	EGLConfig EGLConfigList[MAX_EGL_CONFIGS];
-	if(!( result = eglChooseConfig(PImplData->eglDisplay, Attributes, EGLConfigList, MAX_EGL_CONFIGS,  &PImplData->eglNumConfigs)))
+	EGLConfig* EGLConfigList = NULL;
+	result = eglChooseConfig(PImplData->eglDisplay, Attributes, NULL, 0, &PImplData->eglNumConfigs);
+	if (result)
+	{
+		int NumConfigs = PImplData->eglNumConfigs;
+		EGLConfigList = new EGLConfig[NumConfigs];
+		result = eglChooseConfig(PImplData->eglDisplay, Attributes, EGLConfigList, NumConfigs, &PImplData->eglNumConfigs);
+	}
+	if (!result)
 	{
 		ResetInternal();
 	}
@@ -289,11 +295,11 @@ void AndroidEGL::InitEGL()
 
 	int ResultValue = 0 ;
 	bool haveConfig = false;
-	int score = INT_MAX;
-	for( uint32_t i = 0 ;i < PImplData->eglNumConfigs ; i++)
+	int64 score = LONG_MAX;
+	for (uint32_t i = 0; i < PImplData->eglNumConfigs; i++)
 	{
-		int currScore = 0;
-		int r, g, b, a, d, s, sb, sc,nvi;
+		int64 currScore = 0;
+		int r, g, b, a, d, s, sb, sc, nvi;
 		eglGetConfigAttrib(PImplData->eglDisplay, EGLConfigList[i], EGL_RED_SIZE, &ResultValue); r = ResultValue;
 		eglGetConfigAttrib(PImplData->eglDisplay, EGLConfigList[i], EGL_GREEN_SIZE, &ResultValue); g = ResultValue;
 		eglGetConfigAttrib(PImplData->eglDisplay, EGLConfigList[i], EGL_BLUE_SIZE, &ResultValue); b = ResultValue;
@@ -313,13 +319,16 @@ void AndroidEGL::InitEGL()
 
 		// Favor EGLConfigLists by RGB, then Depth, then Non-linear Depth, then Stencil, then Alpha
 		currScore = 0;
-		currScore += (int)abs(sb - PImplData->Parms.sampleBuffers) << 30;
-		currScore += abs(sc - PImplData->Parms.sampleSamples) << 27;
-		currScore += (abs(r - PImplData->Parms.redSize) + abs(g - PImplData->Parms.greenSize) + abs(b - PImplData->Parms.blueSize)) << 21;
-		currScore += abs(d - PImplData->Parms.depthSize)  << 16;
-		currScore += abs(1 - bNonLinearDepth)  << 15;
-		currScore += abs(s - PImplData->Parms.stencilSize) << 8;
-		currScore += abs(a - PImplData->Parms.alphaSize);
+		currScore |= ((int64)FPlatformMath::Min(FPlatformMath::Abs(sb - PImplData->Parms.sampleBuffers), 15)) << 29;
+		currScore |= ((int64)FPlatformMath::Min(FPlatformMath::Abs(sc - PImplData->Parms.sampleSamples), 31)) << 24;
+		currScore |= FPlatformMath::Min(
+						FPlatformMath::Abs(r - PImplData->Parms.redSize) +
+						FPlatformMath::Abs(g - PImplData->Parms.greenSize) +
+						FPlatformMath::Abs(b - PImplData->Parms.blueSize), 127) << 17;
+		currScore |= FPlatformMath::Min(FPlatformMath::Abs(d - PImplData->Parms.depthSize), 63) << 11;
+		currScore |= FPlatformMath::Min(FPlatformMath::Abs(1 - bNonLinearDepth), 1) << 10;
+		currScore |= FPlatformMath::Min(FPlatformMath::Abs(s - PImplData->Parms.stencilSize), 31) << 6;
+		currScore |= FPlatformMath::Min(FPlatformMath::Abs(a - PImplData->Parms.alphaSize), 31) << 0;
 
 #if ENABLE_EGL_DEBUG
 		LogConfigInfo(EGLConfigList[i]);
@@ -335,10 +344,10 @@ void AndroidEGL::InitEGL()
 		}
 	}
 	check(haveConfig);
-	
+	delete[] EGLConfigList;
 #else
 
-	EGLConfig EGLConfigList[MAX_EGL_CONFIGS];
+	EGLConfig EGLConfigList[1];
 	if(!( result = eglChooseConfig(PImplData->eglDisplay, Attributes, EGLConfigList, 1,  &PImplData->eglNumConfigs)))
 	{
 		ResetInternal();
