Make GL_TIMESTAMP emulation on the Mac opt-in with the -EnableMacGPUTimestamp command-line flag. The old implementation used to crash on AMD, the modified version crashes on Nvidia.

diff --git a/Engine/Source/Runtime/OpenGLDrv/Private/Mac/OpenGLMac.cpp b/Engine/Source/Runtime/OpenGLDrv/Private/Mac/OpenGLMac.cpp
--- a/Engine/Source/Runtime/OpenGLDrv/Private/Mac/OpenGLMac.cpp
+++ b/Engine/Source/Runtime/OpenGLDrv/Private/Mac/OpenGLMac.cpp
@@ -387,7 +387,7 @@ static TAutoConsoleVariable<int32> CVarMacUseFrameBufferSRGB(
 		);
 
 bool GIsRunningOnIntelCard = false; // @todo: remove once Apple fixes radr://16223045 Changes to the GL separate blend state aren't always respected on Intel cards
-static bool GIsEmulatingTimestamp = (!UE_BUILD_SHIPPING); // @todo: Crashes on AMD on window close
+static bool GIsEmulatingTimestamp = false; // @todo: Now crashing on Nvidia cards, but not on AMD...
 
 struct FPlatformOpenGLDevice
 {
@@ -425,6 +425,12 @@ struct FPlatformOpenGLDevice
 		glGenVertexArrays(1,&RenderingContext.VertexArrayObject);
 		glBindVertexArray(RenderingContext.VertexArrayObject);
 		
+		// Only use the timestamp emulation in a non-shipping build - end-users shouldn't care about this profiling feature.
+#if (!UE_BUILD_SHIPPING)
+		// Opt-in since it crashes Nvidia cards at the moment
+		GIsEmulatingTimestamp = FParse::Param(FCommandLine::Get(), TEXT("EnableMacGPUTimestamp"));
+#endif
+		
 		if(GIsEmulatingTimestamp)
 		{
 			intptr_t RenderPtr = (intptr_t)&RenderingContext;
@@ -1384,11 +1390,6 @@ void FMacOpenGL::ProcessExtensions(const FString& ExtensionsString)
 	ProcessQueryGLInt();
 	FOpenGL3::ProcessExtensions(ExtensionsString);
 	
-	// Only use the timestamp emulation in a non-shipping build - end-users shouldn't care about this profiling feature.
-#if (!UE_BUILD_SHIPPING)
-	GIsEmulatingTimestamp = !FParse::Param(FCommandLine::Get(), TEXT("DisableMacGPUTimestamp"));
-#endif
-	
 	if(GIsEmulatingTimestamp)
 	{
 		TimestampQueryBits = 64;
