fix issue with shader compiling on HTML5 when texture lod is not supported
#ue4
#html5

diff --git a/Engine/Source/Runtime/OpenGLDrv/Private/OpenGLShaders.cpp b/Engine/Source/Runtime/OpenGLDrv/Private/OpenGLShaders.cpp
--- a/Engine/Source/Runtime/OpenGLDrv/Private/OpenGLShaders.cpp
+++ b/Engine/Source/Runtime/OpenGLDrv/Private/OpenGLShaders.cpp
@@ -403,7 +403,8 @@ ShaderType* CompileOpenGLShader(const TArray<uint8>& Code)
 #endif
 		else if(!FOpenGL::SupportsShaderTextureLod() || !FOpenGL::SupportsShaderTextureCubeLod())
 		{
-			Prologue = "#define textureCubeLodEXT(a, b, c) textureCube(a, b) \n";
+			Prologue = 	"#define texture2DLodEXT(a, b, c) texture2D(a, b) \n"
+						"#define textureCubeLodEXT(a, b, c) textureCube(a, b) \n";
 		}
 #if !PLATFORM_ANDROID
 		else if(!FOpenGL::SupportsTextureCubeLodEXT())
