UE4 - Fix OpenGL ES2 npot wrap modes (TTP 327493)

diff --git a/Engine/Source/Runtime/OpenGLDrv/Private/OpenGLCommands.cpp b/Engine/Source/Runtime/OpenGLDrv/Private/OpenGLCommands.cpp
--- a/Engine/Source/Runtime/OpenGLDrv/Private/OpenGLCommands.cpp
+++ b/Engine/Source/Runtime/OpenGLDrv/Private/OpenGLCommands.cpp
@@ -529,19 +529,21 @@ inline void FOpenGLDynamicRHI::ApplyTextureStage(FOpenGLContextState& ContextSta
 			ContextState.ActiveTexture = TextureIndex;
 		}
 
-		if (FOpenGL::HasSamplerRestrictions() && bHasTexture)
+		GLint WrapS = SamplerState->Data.WrapS;
+		GLint WrapT = SamplerState->Data.WrapT;
+		if (!FOpenGL::SupportsTextureNPOT() && bHasTexture)
 		{
 			if (!TextureStage.Texture->IsPowerOfTwo())
 			{
 				bool bChanged = false;
-				if (SamplerState->Data.WrapS != GL_CLAMP_TO_EDGE)
+				if (WrapS != GL_CLAMP_TO_EDGE)
 				{
-					SamplerState->Data.WrapS = GL_CLAMP_TO_EDGE;
+					WrapS = GL_CLAMP_TO_EDGE;
 					bChanged = true;
 				}
-				if (SamplerState->Data.WrapT != GL_CLAMP_TO_EDGE)
+				if (WrapT != GL_CLAMP_TO_EDGE)
 				{
-					SamplerState->Data.WrapT = GL_CLAMP_TO_EDGE;
+					WrapT = GL_CLAMP_TO_EDGE;
 					bChanged = true;
 				}
 				if (bChanged)
@@ -560,8 +562,8 @@ inline void FOpenGLDynamicRHI::ApplyTextureStage(FOpenGLContextState& ContextSta
 		}
 
 		// Sets parameters of currently bound texture
-		FOpenGL::TexParameter(Target, GL_TEXTURE_WRAP_S, SamplerState->Data.WrapS);
-		FOpenGL::TexParameter(Target, GL_TEXTURE_WRAP_T, SamplerState->Data.WrapT);
+		FOpenGL::TexParameter(Target, GL_TEXTURE_WRAP_S, WrapS);
+		FOpenGL::TexParameter(Target, GL_TEXTURE_WRAP_T, WrapT);
 		if( FOpenGL::SupportsTexture3D() )
 		{
 			FOpenGL::TexParameter(Target, GL_TEXTURE_WRAP_R, SamplerState->Data.WrapR);
@@ -2761,26 +2763,26 @@ static inline void ClearCurrentDepthStencilWithCurrentScissor( int8 ClearType, f
 #if PLATFORM_MAC
 	switch (ClearType)
 	{
-        case CT_DepthStencil:	// Clear depth and stencil separately to avoid an AMD Dx00 bug which causes depth to clear, but not stencil.
-                                // Especially irritatingly this bug will not manifest when stepping though the program with GL Profiler.
-                                // This was a bug found by me during dev. on Tropico 3 circa. Q4 2011 in the ATi Mac 2xx0 & 4xx0 drivers.
-                                // It was never fixed & has re-emerged in the AMD Mac FirePro Dx00 drivers.
+		case CT_DepthStencil:	// Clear depth and stencil separately to avoid an AMD Dx00 bug which causes depth to clear, but not stencil.
+								// Especially irritatingly this bug will not manifest when stepping though the program with GL Profiler.
+								// This was a bug found by me during dev. on Tropico 3 circa. Q4 2011 in the ATi Mac 2xx0 & 4xx0 drivers.
+								// It was never fixed & has re-emerged in the AMD Mac FirePro Dx00 drivers.
 								// Also, on NVIDIA depth must be cleared first.
-        case CT_Depth:	// Clear depth only
-            FOpenGL::ClearBufferfv(GL_DEPTH, 0, &Depth);
-
-            // If not also clearing depth break
-            if(!(ClearType & CT_Stencil))
-            {
-                break;
-            }
-            // Otherwise fall through to perform a separate stencil clear.
-        case CT_Stencil:	// Clear stencil only
-            FOpenGL::ClearBufferiv(GL_STENCIL, 0, (const GLint*)&Stencil);
-            break;
-
-        default:
-            break;	// impossible anyway
+		case CT_Depth:	// Clear depth only
+			FOpenGL::ClearBufferfv(GL_DEPTH, 0, &Depth);
+
+			// If not also clearing depth break
+			if(!(ClearType & CT_Stencil))
+			{
+				break;
+			}
+			// Otherwise fall through to perform a separate stencil clear.
+		case CT_Stencil:	// Clear stencil only
+			FOpenGL::ClearBufferiv(GL_STENCIL, 0, (const GLint*)&Stencil);
+			break;
+
+		default:
+			break;	// impossible anyway
 	}
 #else
 	switch (ClearType)
diff --git a/Engine/Source/Runtime/OpenGLDrv/Private/OpenGLES2.cpp b/Engine/Source/Runtime/OpenGLDrv/Private/OpenGLES2.cpp
--- a/Engine/Source/Runtime/OpenGLDrv/Private/OpenGLES2.cpp
+++ b/Engine/Source/Runtime/OpenGLDrv/Private/OpenGLES2.cpp
@@ -101,6 +101,9 @@ bool FOpenGLES2::bSupportsShaderTextureCubeLod = true;
 /** GL_APPLE_copy_texture_levels */
 bool FOpenGLES2::bSupportsCopyTextureLevels = false;
 
+/** GL_OES_texture_npot */
+bool FOpenGLES2::bSupportsTextureNPOT = true;
+
 /** GL_EXT_texture_storage */
 bool FOpenGLES2::bSupportsTextureStorageEXT = false;
 
@@ -174,6 +177,7 @@ void FOpenGLES2::ProcessExtensions( const FString& ExtensionsString )
 	bSupportsShaderTextureLod = ExtensionsString.Contains(TEXT("GL_EXT_shader_texture_lod"));
 	bSupportsTextureStorageEXT = ExtensionsString.Contains(TEXT("GL_EXT_texture_storage"));
 	bSupportsCopyTextureLevels = bSupportsTextureStorageEXT && ExtensionsString.Contains(TEXT("GL_APPLE_copy_texture_levels"));
+	bSupportsTextureNPOT = ExtensionsString.Contains(TEXT("GL_OES_texture_npot")) || ExtensionsString.Contains(TEXT("GL_ARB_texture_non_power_of_two"));
 
 	// Report shader precision
 	int Range[2];
diff --git a/Engine/Source/Runtime/OpenGLDrv/Public/OpenGL.h b/Engine/Source/Runtime/OpenGLDrv/Public/OpenGL.h
--- a/Engine/Source/Runtime/OpenGLDrv/Public/OpenGL.h
+++ b/Engine/Source/Runtime/OpenGLDrv/Public/OpenGL.h
@@ -133,7 +133,7 @@ public:
 	static FORCEINLINE bool SupportsVertexAttribBinding()				{ return false; }
 	static FORCEINLINE bool SupportsBufferStorage()						{ return false; }
 	static FORCEINLINE bool SupportsDepthBoundsTest()					{ return false; }
-	static FORCEINLINE bool HasSamplerRestrictions()					{ return false; }
+	static FORCEINLINE bool SupportsTextureNPOT()						{ return true; }
 	static FORCEINLINE bool HasHardwareHiddenSurfaceRemoval()			{ return false; }
 
 	static FORCEINLINE GLenum GetDepthFormat()							{ return GL_DEPTH_COMPONENT16; }
diff --git a/Engine/Source/Runtime/OpenGLDrv/Public/OpenGLES2.h b/Engine/Source/Runtime/OpenGLDrv/Public/OpenGLES2.h
--- a/Engine/Source/Runtime/OpenGLDrv/Public/OpenGLES2.h
+++ b/Engine/Source/Runtime/OpenGLDrv/Public/OpenGLES2.h
@@ -117,7 +117,7 @@ struct FOpenGLES2 : public FOpenGLBase
 	static FORCEINLINE bool SupportsFastBufferData()					{ return false; }
 
 	// ES 2 will not work with non-power of two textures with non-clamp mode
-	static FORCEINLINE bool HasSamplerRestrictions()					{ return true; }
+	static FORCEINLINE bool SupportsTextureNPOT()						{ return bSupportsTextureNPOT; }
 
 	// Optional
 	static FORCEINLINE void BeginQuery(GLenum QueryType, GLuint QueryId)
@@ -419,6 +419,9 @@ protected:
 	/** GL_APPLE_copy_texture_levels */
 	static bool bSupportsCopyTextureLevels;
 
+	/** GL_OES_texture_npot */
+	static bool bSupportsTextureNPOT;
+
 	/** GL_EXT_texture_storage */
 	static bool bSupportsTextureStorageEXT;
 
