Fixed missing text in sample games on Android devices without support for half float textures
#Android
#TTP 337053

diff --git a/Engine/Source/Runtime/OpenGLDrv/Private/OpenGLDevice.cpp b/Engine/Source/Runtime/OpenGLDrv/Private/OpenGLDevice.cpp
--- a/Engine/Source/Runtime/OpenGLDrv/Private/OpenGLDevice.cpp
+++ b/Engine/Source/Runtime/OpenGLDrv/Private/OpenGLDevice.cpp
@@ -723,17 +723,32 @@ static void InitRHICapabilitiesForGL()
 		SetupTextureFormat( PF_G8,					FOpenGLTextureFormat(GL_LUMINANCE,		GL_LUMINANCE,		GL_LUMINANCE,				GL_LUMINANCE	,							GL_LUMINANCE,			GL_UNSIGNED_BYTE,					false));
 		SetupTextureFormat( PF_A8,					FOpenGLTextureFormat(GL_ALPHA,			GL_ALPHA,			GL_ALPHA,					GL_ALPHA,									GL_ALPHA,				GL_UNSIGNED_BYTE,					false));
 	#endif
+
+#if PLATFORM_ANDROID
+	bool bNeedFloatRGBAtoRGBA8 = true;
+#endif
 	if (GSupportsRenderTargetFormat_PF_FloatRGBA)
 	{
 		if (FOpenGL::SupportsTextureHalfFloat())
 		{
 #if PLATFORM_ANDROID && !PLATFORM_ANDROIDGL4
 			SetupTextureFormat( PF_FloatRGBA,		FOpenGLTextureFormat(GL_RGBA,			GL_RGBA,			GL_RGBA16F_EXT,				GL_RGBA16F_EXT,								GL_RGBA,				GL_HALF_FLOAT_OES,					false));
+			bNeedFloatRGBAtoRGBA8 = false;
 #else
 			SetupTextureFormat( PF_FloatRGBA,		FOpenGLTextureFormat(GL_RGBA,							GL_RGBA,									GL_RGBA,				GL_HALF_FLOAT_OES,					false));
+#endif
+#if PLATFORM_ANDROID
+			bNeedFloatRGBAtoRGBA8 = false;
 #endif
 		}
 	}
+#if PLATFORM_ANDROID
+	if (bNeedFloatRGBAtoRGBA8)
+	{
+		SetupTextureFormat( PF_FloatRGBA,			FOpenGLTextureFormat(GL_RGBA,			GL_RGBA,			GL_RGBA,					GL_UNSIGNED_BYTE,					false));
+	}
+#endif
+
 	if (FOpenGL::SupportsPackedDepthStencil())
 	{
 		SetupTextureFormat( PF_DepthStencil,	FOpenGLTextureFormat(GL_DEPTH_STENCIL_OES,				GL_NONE,									GL_DEPTH_STENCIL_OES,	GL_UNSIGNED_INT_24_8_OES,			false));
diff --git a/Engine/Source/Runtime/OpenGLDrv/Private/OpenGLTexture.cpp b/Engine/Source/Runtime/OpenGLDrv/Private/OpenGLTexture.cpp
--- a/Engine/Source/Runtime/OpenGLDrv/Private/OpenGLTexture.cpp
+++ b/Engine/Source/Runtime/OpenGLDrv/Private/OpenGLTexture.cpp
@@ -698,6 +698,57 @@ void* TOpenGLTexture<RHIResourceType>::Lock(uint32 MipIndex,uint32 ArrayIndex,ER
 	return result;
 }
 
+// Copied from OpenGLDebugFrameDump.
+inline uint32 HalfFloatToFloatInteger(uint16 HalfFloat)
+{
+	register uint32 Sign = (HalfFloat >> 15) & 0x00000001;
+	register uint32 Exponent = (HalfFloat >> 10) & 0x0000001f;
+	register uint32 Mantiss = HalfFloat & 0x000003ff;
+
+	if (Exponent == 0)
+	{
+		if (Mantiss == 0) // Plus or minus zero
+		{
+			return Sign << 31;
+		}
+		else // Denormalized number -- renormalize it
+		{
+			while ((Mantiss & 0x00000400) == 0)
+			{
+				Mantiss <<= 1;
+				Exponent -= 1;
+			}
+
+			Exponent += 1;
+			Mantiss &= ~0x00000400;
+		}
+	}
+	else if (Exponent == 31)
+	{
+		if (Mantiss == 0) // Inf
+			return (Sign << 31) | 0x7f800000;
+		else // NaN
+			return (Sign << 31) | 0x7f800000 | (Mantiss << 13);
+	}
+
+	Exponent = Exponent + (127 - 15);
+	Mantiss = Mantiss << 13;
+
+	return (Sign << 31) | (Exponent << 23) | Mantiss;
+}
+
+inline float HalfFloatToFloat(uint16 HalfFloat)
+{
+	union
+	{
+		float F;
+		uint32 I;
+	} Convert;
+
+	Convert.I = HalfFloatToFloatInteger(HalfFloat);
+	return Convert.F;
+}
+
 template<typename RHIResourceType>
 void TOpenGLTexture<RHIResourceType>::Unlock(uint32 MipIndex,uint32 ArrayIndex)
 {
@@ -712,6 +763,70 @@ void TOpenGLTexture<RHIResourceType>::Unlock(uint32 MipIndex,uint32 ArrayIndex)
 	const FOpenGLTextureFormat& GLFormat = GOpenGLTextureFormats[this->GetFormat()];
 	const bool bSRGB = (this->GetFlags() & TexCreate_SRGB) != 0;
 
+#if PLATFORM_ANDROID
+	// check for FloatRGBA to RGBA8 conversion needed
+	if (this->GetFormat() == PF_FloatRGBA && GLFormat.Type == GL_UNSIGNED_BYTE)
+	{
+		UE_LOG(LogRHI, Warning, TEXT("Converting texture from PF_FloatRGBA to RGBA8!  Only supported for limited case of distance field fonts!"));
+
+		// Code path for non-PBO: and always uncompressed! and not cubemap
+		check(!bCubemap);
+
+		// Volume/array textures are currently only supported if PixelBufferObjects are also supported.
+		check(this->GetSizeZ() == 0);
+
+		// Use a texture stage that's not likely to be used for draws, to avoid waiting
+		FOpenGLContextState& ContextState = OpenGLRHI->GetContextStateForCurrentContext();
+		OpenGLRHI->CachedSetupTextureStage(ContextState, FOpenGL::GetMaxCombinedTextureImageUnits() - 1, Target, Resource, -1, this->GetNumMips());
+
+		CachedBindPixelUnpackBuffer(0);
+
+		// get the source data and size
+		uint16* floatData = (uint16*)PixelBuffer->GetLockedBuffer();
+		int32 texWidth = FMath::Max<uint32>(1, (this->GetSizeX() >> MipIndex));
+		int32 texHeight = FMath::Max<uint32>(1, (this->GetSizeY() >> MipIndex));
+		
+		// always RGBA8 so 4 bytes / pixel
+		int nValues = texWidth * texHeight * 4;
+		uint8* rgbaData = (uint8*)FMemory::Malloc(nValues);
+
+		// convert to GL_BYTE
+		uint8* outPtr = rgbaData;
+		while (nValues--)
+		{
+			*outPtr++ = (uint8)(HalfFloatToFloat(*floatData++) * 255.0f);
+		}
+
+		// All construction paths should have called TexStorage2D or TexImage2D. So we will
+		// always call TexSubImage2D.
+		check(GetAllocatedStorageForMip(MipIndex, ArrayIndex) == true);
+		glTexSubImage2D(
+			Target,
+			MipIndex,
+			0,
+			0,
+			texWidth,
+			texHeight,
+			GLFormat.Format,
+			GLFormat.Type,
+			rgbaData);
+
+		// free temporary conversion buffer
+		FMemory::Free(rgbaData);
+
+		// Unlock "PixelBuffer" and free the temp memory after the texture upload.
+		PixelBuffer->Unlock();
+
+		// No need to restore texture stage; leave it like this,
+		// and the next draw will take care of cleaning it up; or
+		// next operation that needs the stage will switch something else in on it.
+
+		CachedBindPixelUnpackBuffer(0);
+
+		return;
+	}
+#endif
+
 	if ( FOpenGL::SupportsPixelBuffers() )
 	{
 		// Code path for PBO per slice
