Fix for failing IOS build.
'register' storage class specifier is deprecated

diff --git a/Engine/Source/Runtime/OpenGLDrv/Private/OpenGLTexture.cpp b/Engine/Source/Runtime/OpenGLDrv/Private/OpenGLTexture.cpp
--- a/Engine/Source/Runtime/OpenGLDrv/Private/OpenGLTexture.cpp
+++ b/Engine/Source/Runtime/OpenGLDrv/Private/OpenGLTexture.cpp
@@ -701,9 +701,9 @@ void* TOpenGLTexture<RHIResourceType>::Lock(uint32 MipIndex,uint32 ArrayIndex,ER
 // Copied from OpenGLDebugFrameDump.
 inline uint32 HalfFloatToFloatInteger(uint16 HalfFloat)
 {
-	register uint32 Sign = (HalfFloat >> 15) & 0x00000001;
-	register uint32 Exponent = (HalfFloat >> 10) & 0x0000001f;
-	register uint32 Mantiss = HalfFloat & 0x000003ff;
+	uint32 Sign = (HalfFloat >> 15) & 0x00000001;
+	uint32 Exponent = (HalfFloat >> 10) & 0x0000001f;
+	uint32 Mantiss = HalfFloat & 0x000003ff;
 
 	if (Exponent == 0)
 	{
