Workaround for what appears to be an AMD driver bug on copy of mipmaps of compressed textures smaller than 4 pixels in any dimension.

diff --git a/Engine/Source/Runtime/OpenGLDrv/Private/OpenGLDrv.cpp b/Engine/Source/Runtime/OpenGLDrv/Private/OpenGLDrv.cpp
--- a/Engine/Source/Runtime/OpenGLDrv/Private/OpenGLDrv.cpp
+++ b/Engine/Source/Runtime/OpenGLDrv/Private/OpenGLDrv.cpp
@@ -449,6 +449,7 @@ bool  FOpenGLBase::bSupportsCopyImage = false;
 bool  FOpenGLBase::bSupportsSeamlessCubemap = false;
 bool  FOpenGLBase::bSupportsVolumeTextureRendering = false;
 bool  FOpenGLBase::bSupportsTextureFilterAnisotropic = false;
+bool  FOpenGLBase::bAmdWorkaround = false;
 
 void FOpenGLBase::ProcessQueryGLInt()
 {
@@ -497,6 +498,14 @@ void FOpenGLBase::ProcessExtensions( const FString& ExtensionsString )
 	bSupportsSeamlessCubemap = ExtensionsString.Contains(TEXT("GL_ARB_seamless_cube_map"));
 	
 	bSupportsTextureFilterAnisotropic = ExtensionsString.Contains(TEXT("GL_EXT_texture_filter_anisotropic"));
+
+	#if PLATFORM_WINDOWS
+		FString VendorName( ANSI_TO_TCHAR((const ANSICHAR*)glGetString(GL_VENDOR) ) );
+		if ( VendorName.Contains(TEXT("ATI ")) )
+		{
+			bAmdWorkaround = true;
+		}
+	#endif
 }
 
 void InitDefaultGLContextState(void)
diff --git a/Engine/Source/Runtime/OpenGLDrv/Private/OpenGLTexture.cpp b/Engine/Source/Runtime/OpenGLDrv/Private/OpenGLTexture.cpp
--- a/Engine/Source/Runtime/OpenGLDrv/Private/OpenGLTexture.cpp
+++ b/Engine/Source/Runtime/OpenGLDrv/Private/OpenGLTexture.cpp
@@ -1020,6 +1020,8 @@ void TOpenGLTexture<RHIResourceType>::CloneViaCopyImage( TOpenGLTexture<RHIResou
 			const uint32 MipSizeX = FMath::Max(this->GetSizeX() >> DstMipIndex,uint32(1));
 			const uint32 MipSizeY = FMath::Max(this->GetSizeY() >> DstMipIndex,uint32(1));
 			
+			if(FOpenGL::AmdWorkaround() && ((MipSizeX < 4) || (MipSizeY < 4))) break;
+
 			// copy the texture data
 			FOpenGL::CopyImageSubData( Src->Resource, Src->Target, SrcMipIndex, 0, 0, ArrayIndex,
 									  Resource, Target, DstMipIndex, 0, 0, ArrayIndex, MipSizeX, MipSizeY, 1);
diff --git a/Engine/Source/Runtime/OpenGLDrv/Private/Windows/OpenGLWindows.cpp b/Engine/Source/Runtime/OpenGLDrv/Private/Windows/OpenGLWindows.cpp
--- a/Engine/Source/Runtime/OpenGLDrv/Private/Windows/OpenGLWindows.cpp
+++ b/Engine/Source/Runtime/OpenGLDrv/Private/Windows/OpenGLWindows.cpp
@@ -746,15 +746,6 @@ bool PlatformInitOpenGL()
 			ENUM_GL_ENTRYPOINTS_DLL(CHECK_GL_ENTRYPOINTS);
 			ENUM_GL_ENTRYPOINTS(CHECK_GL_ENTRYPOINTS);
 			checkf(bFoundAllEntryPoints, TEXT("Failed to find all OpenGL entry points."));
-
-			FString VendorName( ANSI_TO_TCHAR((const ANSICHAR*)glGetString(GL_VENDOR) ) );
-			if ( VendorName.Contains(TEXT("ATI ")) )
-			{
-				// Workaround for AMD driver not handling GL_SRGB8_ALPHA8 in glTexStorage2D() properly (gets treated as non-sRGB)
-				glTexStorage1D = NULL;
-				glTexStorage2D = NULL;
-				glTexStorage3D = NULL;
-			}
 		}
 
 		// The dummy context can now be released.
diff --git a/Engine/Source/Runtime/OpenGLDrv/Public/OpenGL.h b/Engine/Source/Runtime/OpenGLDrv/Public/OpenGL.h
--- a/Engine/Source/Runtime/OpenGLDrv/Public/OpenGL.h
+++ b/Engine/Source/Runtime/OpenGLDrv/Public/OpenGL.h
@@ -135,6 +135,8 @@ public:
 	static FORCEINLINE bool SupportsDepthBoundsTest()					{ return false; }
 	static FORCEINLINE bool SupportsTextureNPOT()						{ return true; }
 	static FORCEINLINE bool HasHardwareHiddenSurfaceRemoval()			{ return false; }
+	static FORCEINLINE bool AmdWorkaround()								{ return false; }
+
 
 	static FORCEINLINE GLenum GetDepthFormat()							{ return GL_DEPTH_COMPONENT16; }
 
@@ -332,6 +334,10 @@ protected:
 	
 	/** GL_EXT_texture_filter_anisotropic Can we use anisotropic filtering? */
 	static bool bSupportsTextureFilterAnisotropic;
+
+	/** Workaround AMD driver issues. */
+	static bool bAmdWorkaround;
+
 };
 
 /** Unreal tokens that maps to different OpenGL tokens by platform. */
diff --git a/Engine/Source/Runtime/OpenGLDrv/Public/OpenGL3.h b/Engine/Source/Runtime/OpenGLDrv/Public/OpenGL3.h
--- a/Engine/Source/Runtime/OpenGLDrv/Public/OpenGL3.h
+++ b/Engine/Source/Runtime/OpenGLDrv/Public/OpenGL3.h
@@ -23,6 +23,8 @@ struct FOpenGL3 : public FOpenGLBase
 	static FORCEINLINE bool SupportsSeamlessCubeMap()			{ return bSupportsSeamlessCubemap; }
 	static FORCEINLINE bool SupportsVolumeTextureRendering()	{ return bSupportsVolumeTextureRendering; }
 	static FORCEINLINE bool SupportsGenerateMipmap()			{ return true; }
+	static FORCEINLINE bool AmdWorkaround()						{ return bAmdWorkaround; }
+
 
 	// Optional
 	static FORCEINLINE void QueryTimestampCounter(GLuint QueryID)
