Allow cubemaps with FloatRGBA to RGBA conversion with clamping for values over 1.0
#Android
#TTP 340016

diff --git a/Engine/Source/Runtime/OpenGLDrv/Private/OpenGLTexture.cpp b/Engine/Source/Runtime/OpenGLDrv/Private/OpenGLTexture.cpp
--- a/Engine/Source/Runtime/OpenGLDrv/Private/OpenGLTexture.cpp
+++ b/Engine/Source/Runtime/OpenGLDrv/Private/OpenGLTexture.cpp
@@ -767,11 +767,9 @@ void TOpenGLTexture<RHIResourceType>::Unlock(uint32 MipIndex,uint32 ArrayIndex)
 	// check for FloatRGBA to RGBA8 conversion needed
 	if (this->GetFormat() == PF_FloatRGBA && GLFormat.Type == GL_UNSIGNED_BYTE)
 	{
-		UE_LOG(LogRHI, Warning, TEXT("Converting texture from PF_FloatRGBA to RGBA8!  Only supported for limited case of distance field fonts!"));
-
-		// Code path for non-PBO: and always uncompressed! and not cubemap
-		check(!bCubemap);
+		UE_LOG(LogRHI, Warning, TEXT("Converting texture from PF_FloatRGBA to RGBA8!  Only supported for limited cases of 0.0 to 1.0 values (clamped)"));
 
+		// Code path for non-PBO: and always uncompressed!
 		// Volume/array textures are currently only supported if PixelBufferObjects are also supported.
 		check(this->GetSizeZ() == 0);
 
@@ -785,23 +783,24 @@ void TOpenGLTexture<RHIResourceType>::Unlock(uint32 MipIndex,uint32 ArrayIndex)
 		uint16* floatData = (uint16*)PixelBuffer->GetLockedBuffer();
 		int32 texWidth = FMath::Max<uint32>(1, (this->GetSizeX() >> MipIndex));
 		int32 texHeight = FMath::Max<uint32>(1, (this->GetSizeY() >> MipIndex));
-		
+
 		// always RGBA8 so 4 bytes / pixel
 		int nValues = texWidth * texHeight * 4;
 		uint8* rgbaData = (uint8*)FMemory::Malloc(nValues);
 
-		// convert to GL_BYTE
+		// convert to GL_BYTE (saturate)
 		uint8* outPtr = rgbaData;
 		while (nValues--)
 		{
-			*outPtr++ = (uint8)(HalfFloatToFloat(*floatData++) * 255.0f);
+			int32 pixelValue = (int32)(HalfFloatToFloat(*floatData++) * 255.0f);
+			*outPtr++ = (uint8)(pixelValue < 0 ? 0 : (pixelValue < 256 ? pixelValue : 255));
 		}
 
 		// All construction paths should have called TexStorage2D or TexImage2D. So we will
 		// always call TexSubImage2D.
 		check(GetAllocatedStorageForMip(MipIndex, ArrayIndex) == true);
 		glTexSubImage2D(
-			Target,
+			bCubemap ? GL_TEXTURE_CUBE_MAP_POSITIVE_X + ArrayIndex : Target,
 			MipIndex,
 			0,
 			0,
