Fixing build break on Mac / Linux / iOS

diff --git a/Engine/Source/Runtime/OpenGLDrv/Private/IOS/IOSOpenGL.cpp b/Engine/Source/Runtime/OpenGLDrv/Private/IOS/IOSOpenGL.cpp
--- a/Engine/Source/Runtime/OpenGLDrv/Private/IOS/IOSOpenGL.cpp
+++ b/Engine/Source/Runtime/OpenGLDrv/Private/IOS/IOSOpenGL.cpp
@@ -117,7 +117,7 @@ void* PlatformGetWindow(FPlatformOpenGLContext* Context, void** AddParam)
 
 bool PlatformBlitToViewport( FPlatformOpenGLDevice* Device, const FOpenGLViewport& Viewport, uint32 BackbufferSizeX, uint32 BackbufferSizeY, bool bPresent,bool bLockToVsync, int32 SyncInterval )
 {
-	FPlatformOpenGLContext* const Context = Viewport.OpenGLContext;
+	FPlatformOpenGLContext* const Context = Viewport.GetGLContext();
 
 	// @todo-mobile
 	check(Device->bSingleContext || Context == &Device->RenderingContext);
diff --git a/Engine/Source/Runtime/OpenGLDrv/Private/Linux/OpenGLLinux.cpp b/Engine/Source/Runtime/OpenGLDrv/Private/Linux/OpenGLLinux.cpp
--- a/Engine/Source/Runtime/OpenGLDrv/Private/Linux/OpenGLLinux.cpp
+++ b/Engine/Source/Runtime/OpenGLDrv/Private/Linux/OpenGLLinux.cpp
@@ -343,7 +343,7 @@ bool PlatformBlitToViewport(FPlatformOpenGLDevice* Device,
 							bool bLockToVsync,
 							int32 SyncInterval )
 {
-	FPlatformOpenGLContext* const Context = Viewport.OpenGLContext;
+	FPlatformOpenGLContext* const Context = Viewport.GetGLContext();
 
 	check( Context && Context->hWnd );
 
diff --git a/Engine/Source/Runtime/OpenGLDrv/Private/Mac/OpenGLMac.cpp b/Engine/Source/Runtime/OpenGLDrv/Private/Mac/OpenGLMac.cpp
--- a/Engine/Source/Runtime/OpenGLDrv/Private/Mac/OpenGLMac.cpp
+++ b/Engine/Source/Runtime/OpenGLDrv/Private/Mac/OpenGLMac.cpp
@@ -632,7 +632,7 @@ void* PlatformGetWindow(FPlatformOpenGLContext* Context, void** AddParam)
 
 bool PlatformBlitToViewport( FPlatformOpenGLDevice* Device, const FOpenGLViewport& Viewport, uint32 BackbufferSizeX, uint32 BackbufferSizeY, bool bPresent,bool bLockToVsync, int32 SyncInterval )
 {
-	FPlatformOpenGLContext* const Context = Viewport.OpenGLContext;
+	FPlatformOpenGLContext* const Context = Viewport.GetGLContext();
 
 	check(Context && Context->OpenGLView);
 
