Expose the amount of VRAM in the GPU to the OpenGLRHI on OS X.

diff --git a/Engine/Source/Runtime/OpenGLDrv/Private/Mac/OpenGLMac.cpp b/Engine/Source/Runtime/OpenGLDrv/Private/Mac/OpenGLMac.cpp
--- a/Engine/Source/Runtime/OpenGLDrv/Private/Mac/OpenGLMac.cpp
+++ b/Engine/Source/Runtime/OpenGLDrv/Private/Mac/OpenGLMac.cpp
@@ -1305,6 +1305,49 @@ PFNGLPOPGROUPMARKEREXTPROC FMacOpenGL::glPopGroupMarkerEXT = NULL;
 bool FMacOpenGL::bSupportsTessellationShader = false;
 PFNGLPATCHPARAMETERIPROC FMacOpenGL::glPatchParameteri = NULL;
 
+uint64 FMacOpenGL::GetVideoMemorySize()
+{
+	uint64 VideoMemory = 0;
+	NSOpenGLContext* NSContext = [NSOpenGLContext currentContext];
+	CGLContextObj Context = NSContext ? (CGLContextObj)[NSContext CGLContextObj] : NULL;
+	if(Context)
+	{
+		GLint VirtualScreen = [NSContext currentVirtualScreen];
+		GLint RendererID = 0;
+		GLint DisplayMask = 0;
+		
+		CGLPixelFormatObj PixelFormat = CGLGetPixelFormat(Context);
+		
+		if(PixelFormat && CGLDescribePixelFormat(PixelFormat, VirtualScreen, kCGLPFADisplayMask, &DisplayMask) == kCGLNoError
+		   && CGLGetParameter(Context, kCGLCPCurrentRendererID, &RendererID) == kCGLNoError)
+		{
+			// Get renderer info for all renderers that match the display mask.
+			GLint Num = 0;
+			CGLRendererInfoObj Renderer;
+			
+			CGLQueryRendererInfo((GLuint)DisplayMask, &Renderer, &Num);
+			
+			for (GLint i = 0; i < Num; i++)
+			{
+				GLint ThisRendererID = 0;
+				CGLDescribeRenderer(Renderer, i, kCGLRPRendererID, &ThisRendererID);
+				
+				// See if this is the one we want
+				if (ThisRendererID == RendererID)
+				{
+					GLint MemoryMB = 0;
+					CGLDescribeRenderer(Renderer, i, kCGLRPVideoMemoryMegabytes, (GLint*)&MemoryMB);
+					VideoMemory = (uint64)MemoryMB * 1024llu * 1024llu;
+					break;
+				}
+			}
+			
+			CGLDestroyRendererInfo(Renderer);
+		}
+	}
+	return VideoMemory;
+}
+
 void FMacOpenGL::ProcessQueryGLInt()
 {
 	GET_GL_INT(GL_MAX_TESS_CONTROL_UNIFORM_COMPONENTS, 0, MaxHullUniformComponents);
diff --git a/Engine/Source/Runtime/OpenGLDrv/Private/Mac/OpenGLMac.h b/Engine/Source/Runtime/OpenGLDrv/Private/Mac/OpenGLMac.h
--- a/Engine/Source/Runtime/OpenGLDrv/Private/Mac/OpenGLMac.h
+++ b/Engine/Source/Runtime/OpenGLDrv/Private/Mac/OpenGLMac.h
@@ -203,6 +203,8 @@ struct FMacOpenGL : public FOpenGL3
 	
 	static FORCEINLINE bool SupportsSeamlessCubeMap()					{ return true; }
 	
+	static uint64 GetVideoMemorySize();
+	
 	static void ProcessQueryGLInt();
 	static void ProcessExtensions(const FString& ExtensionsString);
 	
diff --git a/Engine/Source/Runtime/OpenGLDrv/Private/OpenGLDevice.cpp b/Engine/Source/Runtime/OpenGLDrv/Private/OpenGLDevice.cpp
--- a/Engine/Source/Runtime/OpenGLDrv/Private/OpenGLDevice.cpp
+++ b/Engine/Source/Runtime/OpenGLDrv/Private/OpenGLDevice.cpp
@@ -916,7 +916,7 @@ void FOpenGLDynamicRHI::Init()
 		ResourceIt->InitRHI();
 	}
 
-#if PLATFORM_WINDOWS
+#if PLATFORM_WINDOWS || PLATFORM_MAC
 
 	extern int64 GOpenGLDedicatedVideoMemory;
 	extern int64 GOpenGLTotalGraphicsMemory;
