UE4 - OpenGL Compute shaders debug/verify

diff --git a/Engine/Source/Runtime/OpenGLDrv/Private/OpenGLShaders.cpp b/Engine/Source/Runtime/OpenGLDrv/Private/OpenGLShaders.cpp
--- a/Engine/Source/Runtime/OpenGLDrv/Private/OpenGLShaders.cpp
+++ b/Engine/Source/Runtime/OpenGLDrv/Private/OpenGLShaders.cpp
@@ -1328,6 +1328,14 @@ FComputeShaderRHIRef FOpenGLDynamicRHI::RHICreateComputeShader(const TArray<uint
 	check(GRHIFeatureLevel >= ERHIFeatureLevel::SM5);
 	
 	FOpenGLComputeShader* ComputeShader = CompileOpenGLShader<FOpenGLComputeShader>(Code);
+	const ANSICHAR* GlslCode = NULL;
+	if (!ComputeShader->bSuccessfullyCompiled)
+	{
+#if DEBUG_GL_SHADERS
+		GlslCode = ComputeShader->GlslCodeString;
+#endif
+		ComputeShader->bSuccessfullyCompiled = VerifyCompiledShader(ComputeShader->Resource, GlslCode);
+	}
 
 	check( ComputeShader != 0);
 
@@ -1338,7 +1346,17 @@ FComputeShaderRHIRef FOpenGLDynamicRHI::RHICreateComputeShader(const TArray<uint
 
 	ComputeShader->LinkedProgram = LinkProgram( Config);
 
-	check( ComputeShader->LinkedProgram != 0);
+	if (ComputeShader->LinkedProgram == NULL)
+	{
+#if DEBUG_GL_SHADERS
+		if (ComputeShader->bSuccessfullyCompiled)
+		{
+			UE_LOG(LogRHI,Error,TEXT("Compute Shader:\n%s"),ANSI_TO_TCHAR(ComputeShader->GlslCode.GetTypedData()));
+		}
+#endif //DEBUG_GL_SHADERS
+		checkf(ComputeShader->LinkedProgram, TEXT("Compute shader failed to compile & link."));
+	}
+
 	return ComputeShader;
 }
 
