#UE4  #opengl

   move querying of tesselation related values from base to Opengl3/4 code path.

diff --git a/Engine/Source/Runtime/OpenGLDrv/Private/OpenGL3.cpp b/Engine/Source/Runtime/OpenGLDrv/Private/OpenGL3.cpp
--- a/Engine/Source/Runtime/OpenGLDrv/Private/OpenGL3.cpp
+++ b/Engine/Source/Runtime/OpenGLDrv/Private/OpenGL3.cpp
@@ -8,6 +8,8 @@
 
 #if OPENGL_GL3
 
+bool FOpenGL3::bSupportsTessellation = false;
+
 GLsizei FOpenGL3::NextTextureName = OPENGL_NAME_CACHE_SIZE;
 GLuint FOpenGL3::TextureNamesCache[OPENGL_NAME_CACHE_SIZE];
 GLsizei FOpenGL3::NextBufferName= OPENGL_NAME_CACHE_SIZE;
@@ -24,6 +26,12 @@ void FOpenGL3::ProcessQueryGLInt()
 
 	GET_GL_INT(GL_MAX_GEOMETRY_TEXTURE_IMAGE_UNITS, 0, MaxGeometryTextureImageUnits);
 
+	if (bSupportsTessellation)
+	{
+		GET_GL_INT(GL_MAX_TESS_CONTROL_UNIFORM_COMPONENTS, 0, MaxHullUniformComponents);
+		GET_GL_INT(GL_MAX_TESS_EVALUATION_UNIFORM_COMPONENTS, 0, MaxDomainUniformComponents);
+	}
+
 #ifndef __GNUC__
 #define LOG_AND_GET_GL_QUERY_INT(IntEnum, Default, Dest) do { if (IntEnum) {glGetQueryiv(IntEnum, GL_QUERY_COUNTER_BITS, &Dest);} else {Dest = Default;} /*FPlatformMisc::LowLevelOutputDebugStringf(TEXT(" GL_QUERY_COUNTER_BITS: ") ## TEXT(#IntEnum) ## TEXT(": %d"), Dest);*/ } while(0)
 #else
@@ -32,10 +40,28 @@ void FOpenGL3::ProcessQueryGLInt()
 	LOG_AND_GET_GL_QUERY_INT(GL_TIMESTAMP, 0, TimestampQueryBits);
 	
 #undef LOG_AND_GET_GL_QUERY_INT
+
+    MaxHullTextureImageUnits = 0;
 }
 
 void FOpenGL3::ProcessExtensions( const FString& ExtensionsString )
 {
+	int32 MajorVersion = 0;
+	int32 MinorVersion = 0;
+
+	FString Version = ANSI_TO_TCHAR((const ANSICHAR*)glGetString(GL_VERSION));
+	FString MajorString, MinorString;
+	if (Version.Split(TEXT("."), &MajorString, &MinorString))
+	{
+		MajorVersion = FCString::Atoi(*MajorString);
+		MinorVersion = FCString::Atoi(*MinorString);
+	}
+	check(MajorVersion!=0);
+
+	// 3.3+ may expose this as an extension, 4.0+ should have this. 
+	// https://www.opengl.org/registry/specs/ARB/tessellation_shader.txt
+	bSupportsTessellation = ExtensionsString.Contains(TEXT("GL_ARB_tessellation_shader")) || ((MajorVersion >= 4));
+
 	ProcessQueryGLInt();
 	FOpenGLBase::ProcessExtensions(ExtensionsString);
 
diff --git a/Engine/Source/Runtime/OpenGLDrv/Private/OpenGL4.cpp b/Engine/Source/Runtime/OpenGLDrv/Private/OpenGL4.cpp
--- a/Engine/Source/Runtime/OpenGLDrv/Private/OpenGL4.cpp
+++ b/Engine/Source/Runtime/OpenGLDrv/Private/OpenGL4.cpp
@@ -8,13 +8,10 @@
 
 #if OPENGL_GL4
 
-int32 FOpenGL4::MajorVersion = 0;
-int32 FOpenGL4::MinorVersion = 0;
 GLint FOpenGL4::MaxComputeTextureImageUnits = -1;
 GLint FOpenGL4::MaxComputeUniformComponents = -1;
 bool FOpenGL4::bSupportsComputeShaders = true;
 bool FOpenGL4::bSupportsGPUMemoryInfo = false;
-bool FOpenGL4::bSupportsTessellation = true;
 bool FOpenGL4::bSupportsVertexAttribBinding = true;
 bool FOpenGL4::bSupportsTextureView = true;
 
@@ -26,15 +23,13 @@ void FOpenGL4::ProcessQueryGLInt()
 		GET_GL_INT(GL_MAX_COMPUTE_UNIFORM_COMPONENTS, 0, MaxComputeUniformComponents);
 	}
 
-	if (bSupportsTessellation)
-	{
-		GET_GL_INT(GL_MAX_TESS_CONTROL_UNIFORM_COMPONENTS, 0, MaxHullUniformComponents);
-		GET_GL_INT(GL_MAX_TESS_EVALUATION_UNIFORM_COMPONENTS, 0, MaxDomainUniformComponents);
-	}
 }
 
 void FOpenGL4::ProcessExtensions( const FString& ExtensionsString )
 {
+    int32 MajorVersion =0;
+    int32 MinorVersion =0;
+ 
 	FString Version = ANSI_TO_TCHAR((const ANSICHAR*)glGetString(GL_VERSION));
 	FString MajorString, MinorString;
 	if (Version.Split(TEXT("."), &MajorString, &MinorString))
@@ -42,10 +37,11 @@ void FOpenGL4::ProcessExtensions( const FString& ExtensionsString )
 		MajorVersion = FCString::Atoi(*MajorString);
 		MinorVersion = FCString::Atoi(*MinorString);
 	}
+	check(MajorVersion!=0);
+
 
 	bSupportsGPUMemoryInfo = ExtensionsString.Contains(TEXT("GL_NVX_gpu_memory_info"));
 	bSupportsComputeShaders = ExtensionsString.Contains(TEXT("GL_ARB_compute_shader")) || (MajorVersion ==4 && MinorVersion >= 3) || (MajorVersion > 4);
-	bSupportsTessellation = ExtensionsString.Contains(TEXT("GL_ARB_tessellation_shader")) || (MajorVersion == 4 && MinorVersion >= 1) || (MajorVersion > 4);
 	bSupportsVertexAttribBinding = ExtensionsString.Contains(TEXT("GL_ARB_vertex_attrib_binding")) || (MajorVersion == 4 && MinorVersion >= 3) || (MajorVersion > 4);
 	bSupportsTextureView = ExtensionsString.Contains(TEXT("GL_ARB_texture_view")) || (MajorVersion == 4 && MinorVersion >= 3) || (MajorVersion > 4);
 
diff --git a/Engine/Source/Runtime/OpenGLDrv/Private/OpenGLDrv.cpp b/Engine/Source/Runtime/OpenGLDrv/Private/OpenGLDrv.cpp
--- a/Engine/Source/Runtime/OpenGLDrv/Private/OpenGLDrv.cpp
+++ b/Engine/Source/Runtime/OpenGLDrv/Private/OpenGLDrv.cpp
@@ -457,9 +457,6 @@ void FOpenGLBase::ProcessQueryGLInt()
 	GET_GL_INT(GL_MAX_TEXTURE_IMAGE_UNITS, 0, MaxTextureImageUnits);
 	GET_GL_INT(GL_MAX_VERTEX_TEXTURE_IMAGE_UNITS, 0, MaxVertexTextureImageUnits);
 	GET_GL_INT(GL_MAX_COMBINED_TEXTURE_IMAGE_UNITS, 0, MaxCombinedTextureImageUnits);
-
-	GET_GL_INT(GL_MAX_TESS_CONTROL_TEXTURE_IMAGE_UNITS, 0, MaxHullTextureImageUnits);
-	GET_GL_INT(GL_MAX_TESS_EVALUATION_TEXTURE_IMAGE_UNITS, 0, MaxDomainTextureImageUnits);
 }
 
 void FOpenGLBase::ProcessExtensions( const FString& ExtensionsString )
diff --git a/Engine/Source/Runtime/OpenGLDrv/Public/OpenGL3.h b/Engine/Source/Runtime/OpenGLDrv/Public/OpenGL3.h
--- a/Engine/Source/Runtime/OpenGLDrv/Public/OpenGL3.h
+++ b/Engine/Source/Runtime/OpenGLDrv/Public/OpenGL3.h
@@ -24,8 +24,8 @@ struct FOpenGL3 : public FOpenGLBase
 	static FORCEINLINE bool SupportsVolumeTextureRendering()	{ return bSupportsVolumeTextureRendering; }
 	static FORCEINLINE bool SupportsGenerateMipmap()			{ return true; }
 	static FORCEINLINE bool AmdWorkaround()						{ return bAmdWorkaround; }
-
-
+	static FORCEINLINE bool SupportsTessellation()				{ return bSupportsTessellation; }
+	
 	// Optional
 	static FORCEINLINE void QueryTimestampCounter(GLuint QueryID)
 	{
@@ -569,4 +569,5 @@ protected:
 	static GLint TimestampQueryBits;
 	
 	static bool bDebugContext;
+	static bool bSupportsTessellation;
 };
diff --git a/Engine/Source/Runtime/OpenGLDrv/Public/OpenGL4.h b/Engine/Source/Runtime/OpenGLDrv/Public/OpenGL4.h
--- a/Engine/Source/Runtime/OpenGLDrv/Public/OpenGL4.h
+++ b/Engine/Source/Runtime/OpenGLDrv/Public/OpenGL4.h
@@ -14,7 +14,6 @@
 struct FOpenGL4 : public FOpenGL3
 {
 	static FORCEINLINE bool SupportsSeparateAlphaBlend()				{ return true; }
-	static FORCEINLINE bool SupportsTessellation()						{ return bSupportsTessellation; }
 	static FORCEINLINE bool SupportsComputeShaders()					{ return bSupportsComputeShaders; }
 	static FORCEINLINE bool SupportsDrawIndirect()						{ return true; }
 	static FORCEINLINE bool SupportsVertexAttribBinding()				{ return bSupportsVertexAttribBinding; }
@@ -107,10 +106,7 @@ protected:
 
 	static bool bSupportsComputeShaders;
 	static bool bSupportsGPUMemoryInfo;
-	static bool bSupportsTessellation;
 	static bool bSupportsVertexAttribBinding;
 	static bool bSupportsTextureView;
 
-	static int32 MajorVersion;
-	static int32 MinorVersion;
 };
