Adding error support for missing iARB checks in GL 3.2 for windows.

Moved the blend seperate ARB fns to optional.
Added extension check for the GL_ARB_draw_buffers_blend string which now enables: SupportsSeparateAlphaBlend()
Removed the GL4 override.

#fixes TTP# 341056 and TTP# 339524

diff --git a/Engine/Source/Runtime/OpenGLDrv/Private/OpenGLDrv.cpp b/Engine/Source/Runtime/OpenGLDrv/Private/OpenGLDrv.cpp
--- a/Engine/Source/Runtime/OpenGLDrv/Private/OpenGLDrv.cpp
+++ b/Engine/Source/Runtime/OpenGLDrv/Private/OpenGLDrv.cpp
@@ -450,6 +450,7 @@ bool  FOpenGLBase::bSupportsCopyImage = false;
 bool  FOpenGLBase::bSupportsSeamlessCubemap = false;
 bool  FOpenGLBase::bSupportsVolumeTextureRendering = false;
 bool  FOpenGLBase::bSupportsTextureFilterAnisotropic = false;
+bool  FOpenGLBase::bSupportsDrawBuffersBlend = false;
 bool  FOpenGLBase::bAmdWorkaround = false;
 
 void FOpenGLBase::ProcessQueryGLInt()
@@ -497,6 +498,8 @@ void FOpenGLBase::ProcessExtensions( const FString& ExtensionsString )
 	
 	bSupportsTextureFilterAnisotropic = ExtensionsString.Contains(TEXT("GL_EXT_texture_filter_anisotropic"));
 
+	bSupportsDrawBuffersBlend = ExtensionsString.Contains(TEXT("GL_ARB_draw_buffers_blend"));
+
 	#if PLATFORM_WINDOWS
 		FString VendorName( ANSI_TO_TCHAR((const ANSICHAR*)glGetString(GL_VENDOR) ) );
 		if ( VendorName.Contains(TEXT("ATI ")) )
diff --git a/Engine/Source/Runtime/OpenGLDrv/Public/OpenGL.h b/Engine/Source/Runtime/OpenGLDrv/Public/OpenGL.h
--- a/Engine/Source/Runtime/OpenGLDrv/Public/OpenGL.h
+++ b/Engine/Source/Runtime/OpenGLDrv/Public/OpenGL.h
@@ -123,7 +123,7 @@ public:
 	static FORCEINLINE bool SupportsTextureCubeLodEXT()					{ return true; }
 	static FORCEINLINE bool SupportsShaderTextureLod()					{ return false; }
 	static FORCEINLINE bool SupportsShaderTextureCubeLod()				{ return true; }
-	static FORCEINLINE bool SupportsSeparateAlphaBlend()				{ return false; }
+	static FORCEINLINE bool SupportsSeparateAlphaBlend()				{ return bSupportsDrawBuffersBlend; }
 	static FORCEINLINE bool SupportsTessellation()						{ return false; }
 	static FORCEINLINE bool SupportsComputeShaders()					{ return false; }
 	static FORCEINLINE bool SupportsTextureView()						{ return false; }
@@ -335,6 +335,9 @@ protected:
 	/** GL_EXT_texture_filter_anisotropic Can we use anisotropic filtering? */
 	static bool bSupportsTextureFilterAnisotropic;
 
+	/** GL_ARB_draw_buffers_blend */
+	static bool bSupportsDrawBuffersBlend;
+
 	/** Workaround AMD driver issues. */
 	static bool bAmdWorkaround;
 
diff --git a/Engine/Source/Runtime/OpenGLDrv/Public/OpenGL4.h b/Engine/Source/Runtime/OpenGLDrv/Public/OpenGL4.h
--- a/Engine/Source/Runtime/OpenGLDrv/Public/OpenGL4.h
+++ b/Engine/Source/Runtime/OpenGLDrv/Public/OpenGL4.h
@@ -13,7 +13,6 @@
 
 struct FOpenGL4 : public FOpenGL3
 {
-	static FORCEINLINE bool SupportsSeparateAlphaBlend()				{ return true; }
 	static FORCEINLINE bool SupportsComputeShaders()					{ return bSupportsComputeShaders; }
 	static FORCEINLINE bool SupportsDrawIndirect()						{ return true; }
 	static FORCEINLINE bool SupportsVertexAttribBinding()				{ return bSupportsVertexAttribBinding; }
