UE4 - Common crosscompiler headers

diff --git a/Engine/Source/Developer/ShaderFormatOpenGL/Private/OpenGLShaderCompiler.cpp b/Engine/Source/Developer/ShaderFormatOpenGL/Private/OpenGLShaderCompiler.cpp
--- a/Engine/Source/Developer/ShaderFormatOpenGL/Private/OpenGLShaderCompiler.cpp
+++ b/Engine/Source/Developer/ShaderFormatOpenGL/Private/OpenGLShaderCompiler.cpp
@@ -1,4 +1,4 @@
-// Copyright 1998-2014 Epic Games, Inc. All Rights Reserved.
+// Copyright 1998-2014 Epic Games, Inc. All Rights Reserved. 
 //
 
 #include "ShaderFormatOpenGL.h"
@@ -42,7 +42,7 @@
 	#define GL_TESS_CONTROL_SHADER 0x8E88
 	#endif
 #endif
-#include "OpenGLUtil.h"
+	#include "OpenGLUtil.h"
 #include "OpenGLShaderResources.h"
 
 DEFINE_LOG_CATEGORY_STATIC(LogOpenGLShaderCompiler, Log, All); 
@@ -773,7 +773,7 @@ static void BuildShaderOutput(
 			verify(Match(ShaderSource, ':'));
 
 			CopyInfo.DestUBTypeName = *ShaderSource++;
-			CopyInfo.DestUBTypeIndex = GLPackedTypeNameToTypeIndex(CopyInfo.DestUBTypeName);
+			CopyInfo.DestUBTypeIndex = CrossCompiler::PackedTypeNameToTypeIndex(CopyInfo.DestUBTypeName);
 			verify(Match(ShaderSource, ':'));
 
 			CopyInfo.DestOffsetInFloats = ParseNumber(ShaderSource);
@@ -815,7 +815,7 @@ static void BuildShaderOutput(
 			CopyInfo.DestUBIndex = 0;
 
 			CopyInfo.DestUBTypeName = *ShaderSource++;
-			CopyInfo.DestUBTypeIndex = GLPackedTypeNameToTypeIndex(CopyInfo.DestUBTypeName);
+			CopyInfo.DestUBTypeIndex = CrossCompiler::PackedTypeNameToTypeIndex(CopyInfo.DestUBTypeName);
 			verify(Match(ShaderSource, ':'));
 
 			CopyInfo.DestOffsetInFloats = ParseNumber(ShaderSource);
@@ -849,10 +849,10 @@ static void BuildShaderOutput(
 		ANSICHAR TypeName = Iterator.Key();
 		uint16 Size = Iterator.Value();
 		Size = (Size + 0xf) & (~0xf);
-		FOpenGLPackedArrayInfo Info;
+		CrossCompiler::FPackedArrayInfo Info;
 		Info.Size = Size;
 		Info.TypeName = TypeName;
-		Info.TypeIndex = GLPackedTypeNameToTypeIndex(TypeName);
+		Info.TypeIndex = CrossCompiler::PackedTypeNameToTypeIndex(TypeName);
 		Header.Bindings.PackedGlobalArrays.Add(Info);
 	}
 
@@ -862,17 +862,17 @@ static void BuildShaderOutput(
 	{
 		int BufferIndex = Iterator.Key();
 		auto& ArraySizes = Iterator.Value();
-		TArray<FOpenGLPackedArrayInfo> InfoArray;
+		TArray<CrossCompiler::FPackedArrayInfo> InfoArray;
 		InfoArray.Reserve(ArraySizes.Num());
 		for (auto IterSizes = ArraySizes.CreateIterator(); IterSizes; ++IterSizes)
 		{
 			ANSICHAR TypeName = IterSizes.Key();
 			uint16 Size = IterSizes.Value();
 			Size = (Size + 0xf) & (~0xf);
-			FOpenGLPackedArrayInfo Info;
+			CrossCompiler::FPackedArrayInfo Info;
 			Info.Size = Size;
 			Info.TypeName = TypeName;
-			Info.TypeIndex = GLPackedTypeNameToTypeIndex(TypeName);
+			Info.TypeIndex = CrossCompiler::PackedTypeNameToTypeIndex(TypeName);
 			InfoArray.Add(Info);
 		}
 
@@ -1551,6 +1551,7 @@ CCFlags &= ~HLSLCC_NoPreprocess;
 
 		if (Result != 0)
 		{
+			int32 GlslSourceLen = GlslShaderSource ? FCStringAnsi::Strlen(GlslShaderSource) : 0;
 			if (bDumpDebugInfo)
 			{
 				const FString GLSLFile = (Input.DumpDebugInfoPath / TEXT("Output.glsl"));
@@ -1566,7 +1567,6 @@ CCFlags &= ~HLSLCC_NoPreprocess;
 					FFileHelper::SaveStringToFile(NDABatchFileContents, *(Input.DumpDebugInfoPath / TEXT("NDAGLSLCompile.bat")));
 				}
 
-				int32 GlslSourceLen = FCStringAnsi::Strlen(GlslShaderSource);
 				if (GlslSourceLen > 0)
 				{
 					FArchive* FileWriter = IFileManager::Get().CreateFileWriter(*(Input.DumpDebugInfoPath / Input.SourceFilename + TEXT(".glsl")));
diff --git a/Engine/Source/Developer/iOS/MetalShaderFormat/Private/MetalShaderCompiler.cpp b/Engine/Source/Developer/iOS/MetalShaderFormat/Private/MetalShaderCompiler.cpp
--- a/Engine/Source/Developer/iOS/MetalShaderFormat/Private/MetalShaderCompiler.cpp
+++ b/Engine/Source/Developer/iOS/MetalShaderFormat/Private/MetalShaderCompiler.cpp
@@ -403,7 +403,7 @@ static void BuildMetalShaderOutput(
 			verify(Match(ShaderSource, ':'));
 
 			CopyInfo.DestUBTypeName = *ShaderSource++;
-			CopyInfo.DestUBTypeIndex = MetalPackedTypeNameToTypeIndex(CopyInfo.DestUBTypeName);
+			CopyInfo.DestUBTypeIndex = CrossCompiler::PackedTypeNameToTypeIndex(CopyInfo.DestUBTypeName);
 			verify(Match(ShaderSource, ':'));
 
 			CopyInfo.DestOffsetInFloats = ParseNumber(ShaderSource);
@@ -442,7 +442,7 @@ static void BuildMetalShaderOutput(
 			CopyInfo.DestUBIndex = 0;
 
 			CopyInfo.DestUBTypeName = *ShaderSource++;
-			CopyInfo.DestUBTypeIndex = MetalPackedTypeNameToTypeIndex(CopyInfo.DestUBTypeName);
+			CopyInfo.DestUBTypeIndex = CrossCompiler::PackedTypeNameToTypeIndex(CopyInfo.DestUBTypeName);
 			verify(Match(ShaderSource, ':'));
 
 			CopyInfo.DestOffsetInFloats = ParseNumber(ShaderSource);
@@ -472,10 +472,10 @@ static void BuildMetalShaderOutput(
 		ANSICHAR TypeName = Iterator.Key();
 		uint16 Size = Iterator.Value();
 		Size = (Size + 0xf) & (~0xf);
-		FMetalPackedArrayInfo Info;
+		CrossCompiler::FPackedArrayInfo Info;
 		Info.Size = Size;
 		Info.TypeName = TypeName;
-		Info.TypeIndex = MetalPackedTypeNameToTypeIndex(TypeName);
+		Info.TypeIndex = CrossCompiler::PackedTypeNameToTypeIndex(TypeName);
 		Header.Bindings.PackedGlobalArrays.Add(Info);
 	}
 
@@ -485,17 +485,17 @@ static void BuildMetalShaderOutput(
 	{
 		int BufferIndex = Iterator.Key();
 		auto& ArraySizes = Iterator.Value();
-		TArray<FMetalPackedArrayInfo> InfoArray;
+		TArray<CrossCompiler::FPackedArrayInfo> InfoArray;
 		InfoArray.Reserve(ArraySizes.Num());
 		for (auto IterSizes = ArraySizes.CreateIterator(); IterSizes; ++IterSizes)
 		{
 			ANSICHAR TypeName = IterSizes.Key();
 			uint16 Size = IterSizes.Value();
 			Size = (Size + 0xf) & (~0xf);
-			FMetalPackedArrayInfo Info;
+			CrossCompiler::FPackedArrayInfo Info;
 			Info.Size = Size;
 			Info.TypeName = TypeName;
-			Info.TypeIndex = MetalPackedTypeNameToTypeIndex(TypeName);
+			Info.TypeIndex = CrossCompiler::PackedTypeNameToTypeIndex(TypeName);
 			InfoArray.Add(Info);
 		}
 
diff --git a/Engine/Source/Runtime/IOS/MetalRHI/Private/MetalCommands.cpp b/Engine/Source/Runtime/IOS/MetalRHI/Private/MetalCommands.cpp
--- a/Engine/Source/Runtime/IOS/MetalRHI/Private/MetalCommands.cpp
+++ b/Engine/Source/Runtime/IOS/MetalRHI/Private/MetalCommands.cpp
@@ -216,7 +216,7 @@ void FMetalDynamicRHI::RHISetShaderSampler(FComputeShaderRHIParamRef ComputeShad
 
 void FMetalDynamicRHI::RHISetShaderParameter(FVertexShaderRHIParamRef VertexShaderRHI, uint32 BufferIndex, uint32 BaseIndex, uint32 NumBytes, const void* NewValue)
 {
-	FMetalManager::Get()->GetShaderParameters(METAL_SHADER_STAGE_VERTEX).Set(BufferIndex, BaseIndex, NumBytes, NewValue);
+	FMetalManager::Get()->GetShaderParameters(CrossCompiler::SHADER_STAGE_VERTEX).Set(BufferIndex, BaseIndex, NumBytes, NewValue);
 }
 
 void FMetalDynamicRHI::RHISetShaderParameter(FHullShaderRHIParamRef HullShaderRHI, uint32 BufferIndex, uint32 BaseIndex, uint32 NumBytes, const void* NewValue)
@@ -227,7 +227,7 @@ void FMetalDynamicRHI::RHISetShaderParameter(FHullShaderRHIParamRef HullShaderRH
 
 void FMetalDynamicRHI::RHISetShaderParameter(FPixelShaderRHIParamRef PixelShaderRHI, uint32 BufferIndex, uint32 BaseIndex, uint32 NumBytes, const void* NewValue)
 {
-	FMetalManager::Get()->GetShaderParameters(METAL_SHADER_STAGE_PIXEL).Set(BufferIndex, BaseIndex, NumBytes, NewValue);
+	FMetalManager::Get()->GetShaderParameters(CrossCompiler::SHADER_STAGE_PIXEL).Set(BufferIndex, BaseIndex, NumBytes, NewValue);
 }
 
 void FMetalDynamicRHI::RHISetShaderParameter(FDomainShaderRHIParamRef DomainShaderRHI, uint32 BufferIndex, uint32 BaseIndex, uint32 NumBytes, const void* NewValue)
diff --git a/Engine/Source/Runtime/IOS/MetalRHI/Private/MetalManager.cpp b/Engine/Source/Runtime/IOS/MetalRHI/Private/MetalManager.cpp
--- a/Engine/Source/Runtime/IOS/MetalRHI/Private/MetalManager.cpp
+++ b/Engine/Source/Runtime/IOS/MetalRHI/Private/MetalManager.cpp
@@ -230,9 +230,9 @@ FMetalManager::FMetalManager()
 
 //@todo-rco: What Size???
 	// make a buffer for each shader type
-	ShaderParameters = new FMetalShaderParameterCache[METAL_NUM_SHADER_STAGES];
-	ShaderParameters[METAL_SHADER_STAGE_VERTEX].InitializeResources(1024*1024);
-	ShaderParameters[METAL_SHADER_STAGE_PIXEL].InitializeResources(1024*1024);
+	ShaderParameters = new FMetalShaderParameterCache[CrossCompiler::NUM_SHADER_STAGES];
+	ShaderParameters[CrossCompiler::SHADER_STAGE_VERTEX].InitializeResources(1024*1024);
+	ShaderParameters[CrossCompiler::SHADER_STAGE_PIXEL].InitializeResources(1024*1024);
 
 	// create a semaphore for multi-buffering the command buffer
 	CommandBufferSemaphore = dispatch_semaphore_create(FParse::Param(FCommandLine::Get(),TEXT("gpulockstep")) ? 1 : 3);
@@ -676,7 +676,7 @@ uint32 FMetalManager::AllocateFromQueryBuffer()
 FORCEINLINE void SetResource(uint32 ShaderStage, uint32 BindIndex, FMetalSurface* RESTRICT Surface)
 {
 	check(Surface->Texture != nil);
-	if (ShaderStage == METAL_SHADER_STAGE_PIXEL)
+	if (ShaderStage == CrossCompiler::SHADER_STAGE_PIXEL)
 	{
 		[FMetalManager::GetContext() setFragmentTexture:Surface->Texture atIndex:BindIndex];
 	}
@@ -689,13 +689,13 @@ FORCEINLINE void SetResource(uint32 ShaderStage, uint32 BindIndex, FMetalSurface
 FORCEINLINE void SetResource(uint32 ShaderStage, uint32 BindIndex, FMetalSamplerState* RESTRICT SamplerState)
 {
 	check(SamplerState->State != nil);
-	if (ShaderStage == METAL_SHADER_STAGE_PIXEL)
+	if (ShaderStage == CrossCompiler::SHADER_STAGE_PIXEL)
 	{
 		[FMetalManager::GetContext() setFragmentSamplerState:SamplerState->State atIndex:BindIndex];
 	}
 	else
 	{
-		[FMetalManager::GetContext() setFragmentSamplerState:SamplerState->State atIndex:BindIndex];
+		[FMetalManager::GetContext() setVertexSamplerState:SamplerState->State atIndex:BindIndex];
 	}
 }
 
@@ -763,10 +763,10 @@ void FMetalManager::CommitGraphicsResourceTables()
 
 	check(CurrentBoundShaderState);
 
-	SetResourcesFromTables(CurrentBoundShaderState->VertexShader, METAL_SHADER_STAGE_VERTEX, ResourceTableFrameCounter);
+	SetResourcesFromTables(CurrentBoundShaderState->VertexShader, CrossCompiler::SHADER_STAGE_VERTEX, ResourceTableFrameCounter);
 	if (IsValidRef(CurrentBoundShaderState->PixelShader))
 	{
-		SetResourcesFromTables(CurrentBoundShaderState->PixelShader, METAL_SHADER_STAGE_PIXEL, ResourceTableFrameCounter);
+		SetResourcesFromTables(CurrentBoundShaderState->PixelShader, CrossCompiler::SHADER_STAGE_PIXEL, ResourceTableFrameCounter);
 	}
 
 //	CommitResourceTableCycles += (FPlatformTime::Cycles() - Start);
@@ -774,13 +774,13 @@ void FMetalManager::CommitGraphicsResourceTables()
 
 void FMetalManager::CommitNonComputeShaderConstants()
 {
-	ShaderParameters[METAL_SHADER_STAGE_VERTEX].CommitPackedUniformBuffers(CurrentBoundShaderState, METAL_SHADER_STAGE_VERTEX, CurrentBoundShaderState->VertexShader->BoundUniformBuffers, CurrentBoundShaderState->VertexShader->UniformBuffersCopyInfo);
-	ShaderParameters[METAL_SHADER_STAGE_VERTEX].CommitPackedGlobals(METAL_SHADER_STAGE_VERTEX, CurrentBoundShaderState->VertexShader->Bindings);
+	ShaderParameters[CrossCompiler::SHADER_STAGE_VERTEX].CommitPackedUniformBuffers(CurrentBoundShaderState, CrossCompiler::SHADER_STAGE_VERTEX, CurrentBoundShaderState->VertexShader->BoundUniformBuffers, CurrentBoundShaderState->VertexShader->UniformBuffersCopyInfo);
+	ShaderParameters[CrossCompiler::SHADER_STAGE_VERTEX].CommitPackedGlobals(CrossCompiler::SHADER_STAGE_VERTEX, CurrentBoundShaderState->VertexShader->Bindings);
 	
 	if (IsValidRef(CurrentBoundShaderState->PixelShader))
 	{
-		ShaderParameters[METAL_SHADER_STAGE_PIXEL].CommitPackedUniformBuffers(CurrentBoundShaderState, METAL_SHADER_STAGE_PIXEL, CurrentBoundShaderState->PixelShader->BoundUniformBuffers, CurrentBoundShaderState->PixelShader->UniformBuffersCopyInfo);
-		ShaderParameters[METAL_SHADER_STAGE_PIXEL].CommitPackedGlobals(METAL_SHADER_STAGE_PIXEL, CurrentBoundShaderState->PixelShader->Bindings);
+		ShaderParameters[CrossCompiler::SHADER_STAGE_PIXEL].CommitPackedUniformBuffers(CurrentBoundShaderState, CrossCompiler::SHADER_STAGE_PIXEL, CurrentBoundShaderState->PixelShader->BoundUniformBuffers, CurrentBoundShaderState->PixelShader->UniformBuffersCopyInfo);
+		ShaderParameters[CrossCompiler::SHADER_STAGE_PIXEL].CommitPackedGlobals(CrossCompiler::SHADER_STAGE_PIXEL, CurrentBoundShaderState->PixelShader->Bindings);
 	}
 }
 
diff --git a/Engine/Source/Runtime/IOS/MetalRHI/Private/MetalShaders.cpp b/Engine/Source/Runtime/IOS/MetalRHI/Private/MetalShaders.cpp
--- a/Engine/Source/Runtime/IOS/MetalRHI/Private/MetalShaders.cpp
+++ b/Engine/Source/Runtime/IOS/MetalRHI/Private/MetalShaders.cpp
@@ -344,7 +344,7 @@ FBoundShaderStateRHIRef FMetalDynamicRHI::RHICreateBoundShaderState(
 FMetalShaderParameterCache::FMetalShaderParameterCache() :
 	GlobalUniformArraySize(-1)
 {
-	for (int32 ArrayIndex = 0; ArrayIndex < METAL_PACKED_TYPEINDEX_MAX; ++ArrayIndex)
+	for (int32 ArrayIndex = 0; ArrayIndex < CrossCompiler::PACKED_TYPEINDEX_MAX; ++ArrayIndex)
 	{
 		PackedGlobalUniformDirty[ArrayIndex].LowVector = 0;
 		PackedGlobalUniformDirty[ArrayIndex].HighVector = 0;
@@ -355,19 +355,19 @@ void FMetalShaderParameterCache::InitializeResources(int32 UniformArraySize)
 {
 	check(GlobalUniformArraySize == -1);
 	
-	PackedGlobalUniforms[0] = (uint8*)FMemory::Malloc(UniformArraySize * METAL_PACKED_TYPEINDEX_MAX);
-	PackedUniformsScratch[0] = (uint8*)FMemory::Malloc(UniformArraySize * METAL_PACKED_TYPEINDEX_MAX);
+	PackedGlobalUniforms[0] = (uint8*)FMemory::Malloc(UniformArraySize * CrossCompiler::PACKED_TYPEINDEX_MAX);
+	PackedUniformsScratch[0] = (uint8*)FMemory::Malloc(UniformArraySize * CrossCompiler::PACKED_TYPEINDEX_MAX);
 	
-	FMemory::Memzero(PackedGlobalUniforms[0], UniformArraySize * METAL_PACKED_TYPEINDEX_MAX);
-	FMemory::Memzero(PackedUniformsScratch[0], UniformArraySize * METAL_PACKED_TYPEINDEX_MAX);
-	for (int32 ArrayIndex = 1; ArrayIndex < METAL_PACKED_TYPEINDEX_MAX; ++ArrayIndex)
+	FMemory::Memzero(PackedGlobalUniforms[0], UniformArraySize * CrossCompiler::PACKED_TYPEINDEX_MAX);
+	FMemory::Memzero(PackedUniformsScratch[0], UniformArraySize * CrossCompiler::PACKED_TYPEINDEX_MAX);
+	for (int32 ArrayIndex = 1; ArrayIndex < CrossCompiler::PACKED_TYPEINDEX_MAX; ++ArrayIndex)
 	{
 		PackedGlobalUniforms[ArrayIndex] = PackedGlobalUniforms[ArrayIndex - 1] + UniformArraySize;
 		PackedUniformsScratch[ArrayIndex] = PackedUniformsScratch[ArrayIndex - 1] + UniformArraySize;
 	}
 	GlobalUniformArraySize = UniformArraySize;
 	
-	for (int32 ArrayIndex = 0; ArrayIndex < METAL_PACKED_TYPEINDEX_MAX; ++ArrayIndex)
+	for (int32 ArrayIndex = 0; ArrayIndex < CrossCompiler::PACKED_TYPEINDEX_MAX; ++ArrayIndex)
 	{
 		PackedGlobalUniformDirty[ArrayIndex].LowVector = 0;
 		PackedGlobalUniformDirty[ArrayIndex].HighVector = 0;
@@ -394,7 +394,7 @@ FMetalShaderParameterCache::~FMetalShaderParameterCache()
  */
 void FMetalShaderParameterCache::MarkAllDirty()
 {
-	for (int32 ArrayIndex = 0; ArrayIndex < METAL_PACKED_TYPEINDEX_MAX; ++ArrayIndex)
+	for (int32 ArrayIndex = 0; ArrayIndex < CrossCompiler::PACKED_TYPEINDEX_MAX; ++ArrayIndex)
 	{
 		PackedGlobalUniformDirty[ArrayIndex].LowVector = 0;
 		PackedGlobalUniformDirty[ArrayIndex].HighVector = 0;//UniformArraySize / 16;
@@ -408,9 +408,9 @@ const int SizeOfFloat4 = 4 * sizeof(float);
  */
 void FMetalShaderParameterCache::Set(uint32 BufferIndexName, uint32 ByteOffset, uint32 NumBytes, const void* NewValues)
 {
-	uint32 BufferIndex = MetalPackedTypeNameToTypeIndex(BufferIndexName);
+	uint32 BufferIndex = CrossCompiler::PackedTypeNameToTypeIndex(BufferIndexName);
 	check(GlobalUniformArraySize != -1);
-	check(BufferIndex < METAL_PACKED_TYPEINDEX_MAX);
+	check(BufferIndex < CrossCompiler::PACKED_TYPEINDEX_MAX);
 	check(ByteOffset + NumBytes <= (uint32)GlobalUniformArraySize);
 	PackedGlobalUniformDirty[BufferIndex].LowVector = FMath::Min(PackedGlobalUniformDirty[BufferIndex].LowVector, ByteOffset / SizeOfFloat4);
 	PackedGlobalUniformDirty[BufferIndex].HighVector = FMath::Max(PackedGlobalUniformDirty[BufferIndex].HighVector, (ByteOffset + NumBytes + SizeOfFloat4 - 1) / SizeOfFloat4);
@@ -434,7 +434,7 @@ void FMetalShaderParameterCache::CommitPackedGlobals(int32 Stage, const FMetalSh
 			//check(SizeToUpload <= TotalSize);
 			SizeToUpload = TotalSize;
  /*
-			if (UniformBufferIndex == METAL_PACKED_TYPEINDEX_MEDIUMP)
+			if (UniformBufferIndex == CrossCompiler::PACKED_TYPEINDEX_MEDIUMP)
 			{
 				// Float4 -> Half4
 				SizeToUpload /= 2;
@@ -445,7 +445,7 @@ void FMetalShaderParameterCache::CommitPackedGlobals(int32 Stage, const FMetalSh
 			uint32 Offset = FMetalManager::Get()->AllocateFromRingBuffer(TotalSize);
 /*
 			// copy into the middle of the ring buffer
-			if (UniformBufferIndex == METAL_PACKED_TYPEINDEX_MEDIUMP)
+			if (UniformBufferIndex == CrossCompiler::PACKED_TYPEINDEX_MEDIUMP)
 			{
 				// Convert to half
 				const int NumVectors = PackedGlobalUniformDirty[UniformBufferIndex].HighVector;
@@ -462,19 +462,19 @@ void FMetalShaderParameterCache::CommitPackedGlobals(int32 Stage, const FMetalSh
 			else
  */
 			{
-				//check(UniformBufferIndex != METAL_PACKED_TYPEINDEX_LOWP);
+				//check(UniformBufferIndex != CrossCompiler::PACKED_TYPEINDEX_LOWP);
 //@todo-rco: Temp workaround
 				FMemory::Memcpy(((uint8*)[RingBuffer contents]) + Offset, PackedGlobalUniforms[UniformBufferIndex], FMath::Min(TotalSize, SizeToUpload));
 			}
 			
 			switch (Stage)
 			{
-				case METAL_SHADER_STAGE_VERTEX:
+				case CrossCompiler::SHADER_STAGE_VERTEX:
 //					NSLog(@"Uploading %d bytes to vertex", SizeToUpload);
 					[FMetalManager::GetContext() setVertexBuffer:RingBuffer offset:Offset atIndex:UniformBufferIndex];
 					break;
 
-				case METAL_SHADER_STAGE_PIXEL:
+				case CrossCompiler::SHADER_STAGE_PIXEL:
 //					NSLog(@"Uploading %d bytes to pixel", SizeToUpload);
 					[FMetalManager::GetContext() setFragmentBuffer:RingBuffer offset : Offset atIndex : UniformBufferIndex];
 					break;
@@ -494,12 +494,12 @@ void FMetalShaderParameterCache::CommitPackedUniformBuffers(TRefCountPtr<FMetalB
 //	SCOPE_CYCLE_COUNTER(STAT_MetalConstantBufferUpdateTime);
 	// Uniform Buffers are split into precision/type; the list of RHI UBs is traversed and if a new one was set, its
 	// contents are copied per precision/type into corresponding scratch buffers which are then uploaded to the program
-	if (Stage == METAL_SHADER_STAGE_PIXEL && !IsValidRef(BoundShaderState->PixelShader))
+	if (Stage == CrossCompiler::SHADER_STAGE_PIXEL && !IsValidRef(BoundShaderState->PixelShader))
 	{
 		return;
 	}
 
-	auto& Bindings = Stage == METAL_SHADER_STAGE_VERTEX ? BoundShaderState->VertexShader->Bindings : BoundShaderState->PixelShader->Bindings;
+	auto& Bindings = Stage == CrossCompiler::SHADER_STAGE_VERTEX ? BoundShaderState->VertexShader->Bindings : BoundShaderState->PixelShader->Bindings;
 	check(Bindings.NumUniformBuffers == RHIUniformBuffers.Num());
 	if (!Bindings.bHasRegularUniformBuffers)
 	{
@@ -537,7 +537,7 @@ void FMetalShaderParameterCache::CommitPackedUniformBuffers(TRefCountPtr<FMetalB
 			if (RHIUniformBuffer)
 			{
 				auto* UB = (FMetalUniformBuffer*)RHIUniformBuffer;
-				if (Stage == METAL_SHADER_STAGE_VERTEX)
+				if (Stage == CrossCompiler::SHADER_STAGE_VERTEX)
 				{
 					[FMetalManager::GetContext() setVertexBuffer:UB->Buffer offset:UB->Offset atIndex:BufferIndex];
 				}
diff --git a/Engine/Source/Runtime/IOS/MetalRHI/Public/MetalResources.h b/Engine/Source/Runtime/IOS/MetalRHI/Public/MetalResources.h
--- a/Engine/Source/Runtime/IOS/MetalRHI/Public/MetalResources.h
+++ b/Engine/Source/Runtime/IOS/MetalRHI/Public/MetalResources.h
@@ -440,7 +440,7 @@ public:
 
 private:
 	/** CPU memory block for storing uniform values. */
-	uint8* PackedGlobalUniforms[METAL_PACKED_TYPEINDEX_MAX];
+	uint8* PackedGlobalUniforms[CrossCompiler::PACKED_TYPEINDEX_MAX];
 
 	struct FRange
 	{
@@ -448,10 +448,10 @@ private:
 		uint32	HighVector;
 	};
 	/** Dirty ranges for each uniform array. */
-	FRange	PackedGlobalUniformDirty[METAL_PACKED_TYPEINDEX_MAX];
+	FRange	PackedGlobalUniformDirty[CrossCompiler::PACKED_TYPEINDEX_MAX];
 
 	/** Scratch CPU memory block for uploading packed uniforms. */
-	uint8* PackedUniformsScratch[METAL_PACKED_TYPEINDEX_MAX];
+	uint8* PackedUniformsScratch[CrossCompiler::PACKED_TYPEINDEX_MAX];
 
 	int32 GlobalUniformArraySize;
 };
diff --git a/Engine/Source/Runtime/IOS/MetalRHI/Public/MetalShaderResources.h b/Engine/Source/Runtime/IOS/MetalRHI/Public/MetalShaderResources.h
--- a/Engine/Source/Runtime/IOS/MetalRHI/Public/MetalShaderResources.h
+++ b/Engine/Source/Runtime/IOS/MetalRHI/Public/MetalShaderResources.h
@@ -6,73 +6,19 @@
 
 #pragma once
 
-enum
-{
-	METAL_SHADER_STAGE_VERTEX = 0,
-	METAL_SHADER_STAGE_PIXEL,
-	METAL_SHADER_STAGE_GEOMETRY,
-	METAL_SHADER_STAGE_HULL,
-	METAL_SHADER_STAGE_DOMAIN,
-	METAL_NUM_NON_COMPUTE_SHADER_STAGES,
-	METAL_SHADER_STAGE_COMPUTE = METAL_NUM_NON_COMPUTE_SHADER_STAGES,
-	METAL_NUM_SHADER_STAGES
-};
+#include "CrossCompilerCommon.h"
 
 /**
 * Shader related constants.
 */
 enum
 {
-	METAL_PACKED_TYPENAME_HIGHP = 'h',	// Make sure these enums match hlslcc
-	METAL_PACKED_TYPENAME_MEDIUMP = 'm',
-	METAL_PACKED_TYPENAME_LOWP = 'l',
-	METAL_PACKED_TYPENAME_INT = 'i',
-	METAL_PACKED_TYPENAME_UINT = 'u',
-	METAL_PACKED_TYPENAME_SAMPLER = 's',
-	METAL_PACKED_TYPENAME_IMAGE = 'g',
-
-	METAL_PACKED_TYPEINDEX_HIGHP = 0,
-	METAL_PACKED_TYPEINDEX_MEDIUMP = 1,
-	METAL_PACKED_TYPEINDEX_LOWP = 2,
-	METAL_PACKED_TYPEINDEX_INT = 3,
-	METAL_PACKED_TYPEINDEX_UINT = 4,
-	METAL_PACKED_TYPEINDEX_MAX = 5,
-
 	METAL_MAX_UNIFORM_BUFFER_BINDINGS = 12,	// @todo-mobile: Remove me
 	METAL_FIRST_UNIFORM_BUFFER = 0,			// @todo-mobile: Remove me
 	METAL_MAX_COMPUTE_STAGE_UAV_UNITS = 8,	// @todo-mobile: Remove me
 	METAL_UAV_NOT_SUPPORTED_FOR_GRAPHICS_UNIT = -1, // for now, only CS supports UAVs/ images
 };
 
-static FORCEINLINE uint8 MetalPackedTypeNameToTypeIndex(uint8 ArrayName)
-{
-	switch (ArrayName)
-	{
-	case METAL_PACKED_TYPENAME_HIGHP:		return METAL_PACKED_TYPEINDEX_HIGHP;
-	case METAL_PACKED_TYPENAME_MEDIUMP:	return METAL_PACKED_TYPEINDEX_MEDIUMP;
-	case METAL_PACKED_TYPENAME_LOWP:		return METAL_PACKED_TYPEINDEX_LOWP;
-	case METAL_PACKED_TYPENAME_INT:		return METAL_PACKED_TYPEINDEX_INT;
-	case METAL_PACKED_TYPENAME_UINT:		return METAL_PACKED_TYPEINDEX_UINT;
-	}
-	check(0);
-	return 0;
-}
-
-struct FMetalPackedArrayInfo
-{
-	uint16	Size;		// Bytes
-	uint8	TypeName;	// OGL_PACKED_TYPENAME
-	uint8	TypeIndex;	// OGL_PACKED_TYPE
-};
-
-inline FArchive& operator<<(FArchive& Ar, FMetalPackedArrayInfo& Info)
-{
-	Ar << Info.Size;
-	Ar << Info.TypeName;
-	Ar << Info.TypeIndex;
-	return Ar;
-}
-
 struct FMetalShaderResourceTable : public FBaseShaderResourceTable
 {
 	/** Mapping of bound Textures to their location in resource tables. */
@@ -105,8 +51,8 @@ inline FArchive& operator<<(FArchive& Ar, FMetalShaderResourceTable& SRT)
 
 struct FMetalShaderBindings
 {
-	TArray<TArray<FMetalPackedArrayInfo>>	PackedUniformBuffers;
-	TArray<FMetalPackedArrayInfo>			PackedGlobalArrays;
+	TArray<TArray<CrossCompiler::FPackedArrayInfo>>	PackedUniformBuffers;
+	TArray<CrossCompiler::FPackedArrayInfo>			PackedGlobalArrays;
 	FMetalShaderResourceTable				ShaderResourceTable;
 
 	uint16	InOutMask;
@@ -146,8 +92,8 @@ struct FMetalShaderBindings
 
 		for (int32 Item = 0; Item < A.PackedUniformBuffers.Num(); Item++)
 		{
-			const TArray<FMetalPackedArrayInfo> &ArrayA = A.PackedUniformBuffers[Item];
-			const TArray<FMetalPackedArrayInfo> &ArrayB = B.PackedUniformBuffers[Item];
+			const TArray<CrossCompiler::FPackedArrayInfo> &ArrayA = A.PackedUniformBuffers[Item];
+			const TArray<CrossCompiler::FPackedArrayInfo> &ArrayB = B.PackedUniformBuffers[Item];
 
 			bEqual &= FMemory::Memcmp(ArrayA.GetTypedData(), ArrayB.GetTypedData(), ArrayA.GetTypeSize()*ArrayA.Num()) == 0;
 		}
@@ -167,7 +113,7 @@ struct FMetalShaderBindings
 
 		for (int32 Item = 0; Item < Binding.PackedUniformBuffers.Num(); Item++)
 		{
-			const TArray<FMetalPackedArrayInfo> &Array = Binding.PackedUniformBuffers[Item];
+			const TArray<CrossCompiler::FPackedArrayInfo> &Array = Binding.PackedUniformBuffers[Item];
 			Hash ^= FCrc::MemCrc_DEPRECATED(Array.GetTypedData(), Array.GetTypeSize()* Array.Num());
 		}
 		return Hash;
@@ -207,7 +153,7 @@ inline FArchive& operator<<(FArchive& Ar, FMetalUniformBufferCopyInfo& Info)
 	Ar << Info.DestUBTypeName;
 	if (Ar.IsLoading())
 	{
-		Info.DestUBTypeIndex = MetalPackedTypeNameToTypeIndex(Info.DestUBTypeName);
+		Info.DestUBTypeIndex = CrossCompiler::PackedTypeNameToTypeIndex(Info.DestUBTypeName);
 	}
 	Ar << Info.DestOffsetInFloats;
 	Ar << Info.SizeInFloats;
diff --git a/Engine/Source/Runtime/OpenGLDrv/Private/OpenGLCommands.cpp b/Engine/Source/Runtime/OpenGLDrv/Private/OpenGLCommands.cpp
--- a/Engine/Source/Runtime/OpenGLDrv/Private/OpenGLCommands.cpp
+++ b/Engine/Source/Runtime/OpenGLDrv/Private/OpenGLCommands.cpp
@@ -1042,7 +1042,7 @@ void FOpenGLDynamicRHI::RHISetShaderParameter(FVertexShaderRHIParamRef VertexSha
 	VALIDATE_BOUND_SHADER(VertexShaderRHI);
 
 	VERIFY_GL_SCOPE();
-	PendingState.ShaderParameters[OGL_SHADER_STAGE_VERTEX].Set(BufferIndex, BaseIndex, NumBytes, NewValue);
+	PendingState.ShaderParameters[CrossCompiler::SHADER_STAGE_VERTEX].Set(BufferIndex, BaseIndex, NumBytes, NewValue);
 }
 
 void FOpenGLDynamicRHI::RHISetShaderParameter(FPixelShaderRHIParamRef PixelShaderRHI,uint32 BufferIndex,uint32 BaseIndex,uint32 NumBytes,const void* NewValue)
@@ -1050,7 +1050,7 @@ void FOpenGLDynamicRHI::RHISetShaderParameter(FPixelShaderRHIParamRef PixelShade
 	VALIDATE_BOUND_SHADER(PixelShaderRHI);
 
 	VERIFY_GL_SCOPE();
-	PendingState.ShaderParameters[OGL_SHADER_STAGE_PIXEL].Set(BufferIndex, BaseIndex, NumBytes, NewValue);
+	PendingState.ShaderParameters[CrossCompiler::SHADER_STAGE_PIXEL].Set(BufferIndex, BaseIndex, NumBytes, NewValue);
 }
 
 void FOpenGLDynamicRHI::RHISetShaderParameter(FHullShaderRHIParamRef HullShaderRHI,uint32 BufferIndex,uint32 BaseIndex,uint32 NumBytes,const void* NewValue)
@@ -1059,7 +1059,7 @@ void FOpenGLDynamicRHI::RHISetShaderParameter(FHullShaderRHIParamRef HullShaderR
 
 	check(GRHIFeatureLevel >= ERHIFeatureLevel::SM5);
 	VERIFY_GL_SCOPE();
-	PendingState.ShaderParameters[OGL_SHADER_STAGE_HULL].Set(BufferIndex, BaseIndex, NumBytes, NewValue);
+	PendingState.ShaderParameters[CrossCompiler::SHADER_STAGE_HULL].Set(BufferIndex, BaseIndex, NumBytes, NewValue);
 }
 
 void FOpenGLDynamicRHI::RHISetShaderParameter(FDomainShaderRHIParamRef DomainShaderRHI,uint32 BufferIndex,uint32 BaseIndex,uint32 NumBytes,const void* NewValue)
@@ -1068,7 +1068,7 @@ void FOpenGLDynamicRHI::RHISetShaderParameter(FDomainShaderRHIParamRef DomainSha
 
 	check(GRHIFeatureLevel >= ERHIFeatureLevel::SM5);
 	VERIFY_GL_SCOPE();
-	PendingState.ShaderParameters[OGL_SHADER_STAGE_DOMAIN].Set(BufferIndex, BaseIndex, NumBytes, NewValue);
+	PendingState.ShaderParameters[CrossCompiler::SHADER_STAGE_DOMAIN].Set(BufferIndex, BaseIndex, NumBytes, NewValue);
 }
 
 void FOpenGLDynamicRHI::RHISetShaderParameter(FGeometryShaderRHIParamRef GeometryShaderRHI,uint32 BufferIndex,uint32 BaseIndex,uint32 NumBytes,const void* NewValue)
@@ -1076,13 +1076,13 @@ void FOpenGLDynamicRHI::RHISetShaderParameter(FGeometryShaderRHIParamRef Geometr
 	VALIDATE_BOUND_SHADER(GeometryShaderRHI);
 
 	VERIFY_GL_SCOPE();
-	PendingState.ShaderParameters[OGL_SHADER_STAGE_GEOMETRY].Set(BufferIndex, BaseIndex, NumBytes, NewValue);
+	PendingState.ShaderParameters[CrossCompiler::SHADER_STAGE_GEOMETRY].Set(BufferIndex, BaseIndex, NumBytes, NewValue);
 }
 
 void FOpenGLDynamicRHI::RHISetShaderParameter(FComputeShaderRHIParamRef ComputeShaderRHI,uint32 BufferIndex,uint32 BaseIndex,uint32 NumBytes,const void* NewValue)
 { 
 	check(GRHIFeatureLevel >= ERHIFeatureLevel::SM5);
-	PendingState.ShaderParameters[OGL_SHADER_STAGE_COMPUTE].Set(BufferIndex, BaseIndex, NumBytes, NewValue);
+	PendingState.ShaderParameters[CrossCompiler::SHADER_STAGE_COMPUTE].Set(BufferIndex, BaseIndex, NumBytes, NewValue);
 }
 
 void FOpenGLDynamicRHI::RHISetDepthStencilState(FDepthStencilStateRHIParamRef NewStateRHI,uint32 StencilRef)
@@ -2057,7 +2057,7 @@ void FOpenGLDynamicRHI::OnUniformBufferDeletion( GLuint UniformBufferResource, u
 		RenderingContextState.UniformBufferBound = -1;	// will force refresh
 	}
 
-	for (GLuint UniformBufferIndex = 0; UniformBufferIndex < OGL_NUM_SHADER_STAGES*OGL_MAX_UNIFORM_BUFFER_BINDINGS; UniformBufferIndex++)
+	for (GLuint UniformBufferIndex = 0; UniformBufferIndex < CrossCompiler::NUM_SHADER_STAGES*OGL_MAX_UNIFORM_BUFFER_BINDINGS; UniformBufferIndex++)
 	{
 		if( SharedContextState.UniformBuffers[UniformBufferIndex] == UniformBufferResource )
 		{
@@ -2078,23 +2078,23 @@ void FOpenGLDynamicRHI::CommitNonComputeShaderConstants()
 	FOpenGLLinkedProgram* LinkedProgram = PendingState.BoundShaderState->LinkedProgram;
 	if (IsES2Platform(GRHIShaderPlatform))
 	{
-		PendingState.ShaderParameters[OGL_SHADER_STAGE_VERTEX].CommitPackedUniformBuffers(LinkedProgram, OGL_SHADER_STAGE_VERTEX, PendingState.BoundUniformBuffers[SF_Vertex], PendingState.BoundShaderState->VertexShader->UniformBuffersCopyInfo);
+		PendingState.ShaderParameters[CrossCompiler::SHADER_STAGE_VERTEX].CommitPackedUniformBuffers(LinkedProgram, CrossCompiler::SHADER_STAGE_VERTEX, PendingState.BoundUniformBuffers[SF_Vertex], PendingState.BoundShaderState->VertexShader->UniformBuffersCopyInfo);
 	}
-	PendingState.ShaderParameters[OGL_SHADER_STAGE_VERTEX].CommitPackedGlobals(LinkedProgram, OGL_SHADER_STAGE_VERTEX);
+	PendingState.ShaderParameters[CrossCompiler::SHADER_STAGE_VERTEX].CommitPackedGlobals(LinkedProgram, CrossCompiler::SHADER_STAGE_VERTEX);
 
 	if (IsES2Platform(GRHIShaderPlatform))
 	{
-		PendingState.ShaderParameters[OGL_SHADER_STAGE_PIXEL].CommitPackedUniformBuffers(LinkedProgram, OGL_SHADER_STAGE_PIXEL, PendingState.BoundUniformBuffers[SF_Pixel], PendingState.BoundShaderState->PixelShader->UniformBuffersCopyInfo);
+		PendingState.ShaderParameters[CrossCompiler::SHADER_STAGE_PIXEL].CommitPackedUniformBuffers(LinkedProgram, CrossCompiler::SHADER_STAGE_PIXEL, PendingState.BoundUniformBuffers[SF_Pixel], PendingState.BoundShaderState->PixelShader->UniformBuffersCopyInfo);
 	}
-	PendingState.ShaderParameters[OGL_SHADER_STAGE_PIXEL].CommitPackedGlobals(LinkedProgram, OGL_SHADER_STAGE_PIXEL);
+	PendingState.ShaderParameters[CrossCompiler::SHADER_STAGE_PIXEL].CommitPackedGlobals(LinkedProgram, CrossCompiler::SHADER_STAGE_PIXEL);
 
 	if (PendingState.BoundShaderState->GeometryShader)
 	{
 		if (IsES2Platform(GRHIShaderPlatform))
 		{
-			PendingState.ShaderParameters[OGL_SHADER_STAGE_GEOMETRY].CommitPackedUniformBuffers(LinkedProgram, OGL_SHADER_STAGE_GEOMETRY, PendingState.BoundUniformBuffers[SF_Geometry], PendingState.BoundShaderState->GeometryShader->UniformBuffersCopyInfo);
+			PendingState.ShaderParameters[CrossCompiler::SHADER_STAGE_GEOMETRY].CommitPackedUniformBuffers(LinkedProgram, CrossCompiler::SHADER_STAGE_GEOMETRY, PendingState.BoundUniformBuffers[SF_Geometry], PendingState.BoundShaderState->GeometryShader->UniformBuffersCopyInfo);
 		}
-		PendingState.ShaderParameters[OGL_SHADER_STAGE_GEOMETRY].CommitPackedGlobals(LinkedProgram, OGL_SHADER_STAGE_GEOMETRY);
+		PendingState.ShaderParameters[CrossCompiler::SHADER_STAGE_GEOMETRY].CommitPackedGlobals(LinkedProgram, CrossCompiler::SHADER_STAGE_GEOMETRY);
 	}
 }
 
@@ -2104,10 +2104,10 @@ void FOpenGLDynamicRHI::CommitComputeShaderConstants(FComputeShaderRHIParamRef C
 	check(GRHIFeatureLevel >= ERHIFeatureLevel::SM5);
 
 	DYNAMIC_CAST_OPENGLRESOURCE(ComputeShader,ComputeShader);
-	const int32 Stage = OGL_SHADER_STAGE_COMPUTE;
+	const int32 Stage = CrossCompiler::SHADER_STAGE_COMPUTE;
 
-	FOpenGLShaderParameterCache& StageShaderParameters = PendingState.ShaderParameters[OGL_SHADER_STAGE_COMPUTE];
-	StageShaderParameters.CommitPackedGlobals(ComputeShader->LinkedProgram, OGL_SHADER_STAGE_COMPUTE);
+	FOpenGLShaderParameterCache& StageShaderParameters = PendingState.ShaderParameters[CrossCompiler::SHADER_STAGE_COMPUTE];
+	StageShaderParameters.CommitPackedGlobals(ComputeShader->LinkedProgram, CrossCompiler::SHADER_STAGE_COMPUTE);
 }
 
 template <EShaderFrequency Frequency>
diff --git a/Engine/Source/Runtime/OpenGLDrv/Private/OpenGLShaders.cpp b/Engine/Source/Runtime/OpenGLDrv/Private/OpenGLShaders.cpp
--- a/Engine/Source/Runtime/OpenGLDrv/Private/OpenGLShaders.cpp
+++ b/Engine/Source/Runtime/OpenGLDrv/Private/OpenGLShaders.cpp
@@ -509,8 +509,8 @@ FGeometryShaderRHIRef FOpenGLDynamicRHI::RHICreateGeometryShaderWithStreamOutput
 
 static void MarkShaderParameterCachesDirty(FOpenGLShaderParameterCache* ShaderParameters, bool UpdateCompute)
 {
-	const int32 StageStart = UpdateCompute  ? OGL_SHADER_STAGE_COMPUTE : OGL_SHADER_STAGE_VERTEX;
-	const int32 StageEnd = UpdateCompute ? OGL_NUM_SHADER_STAGES : OGL_NUM_NON_COMPUTE_SHADER_STAGES;
+	const int32 StageStart = UpdateCompute  ? CrossCompiler::SHADER_STAGE_COMPUTE : CrossCompiler::SHADER_STAGE_VERTEX;
+	const int32 StageEnd = UpdateCompute ? CrossCompiler::NUM_SHADER_STAGES : CrossCompiler::NUM_NON_COMPUTE_SHADER_STAGES;
 	for (int32 Stage = StageStart; Stage < StageEnd; ++Stage)
 	{
 		ShaderParameters[Stage].MarkAllDirty();
@@ -573,11 +573,11 @@ public:
 		FOpenGLShaderBindings Bindings;
 		GLuint Resource;
 	}
-	Shaders[OGL_NUM_SHADER_STAGES];
+	Shaders[CrossCompiler::NUM_SHADER_STAGES];
 
 	FOpenGLLinkedProgramConfiguration()
 	{
-		for ( int32 Stage = 0; Stage < OGL_NUM_SHADER_STAGES; Stage++)
+		for ( int32 Stage = 0; Stage < CrossCompiler::NUM_SHADER_STAGES; Stage++)
 		{
 			Shaders[Stage].Resource = 0;
 		}
@@ -586,7 +586,7 @@ public:
 	friend bool operator ==(const FOpenGLLinkedProgramConfiguration& A,const FOpenGLLinkedProgramConfiguration& B)
 	{
 		bool bEqual = true;
-		for ( int32 Stage = 0; Stage < OGL_NUM_SHADER_STAGES && bEqual; Stage++)
+		for ( int32 Stage = 0; Stage < CrossCompiler::NUM_SHADER_STAGES && bEqual; Stage++)
 		{
 			bEqual &= A.Shaders[Stage].Resource == B.Shaders[Stage].Resource;
 			bEqual &= A.Shaders[Stage].Bindings == B.Shaders[Stage].Bindings;
@@ -597,7 +597,7 @@ public:
 	friend uint32 GetTypeHash(const FOpenGLLinkedProgramConfiguration &Config)
 	{
 		uint32 Hash = 0;
-		for ( int32 Stage = 0; Stage < OGL_NUM_SHADER_STAGES; Stage++)
+		for ( int32 Stage = 0; Stage < CrossCompiler::NUM_SHADER_STAGES; Stage++)
 		{
 			Hash ^= GetTypeHash(Config.Shaders[Stage].Bindings);
 			Hash ^= Config.Shaders[Stage].Resource;
@@ -631,7 +631,7 @@ public:
 		// (can't modify existing ones), we use this as a check for dirty/need to mem copy on Mobile
 		TArray<uint32>						LastEmulatedUniformBufferSet;
 	};
-	FStagePackedUniformInfo	StagePackedUniformInfo[OGL_NUM_SHADER_STAGES];
+	FStagePackedUniformInfo	StagePackedUniformInfo[CrossCompiler::NUM_SHADER_STAGES];
 
 	GLuint		Program;
 	bool		bUsingTessellation;
@@ -656,7 +656,7 @@ public:
 	void ConfigureShaderStage( int Stage, uint32 FirstUniformBuffer );
 
 	// Make sure GlobalArrays (created from shader reflection) matches our info (from the cross compiler)
-	static inline void SortPackedUniformInfos(const TArray<FPackedUniformInfo>& ReflectedUniformInfos, const TArray<FOpenGLPackedArrayInfo>& PackedGlobalArrays, TArray<FPackedUniformInfo>& OutPackedUniformInfos)
+	static inline void SortPackedUniformInfos(const TArray<FPackedUniformInfo>& ReflectedUniformInfos, const TArray<CrossCompiler::FPackedArrayInfo>& PackedGlobalArrays, TArray<FPackedUniformInfo>& OutPackedUniformInfos)
 	{
 		check(OutPackedUniformInfos.Num() == 0);
 		OutPackedUniformInfos.AddUninitialized(PackedGlobalArrays.Num());
@@ -718,8 +718,8 @@ static int32 CountSetBits(const TBitArray<>& Array)
 
 void FOpenGLLinkedProgram::ConfigureShaderStage( int Stage, uint32 FirstUniformBuffer )
 {
-	static const ANSICHAR StagePrefix[OGL_NUM_SHADER_STAGES] = { 'v', 'p', 'g', 'h', 'd', 'c'};
-	static const GLint FirstTextureUnit[OGL_NUM_SHADER_STAGES] =
+	static const ANSICHAR StagePrefix[CrossCompiler::NUM_SHADER_STAGES] = { 'v', 'p', 'g', 'h', 'd', 'c'};
+	static const GLint FirstTextureUnit[CrossCompiler::NUM_SHADER_STAGES] =
 	{
 		FOpenGL::GetFirstVertexTextureUnit(),
 		FOpenGL::GetFirstPixelTextureUnit(),
@@ -728,7 +728,7 @@ void FOpenGLLinkedProgram::ConfigureShaderStage( int Stage, uint32 FirstUniformB
 		FOpenGL::GetFirstDomainTextureUnit(),
 		FOpenGL::GetFirstComputeTextureUnit()
 	};
-	static const GLint FirstUAVUnit[OGL_NUM_SHADER_STAGES] =
+	static const GLint FirstUAVUnit[CrossCompiler::NUM_SHADER_STAGES] =
 	{
 		OGL_UAV_NOT_SUPPORTED_FOR_GRAPHICS_UNIT,
 		OGL_UAV_NOT_SUPPORTED_FOR_GRAPHICS_UNIT,
@@ -740,7 +740,7 @@ void FOpenGLLinkedProgram::ConfigureShaderStage( int Stage, uint32 FirstUniformB
 	ANSICHAR NameBuffer[10] = {0};
 
 	// verify that only CS uses UAVs
-	check((Stage != OGL_SHADER_STAGE_COMPUTE) ? (CountSetBits(UAVStageNeeds) == 0) : true);
+	check((Stage != CrossCompiler::SHADER_STAGE_COMPUTE) ? (CountSetBits(UAVStageNeeds) == 0) : true);
 
 	SCOPE_CYCLE_COUNTER(STAT_OpenGLShaderBindParameterTime);
 	VERIFY_GL_SCOPE();
@@ -755,9 +755,9 @@ void FOpenGLLinkedProgram::ConfigureShaderStage( int Stage, uint32 FirstUniformB
 		NameBuffer[4] = 0;
 
 		TArray<FPackedUniformInfo> PackedUniformInfos;
-		for (uint8 Index = 0; Index < OGL_PACKED_TYPEINDEX_MAX; ++Index)
+		for (uint8 Index = 0; Index < CrossCompiler::PACKED_TYPEINDEX_MAX; ++Index)
 		{
-			uint8 ArrayIndexType = GLPackedTypeIndexToTypeName(Index);
+			uint8 ArrayIndexType = CrossCompiler::PackedTypeIndexToTypeName(Index);
 			NameBuffer[3] = ArrayIndexType;
 			GLint Location = glGetUniformLocation(Program, NameBuffer);
 			if ((int32)Location != -1)
@@ -784,9 +784,9 @@ void FOpenGLLinkedProgram::ConfigureShaderStage( int Stage, uint32 FirstUniformB
 			ANSICHAR* Str = SetIndex(NameBuffer, 2, UB);
 			*Str++ = '_';
 			Str[1] = 0;
-			for (uint8 Index = 0; Index < OGL_PACKED_TYPEINDEX_MAX; ++Index)
+			for (uint8 Index = 0; Index < CrossCompiler::PACKED_TYPEINDEX_MAX; ++Index)
 			{
-				uint8 ArrayIndexType = GLPackedTypeIndexToTypeName(Index);
+				uint8 ArrayIndexType = CrossCompiler::PackedTypeIndexToTypeName(Index);
 				Str[0] = ArrayIndexType;
 				GLint Location = glGetUniformLocation(Program, NameBuffer);
 				if ((int32)Location != -1)
@@ -1158,38 +1158,38 @@ static FOpenGLLinkedProgram* LinkProgram( const FOpenGLLinkedProgramConfiguratio
 	GLuint Program = glCreateProgram();
 
 	// ensure that compute shaders are always alone
-	check( (Config.Shaders[OGL_SHADER_STAGE_VERTEX].Resource == 0) != (Config.Shaders[OGL_SHADER_STAGE_COMPUTE].Resource == 0));
-	check( (Config.Shaders[OGL_SHADER_STAGE_PIXEL].Resource == 0) != (Config.Shaders[OGL_SHADER_STAGE_COMPUTE].Resource == 0));
+	check( (Config.Shaders[CrossCompiler::SHADER_STAGE_VERTEX].Resource == 0) != (Config.Shaders[CrossCompiler::SHADER_STAGE_COMPUTE].Resource == 0));
+	check( (Config.Shaders[CrossCompiler::SHADER_STAGE_PIXEL].Resource == 0) != (Config.Shaders[CrossCompiler::SHADER_STAGE_COMPUTE].Resource == 0));
 
-	if (Config.Shaders[OGL_SHADER_STAGE_VERTEX].Resource)
+	if (Config.Shaders[CrossCompiler::SHADER_STAGE_VERTEX].Resource)
 	{
-		glAttachShader(Program, Config.Shaders[OGL_SHADER_STAGE_VERTEX].Resource);
+		glAttachShader(Program, Config.Shaders[CrossCompiler::SHADER_STAGE_VERTEX].Resource);
 	}
-	if (Config.Shaders[OGL_SHADER_STAGE_PIXEL].Resource)
+	if (Config.Shaders[CrossCompiler::SHADER_STAGE_PIXEL].Resource)
 	{
-		glAttachShader(Program, Config.Shaders[OGL_SHADER_STAGE_PIXEL].Resource);
+		glAttachShader(Program, Config.Shaders[CrossCompiler::SHADER_STAGE_PIXEL].Resource);
 	}
-	if (Config.Shaders[OGL_SHADER_STAGE_GEOMETRY].Resource)
+	if (Config.Shaders[CrossCompiler::SHADER_STAGE_GEOMETRY].Resource)
 	{
-		glAttachShader(Program, Config.Shaders[OGL_SHADER_STAGE_GEOMETRY].Resource);
+		glAttachShader(Program, Config.Shaders[CrossCompiler::SHADER_STAGE_GEOMETRY].Resource);
 	}
-	if (Config.Shaders[OGL_SHADER_STAGE_HULL].Resource)
+	if (Config.Shaders[CrossCompiler::SHADER_STAGE_HULL].Resource)
 	{
-		glAttachShader(Program, Config.Shaders[OGL_SHADER_STAGE_HULL].Resource);
+		glAttachShader(Program, Config.Shaders[CrossCompiler::SHADER_STAGE_HULL].Resource);
 	}
-	if (Config.Shaders[OGL_SHADER_STAGE_DOMAIN].Resource)
+	if (Config.Shaders[CrossCompiler::SHADER_STAGE_DOMAIN].Resource)
 	{
-		glAttachShader(Program, Config.Shaders[OGL_SHADER_STAGE_DOMAIN].Resource);
+		glAttachShader(Program, Config.Shaders[CrossCompiler::SHADER_STAGE_DOMAIN].Resource);
 	}
-	if (Config.Shaders[OGL_SHADER_STAGE_COMPUTE].Resource)
+	if (Config.Shaders[CrossCompiler::SHADER_STAGE_COMPUTE].Resource)
 	{
-		glAttachShader(Program, Config.Shaders[OGL_SHADER_STAGE_COMPUTE].Resource);
+		glAttachShader(Program, Config.Shaders[CrossCompiler::SHADER_STAGE_COMPUTE].Resource);
 	}
 	
 	// Bind attribute indices.
-	if (Config.Shaders[OGL_SHADER_STAGE_VERTEX].Resource)
+	if (Config.Shaders[CrossCompiler::SHADER_STAGE_VERTEX].Resource)
 	{
-		uint32 Mask = Config.Shaders[OGL_SHADER_STAGE_VERTEX].Bindings.InOutMask;
+		uint32 Mask = Config.Shaders[CrossCompiler::SHADER_STAGE_VERTEX].Bindings.InOutMask;
 		uint32 Index = 0;
 		FCStringAnsi::Strcpy(Buf, "in_ATTRIBUTE");
 		while (Mask)
@@ -1215,9 +1215,9 @@ static FOpenGLLinkedProgram* LinkProgram( const FOpenGLLinkedProgramConfiguratio
 	}
 
 	// Bind frag data locations.
-	if (Config.Shaders[OGL_SHADER_STAGE_PIXEL].Resource)
+	if (Config.Shaders[CrossCompiler::SHADER_STAGE_PIXEL].Resource)
 	{
-		uint32 Mask = (Config.Shaders[OGL_SHADER_STAGE_PIXEL].Bindings.InOutMask) & 0x7fff; // mask out the depth bit
+		uint32 Mask = (Config.Shaders[CrossCompiler::SHADER_STAGE_PIXEL].Bindings.InOutMask) & 0x7fff; // mask out the depth bit
 		uint32 Index = 0;
 		FCStringAnsi::Strcpy(Buf, "out_Target");
 		while (Mask)
@@ -1252,70 +1252,70 @@ static FOpenGLLinkedProgram* LinkProgram( const FOpenGLLinkedProgramConfiguratio
 	FOpenGLLinkedProgram* LinkedProgram = new FOpenGLLinkedProgram;
 	LinkedProgram->Config = Config;
 	LinkedProgram->Program = Program;
-	LinkedProgram->bUsingTessellation = Config.Shaders[OGL_SHADER_STAGE_HULL].Resource && Config.Shaders[OGL_SHADER_STAGE_DOMAIN].Resource;
+	LinkedProgram->bUsingTessellation = Config.Shaders[CrossCompiler::SHADER_STAGE_HULL].Resource && Config.Shaders[CrossCompiler::SHADER_STAGE_DOMAIN].Resource;
 
 	glUseProgram(Program);
 
-	if (Config.Shaders[OGL_SHADER_STAGE_VERTEX].Resource)
+	if (Config.Shaders[CrossCompiler::SHADER_STAGE_VERTEX].Resource)
 	{
 		LinkedProgram->ConfigureShaderStage(
-			OGL_SHADER_STAGE_VERTEX,
+			CrossCompiler::SHADER_STAGE_VERTEX,
 			OGL_FIRST_UNIFORM_BUFFER
 			);
-		check(LinkedProgram->StagePackedUniformInfo[OGL_SHADER_STAGE_VERTEX].PackedUniformInfos.Num() <= Config.Shaders[OGL_SHADER_STAGE_VERTEX].Bindings.PackedGlobalArrays.Num());
+		check(LinkedProgram->StagePackedUniformInfo[CrossCompiler::SHADER_STAGE_VERTEX].PackedUniformInfos.Num() <= Config.Shaders[CrossCompiler::SHADER_STAGE_VERTEX].Bindings.PackedGlobalArrays.Num());
 	}
 
-	if (Config.Shaders[OGL_SHADER_STAGE_PIXEL].Resource)
+	if (Config.Shaders[CrossCompiler::SHADER_STAGE_PIXEL].Resource)
 	{
 		LinkedProgram->ConfigureShaderStage(
-			OGL_SHADER_STAGE_PIXEL,
+			CrossCompiler::SHADER_STAGE_PIXEL,
 			OGL_FIRST_UNIFORM_BUFFER +
-			Config.Shaders[OGL_SHADER_STAGE_VERTEX].Bindings.NumUniformBuffers
+			Config.Shaders[CrossCompiler::SHADER_STAGE_VERTEX].Bindings.NumUniformBuffers
 			);
-		check(LinkedProgram->StagePackedUniformInfo[OGL_SHADER_STAGE_PIXEL].PackedUniformInfos.Num() <= Config.Shaders[OGL_SHADER_STAGE_PIXEL].Bindings.PackedGlobalArrays.Num());
+		check(LinkedProgram->StagePackedUniformInfo[CrossCompiler::SHADER_STAGE_PIXEL].PackedUniformInfos.Num() <= Config.Shaders[CrossCompiler::SHADER_STAGE_PIXEL].Bindings.PackedGlobalArrays.Num());
 	}
 
-	if (Config.Shaders[OGL_SHADER_STAGE_GEOMETRY].Resource)
+	if (Config.Shaders[CrossCompiler::SHADER_STAGE_GEOMETRY].Resource)
 	{
 		LinkedProgram->ConfigureShaderStage(
-			OGL_SHADER_STAGE_GEOMETRY,
+			CrossCompiler::SHADER_STAGE_GEOMETRY,
 			OGL_FIRST_UNIFORM_BUFFER +
-			Config.Shaders[OGL_SHADER_STAGE_VERTEX].Bindings.NumUniformBuffers +
-			Config.Shaders[OGL_SHADER_STAGE_PIXEL].Bindings.NumUniformBuffers
+			Config.Shaders[CrossCompiler::SHADER_STAGE_VERTEX].Bindings.NumUniformBuffers +
+			Config.Shaders[CrossCompiler::SHADER_STAGE_PIXEL].Bindings.NumUniformBuffers
 			);
-		check(LinkedProgram->StagePackedUniformInfo[OGL_SHADER_STAGE_GEOMETRY].PackedUniformInfos.Num() <= Config.Shaders[OGL_SHADER_STAGE_GEOMETRY].Bindings.PackedGlobalArrays.Num());
+		check(LinkedProgram->StagePackedUniformInfo[CrossCompiler::SHADER_STAGE_GEOMETRY].PackedUniformInfos.Num() <= Config.Shaders[CrossCompiler::SHADER_STAGE_GEOMETRY].Bindings.PackedGlobalArrays.Num());
 	}
 
-	if (Config.Shaders[OGL_SHADER_STAGE_HULL].Resource)
+	if (Config.Shaders[CrossCompiler::SHADER_STAGE_HULL].Resource)
 	{
 		LinkedProgram->ConfigureShaderStage(
-			OGL_SHADER_STAGE_HULL,
+			CrossCompiler::SHADER_STAGE_HULL,
 			OGL_FIRST_UNIFORM_BUFFER +
-			Config.Shaders[OGL_SHADER_STAGE_VERTEX].Bindings.NumUniformBuffers +
-			Config.Shaders[OGL_SHADER_STAGE_PIXEL].Bindings.NumUniformBuffers +
-			Config.Shaders[OGL_SHADER_STAGE_GEOMETRY].Bindings.NumUniformBuffers
+			Config.Shaders[CrossCompiler::SHADER_STAGE_VERTEX].Bindings.NumUniformBuffers +
+			Config.Shaders[CrossCompiler::SHADER_STAGE_PIXEL].Bindings.NumUniformBuffers +
+			Config.Shaders[CrossCompiler::SHADER_STAGE_GEOMETRY].Bindings.NumUniformBuffers
 		);
 	}
 
-	if (Config.Shaders[OGL_SHADER_STAGE_DOMAIN].Resource)
+	if (Config.Shaders[CrossCompiler::SHADER_STAGE_DOMAIN].Resource)
 	{
 		LinkedProgram->ConfigureShaderStage(
-			OGL_SHADER_STAGE_DOMAIN,
+			CrossCompiler::SHADER_STAGE_DOMAIN,
 			OGL_FIRST_UNIFORM_BUFFER +
-			Config.Shaders[OGL_SHADER_STAGE_VERTEX].Bindings.NumUniformBuffers +
-			Config.Shaders[OGL_SHADER_STAGE_PIXEL].Bindings.NumUniformBuffers +
-			Config.Shaders[OGL_SHADER_STAGE_GEOMETRY].Bindings.NumUniformBuffers +
-			Config.Shaders[OGL_SHADER_STAGE_HULL].Bindings.NumUniformBuffers
+			Config.Shaders[CrossCompiler::SHADER_STAGE_VERTEX].Bindings.NumUniformBuffers +
+			Config.Shaders[CrossCompiler::SHADER_STAGE_PIXEL].Bindings.NumUniformBuffers +
+			Config.Shaders[CrossCompiler::SHADER_STAGE_GEOMETRY].Bindings.NumUniformBuffers +
+			Config.Shaders[CrossCompiler::SHADER_STAGE_HULL].Bindings.NumUniformBuffers
 		);
 	}
 
-	if (Config.Shaders[OGL_SHADER_STAGE_COMPUTE].Resource)
+	if (Config.Shaders[CrossCompiler::SHADER_STAGE_COMPUTE].Resource)
 	{
 		LinkedProgram->ConfigureShaderStage(
-			OGL_SHADER_STAGE_COMPUTE,
+			CrossCompiler::SHADER_STAGE_COMPUTE,
 			OGL_FIRST_UNIFORM_BUFFER
 			);
-		check(LinkedProgram->StagePackedUniformInfo[OGL_SHADER_STAGE_COMPUTE].PackedUniformInfos.Num() <= Config.Shaders[OGL_SHADER_STAGE_COMPUTE].Bindings.PackedGlobalArrays.Num());
+		check(LinkedProgram->StagePackedUniformInfo[CrossCompiler::SHADER_STAGE_COMPUTE].PackedUniformInfos.Num() <= Config.Shaders[CrossCompiler::SHADER_STAGE_COMPUTE].Bindings.PackedGlobalArrays.Num());
 	}
 #if ENABLE_UNIFORM_BUFFER_LAYOUT_VERIFICATION	
 	VerifyUniformBufferLayouts(Program);
@@ -1341,8 +1341,8 @@ FComputeShaderRHIRef FOpenGLDynamicRHI::RHICreateComputeShader(const TArray<uint
 
 	FOpenGLLinkedProgramConfiguration Config;
 
-	Config.Shaders[OGL_SHADER_STAGE_COMPUTE].Resource = ComputeShader->Resource;
-	Config.Shaders[OGL_SHADER_STAGE_COMPUTE].Bindings = ComputeShader->Bindings;
+	Config.Shaders[CrossCompiler::SHADER_STAGE_COMPUTE].Resource = ComputeShader->Resource;
+	Config.Shaders[CrossCompiler::SHADER_STAGE_COMPUTE].Bindings = ComputeShader->Bindings;
 
 	ComputeShader->LinkedProgram = LinkProgram( Config);
 
@@ -1410,27 +1410,27 @@ FBoundShaderStateRHIRef FOpenGLDynamicRHI::RHICreateBoundShaderState(
 		FOpenGLLinkedProgramConfiguration Config;
 
 		// Fill-in the configuration
-		Config.Shaders[OGL_SHADER_STAGE_VERTEX].Bindings = VertexShader->Bindings;
-		Config.Shaders[OGL_SHADER_STAGE_VERTEX].Resource = VertexShader->Resource;
-		Config.Shaders[OGL_SHADER_STAGE_PIXEL].Bindings = PixelShader->Bindings;
-		Config.Shaders[OGL_SHADER_STAGE_PIXEL].Resource = PixelShader->Resource;
+		Config.Shaders[CrossCompiler::SHADER_STAGE_VERTEX].Bindings = VertexShader->Bindings;
+		Config.Shaders[CrossCompiler::SHADER_STAGE_VERTEX].Resource = VertexShader->Resource;
+		Config.Shaders[CrossCompiler::SHADER_STAGE_PIXEL].Bindings = PixelShader->Bindings;
+		Config.Shaders[CrossCompiler::SHADER_STAGE_PIXEL].Resource = PixelShader->Resource;
 		if (GeometryShader)
 		{
-			Config.Shaders[OGL_SHADER_STAGE_GEOMETRY].Bindings = GeometryShader->Bindings;
-			Config.Shaders[OGL_SHADER_STAGE_GEOMETRY].Resource = GeometryShader->Resource;
+			Config.Shaders[CrossCompiler::SHADER_STAGE_GEOMETRY].Bindings = GeometryShader->Bindings;
+			Config.Shaders[CrossCompiler::SHADER_STAGE_GEOMETRY].Resource = GeometryShader->Resource;
 		}
 
 		if ( FOpenGL::SupportsTessellation())
 		{
 			if ( HullShader)
 			{
-				Config.Shaders[OGL_SHADER_STAGE_HULL].Bindings = HullShader->Bindings;
-				Config.Shaders[OGL_SHADER_STAGE_HULL].Resource = HullShader->Resource;
+				Config.Shaders[CrossCompiler::SHADER_STAGE_HULL].Bindings = HullShader->Bindings;
+				Config.Shaders[CrossCompiler::SHADER_STAGE_HULL].Resource = HullShader->Resource;
 			}
 			if ( DomainShader)
 			{
-				Config.Shaders[OGL_SHADER_STAGE_DOMAIN].Bindings = DomainShader->Bindings;
-				Config.Shaders[OGL_SHADER_STAGE_DOMAIN].Resource = DomainShader->Resource;
+				Config.Shaders[CrossCompiler::SHADER_STAGE_DOMAIN].Bindings = DomainShader->Bindings;
+				Config.Shaders[CrossCompiler::SHADER_STAGE_DOMAIN].Resource = DomainShader->Resource;
 			}
 		}
 
@@ -1750,7 +1750,7 @@ void FOpenGLDynamicRHI::BindPendingComputeShaderState(FOpenGLContextState& Conte
 FOpenGLShaderParameterCache::FOpenGLShaderParameterCache() :
 	GlobalUniformArraySize(-1)
 {
-	for (int32 ArrayIndex = 0; ArrayIndex < OGL_PACKED_TYPEINDEX_MAX; ++ArrayIndex)
+	for (int32 ArrayIndex = 0; ArrayIndex < CrossCompiler::PACKED_TYPEINDEX_MAX; ++ArrayIndex)
 	{
 		PackedGlobalUniformDirty[ArrayIndex].LowVector = 0;
 		PackedGlobalUniformDirty[ArrayIndex].HighVector = MAX_uint32;
@@ -1764,19 +1764,19 @@ void FOpenGLShaderParameterCache::InitializeResources(int32 UniformArraySize)
 	// Uniform arrays have to be multiples of float4s.
 	UniformArraySize = Align(UniformArraySize,SizeOfFloat4);
 
-	PackedGlobalUniforms[0] = (uint8*)FMemory::Malloc(UniformArraySize * OGL_PACKED_TYPEINDEX_MAX);
-	PackedUniformsScratch[0] = (uint8*)FMemory::Malloc(UniformArraySize * OGL_PACKED_TYPEINDEX_MAX);
+	PackedGlobalUniforms[0] = (uint8*)FMemory::Malloc(UniformArraySize * CrossCompiler::PACKED_TYPEINDEX_MAX);
+	PackedUniformsScratch[0] = (uint8*)FMemory::Malloc(UniformArraySize * CrossCompiler::PACKED_TYPEINDEX_MAX);
 
-	FMemory::Memzero(PackedGlobalUniforms[0], UniformArraySize * OGL_PACKED_TYPEINDEX_MAX);
-	FMemory::Memzero(PackedUniformsScratch[0], UniformArraySize * OGL_PACKED_TYPEINDEX_MAX);
-	for (int32 ArrayIndex = 1; ArrayIndex < OGL_PACKED_TYPEINDEX_MAX; ++ArrayIndex)
+	FMemory::Memzero(PackedGlobalUniforms[0], UniformArraySize * CrossCompiler::PACKED_TYPEINDEX_MAX);
+	FMemory::Memzero(PackedUniformsScratch[0], UniformArraySize * CrossCompiler::PACKED_TYPEINDEX_MAX);
+	for (int32 ArrayIndex = 1; ArrayIndex < CrossCompiler::PACKED_TYPEINDEX_MAX; ++ArrayIndex)
 	{
 		PackedGlobalUniforms[ArrayIndex] = PackedGlobalUniforms[ArrayIndex - 1] + UniformArraySize;
 		PackedUniformsScratch[ArrayIndex] = PackedUniformsScratch[ArrayIndex - 1] + UniformArraySize;
 	}
 	GlobalUniformArraySize = UniformArraySize;
 
-	for (int32 ArrayIndex = 0; ArrayIndex < OGL_PACKED_TYPEINDEX_MAX; ++ArrayIndex)
+	for (int32 ArrayIndex = 0; ArrayIndex < CrossCompiler::PACKED_TYPEINDEX_MAX; ++ArrayIndex)
 	{
 		PackedGlobalUniformDirty[ArrayIndex].LowVector = 0;
 		PackedGlobalUniformDirty[ArrayIndex].HighVector = UniformArraySize / SizeOfFloat4;
@@ -1803,7 +1803,7 @@ FOpenGLShaderParameterCache::~FOpenGLShaderParameterCache()
  */
 void FOpenGLShaderParameterCache::MarkAllDirty()
 {
-	for (int32 ArrayIndex = 0; ArrayIndex < OGL_PACKED_TYPEINDEX_MAX; ++ArrayIndex)
+	for (int32 ArrayIndex = 0; ArrayIndex < CrossCompiler::PACKED_TYPEINDEX_MAX; ++ArrayIndex)
 	{
 		PackedGlobalUniformDirty[ArrayIndex].LowVector = 0;
 		PackedGlobalUniformDirty[ArrayIndex].HighVector = GlobalUniformArraySize / SizeOfFloat4;
@@ -1815,9 +1815,9 @@ void FOpenGLShaderParameterCache::MarkAllDirty()
  */
 void FOpenGLShaderParameterCache::Set(uint32 BufferIndexName, uint32 ByteOffset, uint32 NumBytes, const void* NewValues)
 {
-	uint32 BufferIndex = GLPackedTypeNameToTypeIndex(BufferIndexName);
+	uint32 BufferIndex = CrossCompiler::PackedTypeNameToTypeIndex(BufferIndexName);
 	check(GlobalUniformArraySize != -1);
-	check(BufferIndex < OGL_PACKED_TYPEINDEX_MAX);
+	check(BufferIndex < CrossCompiler::PACKED_TYPEINDEX_MAX);
 	check(ByteOffset + NumBytes <= (uint32)GlobalUniformArraySize);
 	PackedGlobalUniformDirty[BufferIndex].LowVector = FMath::Min(PackedGlobalUniformDirty[BufferIndex].LowVector, ByteOffset / SizeOfFloat4);
 	PackedGlobalUniformDirty[BufferIndex].HighVector = FMath::Max(PackedGlobalUniformDirty[BufferIndex].HighVector, (ByteOffset + NumBytes + SizeOfFloat4 - 1) / SizeOfFloat4);
@@ -1848,7 +1848,7 @@ void FOpenGLShaderParameterCache::CommitPackedGlobals(const FOpenGLLinkedProgram
 	 * per-program.
 	 */
 	const TArray<FOpenGLLinkedProgram::FPackedUniformInfo>& PackedUniforms = LinkedProgram->StagePackedUniformInfo[Stage].PackedUniformInfos;
-	const TArray<FOpenGLPackedArrayInfo>& PackedArrays = LinkedProgram->Config.Shaders[Stage].Bindings.PackedGlobalArrays;
+	const TArray<CrossCompiler::FPackedArrayInfo>& PackedArrays = LinkedProgram->Config.Shaders[Stage].Bindings.PackedGlobalArrays;
 	for (int32 PackedUniform = 0; PackedUniform < PackedUniforms.Num(); ++PackedUniform)
 	{
 		const FOpenGLLinkedProgram::FPackedUniformInfo& UniformInfo = PackedUniforms[PackedUniform];
@@ -1870,17 +1870,17 @@ void FOpenGLShaderParameterCache::CommitPackedGlobals(const FOpenGLLinkedProgram
 			Location += StartVector;
 			switch (UniformInfo.Index)
 			{
-			case OGL_PACKED_TYPEINDEX_HIGHP:
-			case OGL_PACKED_TYPEINDEX_MEDIUMP:
-			case OGL_PACKED_TYPEINDEX_LOWP:
+			case CrossCompiler::PACKED_TYPEINDEX_HIGHP:
+			case CrossCompiler::PACKED_TYPEINDEX_MEDIUMP:
+			case CrossCompiler::PACKED_TYPEINDEX_LOWP:
 				glUniform4fv(Location, NumDirtyVectors, (GLfloat*)UniformData);
 				break;
 
-			case OGL_PACKED_TYPEINDEX_INT:
+			case CrossCompiler::PACKED_TYPEINDEX_INT:
 				glUniform4iv(Location, NumDirtyVectors, (GLint*)UniformData);
 				break;
 
-			case OGL_PACKED_TYPEINDEX_UINT:
+			case CrossCompiler::PACKED_TYPEINDEX_UINT:
 				FOpenGL::Uniform4uiv(Location, NumDirtyVectors, (GLuint*)UniformData);
 				break;
 			}
@@ -1976,17 +1976,17 @@ void FOpenGLShaderParameterCache::CommitPackedUniformBuffers(FOpenGLLinkedProgra
 					check(UniformInfo.ArrayType == UBInfo[InfoIndex].TypeName);
 					switch (UniformInfo.Index)
 					{
-					case OGL_PACKED_TYPEINDEX_HIGHP:
-					case OGL_PACKED_TYPEINDEX_MEDIUMP:
-					case OGL_PACKED_TYPEINDEX_LOWP:
+					case CrossCompiler::PACKED_TYPEINDEX_HIGHP:
+					case CrossCompiler::PACKED_TYPEINDEX_MEDIUMP:
+					case CrossCompiler::PACKED_TYPEINDEX_LOWP:
 						glUniform4fv(UniformInfo.Location, NumVectors, (GLfloat*)UniformData);
 						break;
 
-					case OGL_PACKED_TYPEINDEX_INT:
+					case CrossCompiler::PACKED_TYPEINDEX_INT:
 						glUniform4iv(UniformInfo.Location, NumVectors, (GLint*)UniformData);
 						break;
 
-					case OGL_PACKED_TYPEINDEX_UINT:
+					case CrossCompiler::PACKED_TYPEINDEX_UINT:
 						FOpenGL::Uniform4uiv(UniformInfo.Location, NumVectors, (GLuint*)UniformData);
 						break;
 					}
diff --git a/Engine/Source/Runtime/OpenGLDrv/Public/OpenGLShaderResources.h b/Engine/Source/Runtime/OpenGLDrv/Public/OpenGLShaderResources.h
--- a/Engine/Source/Runtime/OpenGLDrv/Public/OpenGLShaderResources.h
+++ b/Engine/Source/Runtime/OpenGLDrv/Public/OpenGLShaderResources.h
@@ -6,90 +6,19 @@
 
 #pragma once
 
-/**
- * OpenGL shader stages.
- */
-enum
-{
-	OGL_SHADER_STAGE_VERTEX = 0,
-	OGL_SHADER_STAGE_PIXEL,
-	OGL_SHADER_STAGE_GEOMETRY,
-	OGL_SHADER_STAGE_HULL,
-	OGL_SHADER_STAGE_DOMAIN,
-	OGL_NUM_NON_COMPUTE_SHADER_STAGES,
-	OGL_SHADER_STAGE_COMPUTE = OGL_NUM_NON_COMPUTE_SHADER_STAGES,
-	OGL_NUM_SHADER_STAGES
-};
+#include "CrossCompilerCommon.h"
 
 /**
  * Shader related constants.
  */
 enum
 {
-	OGL_PACKED_TYPENAME_HIGHP	= 'h',	// Make sure these enums match hlslcc
-	OGL_PACKED_TYPENAME_MEDIUMP	= 'm',
-	OGL_PACKED_TYPENAME_LOWP	= 'l',
-	OGL_PACKED_TYPENAME_INT		= 'i',
-	OGL_PACKED_TYPENAME_UINT	= 'u',
-	OGL_PACKED_TYPENAME_SAMPLER	= 's',
-	OGL_PACKED_TYPENAME_IMAGE	= 'g',
-
-	OGL_PACKED_TYPEINDEX_HIGHP		= 0,
-	OGL_PACKED_TYPEINDEX_MEDIUMP	= 1,
-	OGL_PACKED_TYPEINDEX_LOWP		= 2,
-	OGL_PACKED_TYPEINDEX_INT		= 3,
-	OGL_PACKED_TYPEINDEX_UINT		= 4,
-	OGL_PACKED_TYPEINDEX_MAX		= 5,
-
 	OGL_MAX_UNIFORM_BUFFER_BINDINGS = 12,	// @todo-mobile: Remove me
 	OGL_FIRST_UNIFORM_BUFFER = 0,			// @todo-mobile: Remove me
 	OGL_MAX_COMPUTE_STAGE_UAV_UNITS = 8,	// @todo-mobile: Remove me
 	OGL_UAV_NOT_SUPPORTED_FOR_GRAPHICS_UNIT = -1, // for now, only CS supports UAVs/ images
 };
 
-static FORCEINLINE uint8 GLPackedTypeIndexToTypeName(uint8 ArrayType)
-{
-	switch (ArrayType)
-	{
-		case OGL_PACKED_TYPEINDEX_HIGHP:	return OGL_PACKED_TYPENAME_HIGHP;
-		case OGL_PACKED_TYPEINDEX_MEDIUMP:	return OGL_PACKED_TYPENAME_MEDIUMP;
-		case OGL_PACKED_TYPEINDEX_LOWP:		return OGL_PACKED_TYPENAME_LOWP;
-		case OGL_PACKED_TYPEINDEX_INT:		return OGL_PACKED_TYPENAME_INT;
-		case OGL_PACKED_TYPEINDEX_UINT:		return OGL_PACKED_TYPENAME_UINT;
-	}
-	check(0);
-	return 0;
-}
-
-static FORCEINLINE uint8 GLPackedTypeNameToTypeIndex(uint8 ArrayName)
-{
-	switch (ArrayName)
-	{
-		case OGL_PACKED_TYPENAME_HIGHP:		return OGL_PACKED_TYPEINDEX_HIGHP;
-		case OGL_PACKED_TYPENAME_MEDIUMP:	return OGL_PACKED_TYPEINDEX_MEDIUMP;
-		case OGL_PACKED_TYPENAME_LOWP:		return OGL_PACKED_TYPEINDEX_LOWP;
-		case OGL_PACKED_TYPENAME_INT:		return OGL_PACKED_TYPEINDEX_INT;
-		case OGL_PACKED_TYPENAME_UINT:		return OGL_PACKED_TYPEINDEX_UINT;
-	}
-	check(0);
-	return 0;
-}
-
-struct FOpenGLPackedArrayInfo
-{
-	uint16	Size;		// Bytes
-	uint8	TypeName;	// OGL_PACKED_TYPENAME
-	uint8	TypeIndex;	// OGL_PACKED_TYPE
-};
-
-inline FArchive& operator<<(FArchive& Ar, FOpenGLPackedArrayInfo& Info)
-{
-	Ar << Info.Size;
-	Ar << Info.TypeName;
-	Ar << Info.TypeIndex;
-	return Ar;
-}
-
 struct FOpenGLShaderResourceTable : public FBaseShaderResourceTable
 {
 	/** Mapping of bound Textures to their location in resource tables. */
@@ -124,8 +53,8 @@ inline FArchive& operator<<(FArchive& Ar, FOpenGLShaderResourceTable& SRT)
  */
 struct FOpenGLShaderBindings
 {
-	TArray<TArray<FOpenGLPackedArrayInfo>>	PackedUniformBuffers;
-	TArray<FOpenGLPackedArrayInfo>			PackedGlobalArrays;
+	TArray<TArray<CrossCompiler::FPackedArrayInfo>>	PackedUniformBuffers;
+	TArray<CrossCompiler::FPackedArrayInfo>			PackedGlobalArrays;
 	FOpenGLShaderResourceTable				ShaderResourceTable;
 
 	uint16	InOutMask;
@@ -166,8 +95,8 @@ struct FOpenGLShaderBindings
 
 		for (int32 Item = 0; Item < A.PackedUniformBuffers.Num(); Item++)
 		{
-			const TArray<FOpenGLPackedArrayInfo> &ArrayA = A.PackedUniformBuffers[Item];
-			const TArray<FOpenGLPackedArrayInfo> &ArrayB = B.PackedUniformBuffers[Item];
+			const TArray<CrossCompiler::FPackedArrayInfo> &ArrayA = A.PackedUniformBuffers[Item];
+			const TArray<CrossCompiler::FPackedArrayInfo> &ArrayB = B.PackedUniformBuffers[Item];
 
 			bEqual &= FMemory::Memcmp(ArrayA.GetTypedData(),ArrayB.GetTypedData(),ArrayA.GetTypeSize()*ArrayA.Num()) == 0;
 		}
@@ -189,7 +118,7 @@ struct FOpenGLShaderBindings
 
 		for (int32 Item = 0; Item < Binding.PackedUniformBuffers.Num(); Item++)
 		{
-			const TArray<FOpenGLPackedArrayInfo> &Array = Binding.PackedUniformBuffers[Item];
+			const TArray<CrossCompiler::FPackedArrayInfo> &Array = Binding.PackedUniformBuffers[Item];
 			Hash ^= FCrc::MemCrc_DEPRECATED( Array.GetTypedData(), Array.GetTypeSize()* Array.Num());
 		}
 		return Hash;
@@ -229,7 +158,7 @@ inline FArchive& operator<<(FArchive& Ar, FOpenGLUniformBufferCopyInfo& Info)
 	Ar << Info.DestUBTypeName;
 	if (Ar.IsLoading())
 	{
-		Info.DestUBTypeIndex = GLPackedTypeNameToTypeIndex(Info.DestUBTypeName);
+		Info.DestUBTypeIndex = CrossCompiler::PackedTypeNameToTypeIndex(Info.DestUBTypeName);
 	}
 	Ar << Info.DestOffsetInFloats;
 	Ar << Info.SizeInFloats;
@@ -382,7 +311,7 @@ public:
 
 private:
 	/** CPU memory block for storing uniform values. */
-	uint8* PackedGlobalUniforms[OGL_PACKED_TYPEINDEX_MAX];
+	uint8* PackedGlobalUniforms[CrossCompiler::PACKED_TYPEINDEX_MAX];
 	
 	struct FRange
 	{
@@ -390,10 +319,10 @@ private:
 		uint32	HighVector;
 	};
 	/** Dirty ranges for each uniform array. */
-	FRange	PackedGlobalUniformDirty[OGL_PACKED_TYPEINDEX_MAX];
+	FRange	PackedGlobalUniformDirty[CrossCompiler::PACKED_TYPEINDEX_MAX];
 
 	/** Scratch CPU memory block for uploading packed uniforms. */
-	uint8* PackedUniformsScratch[OGL_PACKED_TYPEINDEX_MAX];
+	uint8* PackedUniformsScratch[CrossCompiler::PACKED_TYPEINDEX_MAX];
 
 	int32 GlobalUniformArraySize;
 };
diff --git a/Engine/Source/Runtime/ShaderCore/Public/CrossCompilerCommon.h b/Engine/Source/Runtime/ShaderCore/Public/CrossCompilerCommon.h
new file mode 100644
--- /dev/null
+++ b/Engine/Source/Runtime/ShaderCore/Public/CrossCompilerCommon.h
@@ -0,0 +1,82 @@
+// Copyright 1998-2014 Epic Games, Inc. All Rights Reserved.
+
+/*=============================================================================
+	CrossCompilerCommon.h: Common functionality between compiler & runtime.
+=============================================================================*/
+
+#pragma once
+
+#include "Core.h"
+
+namespace CrossCompiler
+{
+	enum
+	{
+		SHADER_STAGE_VERTEX = 0,
+		SHADER_STAGE_PIXEL,
+		SHADER_STAGE_GEOMETRY,
+		SHADER_STAGE_HULL,
+		SHADER_STAGE_DOMAIN,
+		NUM_NON_COMPUTE_SHADER_STAGES,
+		SHADER_STAGE_COMPUTE = NUM_NON_COMPUTE_SHADER_STAGES,
+		NUM_SHADER_STAGES,
+
+		PACKED_TYPENAME_HIGHP	= 'h',	// Make sure these enums match hlslcc
+		PACKED_TYPENAME_MEDIUMP	= 'm',
+		PACKED_TYPENAME_LOWP		= 'l',
+		PACKED_TYPENAME_INT		= 'i',
+		PACKED_TYPENAME_UINT		= 'u',
+		PACKED_TYPENAME_SAMPLER	= 's',
+		PACKED_TYPENAME_IMAGE	= 'g',
+
+		PACKED_TYPEINDEX_HIGHP	= 0,
+		PACKED_TYPEINDEX_MEDIUMP	= 1,
+		PACKED_TYPEINDEX_LOWP	= 2,
+		PACKED_TYPEINDEX_INT		= 3,
+		PACKED_TYPEINDEX_UINT	= 4,
+		PACKED_TYPEINDEX_MAX		= 5,
+	};
+
+	static FORCEINLINE uint8 PackedTypeIndexToTypeName(uint8 ArrayType)
+	{
+		switch (ArrayType)
+		{
+		case PACKED_TYPEINDEX_HIGHP:	return PACKED_TYPENAME_HIGHP;
+		case PACKED_TYPEINDEX_MEDIUMP:	return PACKED_TYPENAME_MEDIUMP;
+		case PACKED_TYPEINDEX_LOWP:		return PACKED_TYPENAME_LOWP;
+		case PACKED_TYPEINDEX_INT:		return PACKED_TYPENAME_INT;
+		case PACKED_TYPEINDEX_UINT:		return PACKED_TYPENAME_UINT;
+		}
+		check(0);
+		return 0;
+	}
+
+	static FORCEINLINE uint8 PackedTypeNameToTypeIndex(uint8 ArrayName)
+	{
+		switch (ArrayName)
+		{
+		case PACKED_TYPENAME_HIGHP:		return PACKED_TYPEINDEX_HIGHP;
+		case PACKED_TYPENAME_MEDIUMP:	return PACKED_TYPEINDEX_MEDIUMP;
+		case PACKED_TYPENAME_LOWP:		return PACKED_TYPEINDEX_LOWP;
+		case PACKED_TYPENAME_INT:		return PACKED_TYPEINDEX_INT;
+		case PACKED_TYPENAME_UINT:		return PACKED_TYPEINDEX_UINT;
+		}
+		check(0);
+		return 0;
+	}
+
+	struct FPackedArrayInfo
+	{
+		uint16	Size;		// Bytes
+		uint8	TypeName;	// PACKED_TYPENAME
+		uint8	TypeIndex;	// PACKED_TYPE
+	};
+
+	inline FArchive& operator<<(FArchive& Ar, FPackedArrayInfo& Info)
+	{
+		Ar << Info.Size;
+		Ar << Info.TypeName;
+		Ar << Info.TypeIndex;
+		return Ar;
+	}
+}
