UE4 - Compile fix

diff --git a/Engine/Source/Runtime/OpenGLDrv/Public/OpenGLState.h b/Engine/Source/Runtime/OpenGLDrv/Public/OpenGLState.h
--- a/Engine/Source/Runtime/OpenGLDrv/Public/OpenGLState.h
+++ b/Engine/Source/Runtime/OpenGLDrv/Public/OpenGLState.h
@@ -294,8 +294,8 @@ struct FOpenGLContextState : public FOpenGLCommonState
 	GLuint							OcclusionQuery;
 	GLuint							Program;
 	bool							bUsingTessellation;
-	GLuint 							UniformBuffers[OGL_NUM_SHADER_STAGES*OGL_MAX_UNIFORM_BUFFER_BINDINGS];
-	GLuint 							UniformBufferOffsets[OGL_NUM_SHADER_STAGES*OGL_MAX_UNIFORM_BUFFER_BINDINGS];
+	GLuint 							UniformBuffers[CrossCompiler::NUM_SHADER_STAGES*OGL_MAX_UNIFORM_BUFFER_BINDINGS];
+	GLuint 							UniformBufferOffsets[CrossCompiler::NUM_SHADER_STAGES*OGL_MAX_UNIFORM_BUFFER_BINDINGS];
 	TArray<FOpenGLSamplerState*>	CachedSamplerStates;
 	GLenum							ActiveTexture;
 	bool							bScissorEnabled;
@@ -472,20 +472,20 @@ struct FOpenGLRHIState : public FOpenGLCommonState
 	{
 		check(!ShaderParameters);
 		FOpenGLCommonState::InitializeResources(NumCombinedTextures, NumComputeUAVUnits);
-		ShaderParameters = new FOpenGLShaderParameterCache[OGL_NUM_SHADER_STAGES];
-		ShaderParameters[OGL_SHADER_STAGE_VERTEX].InitializeResources(FOpenGL::GetMaxVertexUniformComponents() * sizeof(float));
-		ShaderParameters[OGL_SHADER_STAGE_PIXEL].InitializeResources(FOpenGL::GetMaxPixelUniformComponents() * sizeof(float));
-		ShaderParameters[OGL_SHADER_STAGE_GEOMETRY].InitializeResources(FOpenGL::GetMaxGeometryUniformComponents() * sizeof(float));
+		ShaderParameters = new FOpenGLShaderParameterCache[CrossCompiler::NUM_SHADER_STAGES];
+		ShaderParameters[CrossCompiler::SHADER_STAGE_VERTEX].InitializeResources(FOpenGL::GetMaxVertexUniformComponents() * sizeof(float));
+		ShaderParameters[CrossCompiler::SHADER_STAGE_PIXEL].InitializeResources(FOpenGL::GetMaxPixelUniformComponents() * sizeof(float));
+		ShaderParameters[CrossCompiler::SHADER_STAGE_GEOMETRY].InitializeResources(FOpenGL::GetMaxGeometryUniformComponents() * sizeof(float));
 		
 		if ( FOpenGL::SupportsTessellation() )
 		{
-			ShaderParameters[OGL_SHADER_STAGE_HULL].InitializeResources(FOpenGL::GetMaxHullUniformComponents() * sizeof(float));
-			ShaderParameters[OGL_SHADER_STAGE_DOMAIN].InitializeResources(FOpenGL::GetMaxDomainUniformComponents() * sizeof(float));
+			ShaderParameters[CrossCompiler::SHADER_STAGE_HULL].InitializeResources(FOpenGL::GetMaxHullUniformComponents() * sizeof(float));
+			ShaderParameters[CrossCompiler::SHADER_STAGE_DOMAIN].InitializeResources(FOpenGL::GetMaxDomainUniformComponents() * sizeof(float));
 		}
 
 		if ( FOpenGL::SupportsComputeShaders() )
 		{
-			ShaderParameters[OGL_SHADER_STAGE_COMPUTE].InitializeResources(FOpenGL::GetMaxComputeUniformComponents() * sizeof(float));
+			ShaderParameters[CrossCompiler::SHADER_STAGE_COMPUTE].InitializeResources(FOpenGL::GetMaxComputeUniformComponents() * sizeof(float));
 		}
 
 		for (int32 Frequency = 0; Frequency < SF_NumFrequencies; ++Frequency)
