#UE4

    not all ES2 implemenations support npot textures with a wrap mode different form clamp to edge.

diff --git a/Engine/Source/Runtime/OpenGLDrv/Private/OpenGLES2.cpp b/Engine/Source/Runtime/OpenGLDrv/Private/OpenGLES2.cpp
--- a/Engine/Source/Runtime/OpenGLDrv/Private/OpenGLES2.cpp
+++ b/Engine/Source/Runtime/OpenGLDrv/Private/OpenGLES2.cpp
@@ -102,7 +102,7 @@ bool FOpenGLES2::bSupportsShaderTextureCubeLod = true;
 bool FOpenGLES2::bSupportsCopyTextureLevels = false;
 
 /** GL_OES_texture_npot */
-bool FOpenGLES2::bSupportsTextureNPOT = true;
+bool FOpenGLES2::bSupportsTextureNPOT = false;
 
 /** GL_EXT_texture_storage */
 bool FOpenGLES2::bSupportsTextureStorageEXT = false;
