#UE4 #HTML5

     Fix how framebuffer completeness is checked. Make sure the minification filter doesn't require mipmaps. This silences firefox webgl errors about mipmap completeness.

diff --git a/Engine/Source/Runtime/OpenGLDrv/Private/HTML5/HTML5OpenGL.cpp b/Engine/Source/Runtime/OpenGLDrv/Private/HTML5/HTML5OpenGL.cpp
--- a/Engine/Source/Runtime/OpenGLDrv/Private/HTML5/HTML5OpenGL.cpp
+++ b/Engine/Source/Runtime/OpenGLDrv/Private/HTML5/HTML5OpenGL.cpp
@@ -89,6 +89,10 @@ void FHTML5OpenGL::ProcessExtensions( const FString& ExtensionsString )
         GLuint tex, fb;
         glGenTextures(1, &tex);
         glBindTexture(GL_TEXTURE_2D, tex);
+		glTexParameteri(GL_TEXTURE_2D, GL_TEXTURE_WRAP_S, GL_REPEAT);
+		glTexParameteri (GL_TEXTURE_2D, GL_TEXTURE_WRAP_T, GL_REPEAT);
+		glTexParameteri (GL_TEXTURE_2D, GL_TEXTURE_MAG_FILTER, GL_LINEAR);
+		glTexParameteri (GL_TEXTURE_2D, GL_TEXTURE_MIN_FILTER, GL_LINEAR);
         glTexImage2D(GL_TEXTURE_2D, 0, GL_RGBA, 32, 32, 0, GL_RGBA, GL_HALF_FLOAT_OES, NULL);
         glGenFramebuffers(1, &fb);
         glBindFramebuffer(GL_FRAMEBUFFER, fb);
@@ -308,6 +312,7 @@ void PlatformRestoreDesktopDisplayMode()
 
 extern "C"
 {
+#if PLATFORM_HTML5_BROWSER
 	// callback from javascript. 
 	void resize_game(int w, int h)
 	{
@@ -323,6 +328,7 @@ extern "C"
 
 		SDL_PushEvent((SDL_Event*)&Event);
 	}
+#endif 
 }
 
 
