Fix bug related to viewport min not getting correctly set in OpenGL. This was found from attempting to repro TTP 340484. Not sure if it is related yet.

diff --git a/Engine/Source/Runtime/OpenGLDrv/Private/OpenGLCommands.cpp b/Engine/Source/Runtime/OpenGLDrv/Private/OpenGLCommands.cpp
--- a/Engine/Source/Runtime/OpenGLDrv/Private/OpenGLCommands.cpp
+++ b/Engine/Source/Runtime/OpenGLDrv/Private/OpenGLCommands.cpp
@@ -387,6 +387,8 @@ inline void FOpenGLDynamicRHI::UpdateScissorRectInOpenGLContext( FOpenGLContextS
 	if( PendingState.bScissorEnabled &&
 		ContextState.Scissor != PendingState.Scissor )
 	{
+		check(PendingState.Scissor.Min.X < PendingState.Scissor.Max.X);
+		check(PendingState.Scissor.Min.Y < PendingState.Scissor.Max.Y);
 		glScissor(PendingState.Scissor.Min.X, PendingState.Scissor.Min.Y, PendingState.Scissor.Max.X - PendingState.Scissor.Min.X, PendingState.Scissor.Max.Y - PendingState.Scissor.Min.Y);
 		ContextState.Scissor = PendingState.Scissor;
 	}
@@ -1590,6 +1592,8 @@ void FOpenGLDynamicRHI::RHISetRenderTargets(
 	else if( NewDepthStencilTargetRHI )
 	{
 		// Set viewport size to new depth target size.
+		PendingState.Viewport.Min.X = 0;
+		PendingState.Viewport.Min.Y = 0;
 		PendingState.Viewport.Max.X = GetOpenGLTextureSizeXFromRHITexture(NewDepthStencilTargetRHI);
 		PendingState.Viewport.Max.Y = GetOpenGLTextureSizeYFromRHITexture(NewDepthStencilTargetRHI);
 	}
