Change the workaround for radr://15553950, TTP# 315197 (Must flush after glTexStorage to avoid crash) to only be used on Nvidia cards where it is needed. This improves texture allocation performance on Macs with GPUs from other vendors.

diff --git a/Engine/Source/Runtime/OpenGLDrv/Private/Mac/OpenGLMac.cpp b/Engine/Source/Runtime/OpenGLDrv/Private/Mac/OpenGLMac.cpp
--- a/Engine/Source/Runtime/OpenGLDrv/Private/Mac/OpenGLMac.cpp
+++ b/Engine/Source/Runtime/OpenGLDrv/Private/Mac/OpenGLMac.cpp
@@ -302,6 +302,15 @@ static TAutoConsoleVariable<int32> CVarMacUseFrameBufferSRGB(
 		ECVF_RenderThreadSafe
 		);
 
+// @todo: remove once Apple fixes radr://15553950, TTP# 315197
+static int32 GMacMustFlushTexStorage = 0;
+static FAutoConsoleVariableRef CVarMacMustFlushTexStorage(
+	TEXT("r.Mac.MustFlushTexStorage"),
+	GMacMustFlushTexStorage,
+	TEXT("If true, flush the OpenGL command stream after calls to glTexStorage* to avoid driver errors, do nothing if false (faster, the default)."),
+	ECVF_RenderThreadSafe
+	);
+
 bool GIsRunningOnIntelCard = false; // @todo: remove once Apple fixes radr://16223045 Changes to the GL separate blend state aren't always respected on Intel cards
 static bool GIsEmulatingTimestamp = false; // @todo: Now crashing on Nvidia cards, but not on AMD...
 
@@ -544,6 +553,10 @@ FPlatformOpenGLContext* PlatformCreateOpenGLContext(FPlatformOpenGLDevice* Devic
 	{
 		GIsRunningOnIntelCard = true;
 	}
+	else if (VendorName.Contains(TEXT("NVIDIA ")))
+	{
+		GMacMustFlushTexStorage = true;
+	}
 	
 	// Renderer IDs matchup to driver kexts, so switching based on them will allow us to target workarouds to many GPUs
 	// which exhibit the same unfortunate driver bugs without having to parse their individual ID strings.
@@ -1789,3 +1802,8 @@ void FMacOpenGL::MacGetQueryObject(GLuint QueryId, EQueryMode QueryMode, GLuint
 		FOpenGL3::GetQueryObject(QueryId, QueryMode, OutResult);
 	}
 }
+
+bool FMacOpenGL::MustFlushTexStorage(void)
+{
+	return GMacMustFlushTexStorage;
+}
diff --git a/Engine/Source/Runtime/OpenGLDrv/Private/Mac/OpenGLMac.h b/Engine/Source/Runtime/OpenGLDrv/Private/Mac/OpenGLMac.h
--- a/Engine/Source/Runtime/OpenGLDrv/Private/Mac/OpenGLMac.h
+++ b/Engine/Source/Runtime/OpenGLDrv/Private/Mac/OpenGLMac.h
@@ -47,6 +47,9 @@ struct FMacOpenGL : public FOpenGL3
 	static void MacEndQuery(GLenum QueryType);
 	static void MacGetQueryObject(GLuint QueryId, EQueryMode QueryMode, uint64 *OutResult);
 	
+	/** Must we employ a workaround for radr://15553950, TTP# 315197 */
+	static bool MustFlushTexStorage(void);
+	
 	static FORCEINLINE void GetQueryObject(GLuint QueryId, EQueryMode QueryMode, GLuint *OutResult)
 	{
 		MacGetQueryObject(QueryId, QueryMode, OutResult);
@@ -102,7 +105,10 @@ struct FMacOpenGL : public FOpenGL3
 		if( glTexStorage3D )
 		{
 			glTexStorage3D( Target, Levels, InternalFormat, Width, Height, Depth);
-			glFlushRenderAPPLE(); // Workaround for radr://15553950, TTP# 315197
+			if(MustFlushTexStorage())
+			{
+				glFlushRenderAPPLE(); // Workaround for radr://15553950, TTP# 315197
+			}
 		}
 		else
 		{
@@ -131,7 +137,10 @@ struct FMacOpenGL : public FOpenGL3
 		if( glTexStorage2D )
 		{
 			glTexStorage2D(Target, Levels, InternalFormat, Width, Height);
-			glFlushRenderAPPLE(); // Workaround for radr://15553950, TTP# 315197
+			if(MustFlushTexStorage())
+			{
+				glFlushRenderAPPLE(); // Workaround for radr://15553950, TTP# 315197
+			}
 			return true;
 		}
 		else
