Refactored string parsing code so it's a little less ugly.

diff --git a/Engine/Source/Runtime/OpenGLDrv/Private/OpenGLShaders.cpp b/Engine/Source/Runtime/OpenGLDrv/Private/OpenGLShaders.cpp
--- a/Engine/Source/Runtime/OpenGLDrv/Private/OpenGLShaders.cpp
+++ b/Engine/Source/Runtime/OpenGLDrv/Private/OpenGLShaders.cpp
@@ -330,11 +330,12 @@ ShaderType* CompileOpenGLShader(const TArray<uint8>& Code)
 		}
 		
 		//remove "#version 100" line,  if found in existing GlslCode. 
+		static const ANSICHAR *Version100String = "#version 100";
 		ANSICHAR* VersionCodePointer = const_cast<ANSICHAR*>(GlslCode);
-		VersionCodePointer = strstr(VersionCodePointer, "#version 100");
+		VersionCodePointer = strstr(VersionCodePointer, Version100String);
 		if(VersionCodePointer != NULL)
 		{
-			for(int index=0; index < 12; index++)
+			for(int index=0; index < strlen( Version100String ); index++)
 			{
 				VersionCodePointer[index] = ' ';
 			}
