Mobile Preview Continued Work:
* For several preview components (cascade + a couple of animation previews), don't create a separate FXSystem for the preview world. Just use the one created internally by the world
* Several bug fixes and improvements to dynamic feature level switch
* Continued conversion of code that references to global feature level.

diff --git a/Engine/Source/Editor/Cascade/Private/Cascade.cpp b/Engine/Source/Editor/Cascade/Private/Cascade.cpp
--- a/Engine/Source/Editor/Cascade/Private/Cascade.cpp
+++ b/Engine/Source/Editor/Cascade/Private/Cascade.cpp
@@ -74,7 +74,7 @@ void FCascade::UnregisterTabSpawners(const TSharedRef<class FTabManager>& TabMan
 
 FCascade::~FCascade()
 {
-	UE_LOG(LogCascade,Log,TEXT("Quitting Cascade. FXSystem=0x%p"),FXSystem);
+	UE_LOG(LogCascade,Log,TEXT("Quitting Cascade. FXSystem=0x%p"),GetFXSystem());
 
 	GEditor->UnregisterForUndo(this);
 	// If the user opened the geometry properties window, we request it be destroyed.
@@ -105,9 +105,6 @@ FCascade::~FCascade()
 
 	DestroyColorPicker();
 
-	FFXSystemInterface::Destroy(FXSystem);
-	FXSystem = NULL;
-
 	if (PreviewViewport.IsValid() && PreviewViewport->GetViewportClient().IsValid())
 	{
 		// Save the preview scene
@@ -195,9 +192,6 @@ void FCascade::InitCascade(const EToolkitMode::Type Mode, const TSharedPtr< clas
 		}
 	}
 
-	// Construct an FX system for this preview.
-	FXSystem = FFXSystemInterface::Create(GRHIFeatureLevel);
-
 	ParticleSystemComponent = ConstructObject<UCascadeParticleSystemComponent>(UCascadeParticleSystemComponent::StaticClass());
 
 	LocalVectorFieldPreviewComponent = ConstructObject<UVectorFieldComponent>(UVectorFieldComponent::StaticClass());
@@ -326,7 +320,10 @@ UVectorFieldComponent* FCascade::GetLocalVectorFieldComponent() const
 
 FFXSystemInterface* FCascade::GetFXSystem() const
 {
-	return FXSystem;
+	check(PreviewViewport.IsValid());
+	auto World = PreviewViewport->GetViewportClient()->GetPreviewScene().GetWorld();
+	check(World);
+	return World->FXSystem;
 }
 
 UParticleEmitter* FCascade::GetSelectedEmitter() const
@@ -1475,7 +1472,7 @@ void FCascade::Tick(float DeltaTime)
 
 		ParticleSystemComponent->CascadeTickComponent(CurrDeltaTime, LEVELTICK_All);
 		ParticleSystemComponent->DoDeferredRenderUpdates_Concurrent();
-		FXSystem->Tick(CurrDeltaTime);
+		GetFXSystem()->Tick(CurrDeltaTime);
 		TotalTime += CurrDeltaTime;
 		ParticleSystem->UpdateTime_Delta = fSaveUpdateDelta;
 
diff --git a/Engine/Source/Editor/Cascade/Private/Cascade.h b/Engine/Source/Editor/Cascade/Private/Cascade.h
--- a/Engine/Source/Editor/Cascade/Private/Cascade.h
+++ b/Engine/Source/Editor/Cascade/Private/Cascade.h
@@ -386,7 +386,6 @@ private:
 	/** Components used for drawing the particle system in the preview viewport */
 	UCascadeParticleSystemComponent* ParticleSystemComponent;
 	UVectorFieldComponent* LocalVectorFieldPreviewComponent;
-	FFXSystemInterface* FXSystem;
 
 	/** Currently selected LOD index */
 	int32 CurrentLODIdx;
diff --git a/Engine/Source/Editor/Cascade/Private/CascadePreviewViewportClient.cpp b/Engine/Source/Editor/Cascade/Private/CascadePreviewViewportClient.cpp
--- a/Engine/Source/Editor/Cascade/Private/CascadePreviewViewportClient.cpp
+++ b/Engine/Source/Editor/Cascade/Private/CascadePreviewViewportClient.cpp
@@ -33,13 +33,10 @@ FCascadeEdPreviewViewportClient::FCascadeEdPreviewViewportClient(TWeakPtr<FCasca
 	UParticleSystem* ParticleSystem = CascadePtr.Pin()->GetParticleSystem();
 	UCascadeParticleSystemComponent* ParticleSystemComponent = CascadePtr.Pin()->GetParticleSystemComponent();
 	UVectorFieldComponent* LocalVectorFieldPreviewComponent = CascadePtr.Pin()->GetLocalVectorFieldComponent();
-	FFXSystemInterface* FXSystem = CascadePtr.Pin()->GetFXSystem();
 	UCascadeOptions* EditorOptions = CascadePtr.Pin()->GetEditorOptions();
 
 	check(EditorOptions);
 
-	CascadePreviewScene.GetScene()->SetFXSystem(FXSystem);
-
 	// Create ParticleSystemComponent to use for preview.
 	ParticleSystemComponent->CascadePreviewViewportPtr = this;
 	ParticleSystemComponent->CastShadow = true;
diff --git a/Engine/Source/Editor/DetailCustomizations/Private/AnimMontageSegmentDetails.cpp b/Engine/Source/Editor/DetailCustomizations/Private/AnimMontageSegmentDetails.cpp
--- a/Engine/Source/Editor/DetailCustomizations/Private/AnimMontageSegmentDetails.cpp
+++ b/Engine/Source/Editor/DetailCustomizations/Private/AnimMontageSegmentDetails.cpp
@@ -5,11 +5,10 @@
 #include "AnimPreviewInstance.h"
 #include "Runtime/Engine/Public/Slate/SceneViewport.h"
 #include "Editor/KismetWidgets/Public/SScrubControlPanel.h"
-#include "Runtime/Engine/Public/FXSystem.h"
 #define LOCTEXT_NAMESPACE "AnimMontageSegmentDetails"
 
 /////////////////////////////////////////////////////////////////////////
-FAnimationSegmentViewportClient::FAnimationSegmentViewportClient(FPreviewScene& InPreviewScene, FFXSystemInterface * FXSystem)
+FAnimationSegmentViewportClient::FAnimationSegmentViewportClient(FPreviewScene& InPreviewScene)
 	: FEditorViewportClient(GLevelEditorModeTools(), &InPreviewScene)
 {
 	SetViewMode(VMI_Lit);
@@ -17,10 +16,7 @@ FAnimationSegmentViewportClient::FAnimationSegmentViewportClient(FPreviewScene&
 	// Always composite editor objects after post processing in the editor
 	EngineShowFlags.CompositeEditorPrimitives = true;
 	EngineShowFlags.DisableAdvancedFeatures();
-
-
-	PreviewScene->GetScene()->SetFXSystem(FXSystem);
-
+	
 	UpdateLighting();
 
 	// Setup defaults for the common draw helper.
@@ -102,8 +98,6 @@ void FAnimMontageSegmentDetails::CustomizeDetails( IDetailLayoutBuilder& DetailB
 
 SAnimationSegmentViewport::~SAnimationSegmentViewport()
 {
-	FFXSystemInterface::Destroy(FXSystem);
-
 	// clean up components
 	if (PreviewComponent)
 	{
@@ -148,7 +142,6 @@ void SAnimationSegmentViewport::Construct(const FArguments& InArgs)
 {
 	TargetSkeleton = InArgs._Skeleton;
 	AnimRef = InArgs._AnimRef;
-	FXSystem = FFXSystemInterface::Create(GRHIFeatureLevel);
 	
 	AnimRefPropertyHandle = InArgs._AnimRefPropertyHandle;
 	StartTimePropertyHandle = InArgs._StartTimePropertyHandle;
@@ -193,7 +186,7 @@ void SAnimationSegmentViewport::Construct(const FArguments& InArgs)
 	
 
 	// Create a viewport client
-	LevelViewportClient	= MakeShareable( new FAnimationSegmentViewportClient(PreviewScene, FXSystem) );
+	LevelViewportClient	= MakeShareable( new FAnimationSegmentViewportClient(PreviewScene) );
 
 	LevelViewportClient->ViewportType = LVT_Perspective;
 	LevelViewportClient->bSetListenerPosition = false;
diff --git a/Engine/Source/Editor/DetailCustomizations/Private/AnimMontageSegmentDetails.h b/Engine/Source/Editor/DetailCustomizations/Private/AnimMontageSegmentDetails.h
--- a/Engine/Source/Editor/DetailCustomizations/Private/AnimMontageSegmentDetails.h
+++ b/Engine/Source/Editor/DetailCustomizations/Private/AnimMontageSegmentDetails.h
@@ -67,8 +67,6 @@ private:
 	UAnimSequenceBase *AnimRef;
 
 	FPreviewScene PreviewScene;
-	class FFXSystemInterface* FXSystem;
-
 	class UDebugSkelMeshComponent* PreviewComponent;
 
 	TSharedPtr<STextBlock> Description;
@@ -152,7 +150,7 @@ protected:
 class FAnimationSegmentViewportClient : public FEditorViewportClient
 {
 public:
-	FAnimationSegmentViewportClient(FPreviewScene& InPreviewScene, FFXSystemInterface * FXSystem);
+	FAnimationSegmentViewportClient(FPreviewScene& InPreviewScene);
 
 	// FlEditorViewportClient interface
 	virtual FSceneInterface* GetScene() const override;
diff --git a/Engine/Source/Editor/DetailCustomizations/Private/AnimSequenceDetails.cpp b/Engine/Source/Editor/DetailCustomizations/Private/AnimSequenceDetails.cpp
--- a/Engine/Source/Editor/DetailCustomizations/Private/AnimSequenceDetails.cpp
+++ b/Engine/Source/Editor/DetailCustomizations/Private/AnimSequenceDetails.cpp
@@ -6,7 +6,6 @@
 #include "AnimPreviewInstance.h"
 #include "Runtime/Engine/Public/Slate/SceneViewport.h"
 #include "Editor/KismetWidgets/Public/SScrubControlPanel.h"
-#include "Runtime/Engine/Public/FXSystem.h"
 
 #define LOCTEXT_NAMESPACE	"AnimSequenceDetails"
 
@@ -292,8 +291,6 @@ TSharedPtr<FString> FAnimSequenceDetails::GetRetargetSourceString(FName Retarget
 // based on SAnimationSegmentViewport
 SAnimationRefPoseViewport::~SAnimationRefPoseViewport()
 {
-	FFXSystemInterface::Destroy(FXSystem);
-
 	// clean up components
 	if (PreviewComponent)
 	{
@@ -337,8 +334,6 @@ SAnimationRefPoseViewport::SAnimationRefPoseViewport()
 BEGIN_SLATE_FUNCTION_BUILD_OPTIMIZATION
 void SAnimationRefPoseViewport::Construct(const FArguments& InArgs)
 {
-	FXSystem = FFXSystemInterface::Create(GRHIFeatureLevel);
-
 	TargetSkeleton = InArgs._Skeleton;
 	AnimRefPropertyHandle = InArgs._AnimRefPropertyHandle;
 	RefPoseTypeHandle = InArgs._RefPoseTypeHandle;
@@ -378,7 +373,7 @@ void SAnimationRefPoseViewport::Construct(const FArguments& InArgs)
 
 
 	// Create a viewport client
-	LevelViewportClient	= MakeShareable(new FAnimationSegmentViewportClient(PreviewScene, FXSystem));
+	LevelViewportClient	= MakeShareable(new FAnimationSegmentViewportClient(PreviewScene));
 
 	LevelViewportClient->ViewportType = LVT_Perspective;
 	LevelViewportClient->bSetListenerPosition = false;
diff --git a/Engine/Source/Editor/LevelEditor/Private/LevelEditorActions.cpp b/Engine/Source/Editor/LevelEditor/Private/LevelEditorActions.cpp
--- a/Engine/Source/Editor/LevelEditor/Private/LevelEditorActions.cpp
+++ b/Engine/Source/Editor/LevelEditor/Private/LevelEditorActions.cpp
@@ -452,12 +452,12 @@ bool FLevelEditorActionCallbacks::IsMaterialQualityLevelChecked( EMaterialQualit
 void FLevelEditorActionCallbacks::SetFeatureLevelPreview(ERHIFeatureLevel::Type InPreviewFeatureLevel)
 {
 	UWorld::ChangeAllWorldFeatureLevels(InPreviewFeatureLevel);
+//	GetWorld()->ChangeFeatureLevel(InPreviewFeatureLevel); // The ultimate goal, but needs removal of the global feature level first!
 }
 
 bool FLevelEditorActionCallbacks::IsFeatureLevelPreviewChecked(ERHIFeatureLevel::Type InPreviewFeatureLevel)
 {
-	// For now, we just stay stuck on the system max feature level
-	return InPreviewFeatureLevel == GRHIFeatureLevel;
+	return InPreviewFeatureLevel == GetWorld()->FeatureLevel;
 }
 
 void FLevelEditorActionCallbacks::ConfigureLightingBuildOptions( const FLightingBuildOptions& Options )
diff --git a/Engine/Source/Editor/LevelEditor/Private/SLevelViewport.cpp b/Engine/Source/Editor/LevelEditor/Private/SLevelViewport.cpp
--- a/Engine/Source/Editor/LevelEditor/Private/SLevelViewport.cpp
+++ b/Engine/Source/Editor/LevelEditor/Private/SLevelViewport.cpp
@@ -2999,11 +2999,16 @@ FText SLevelViewport::GetCurrentFeatureLevelPreviewText( bool bDrawOnlyLabel ) c
 	}
 	else
 	{
-		if (GRHIFeatureLevel != GMaxRHIFeatureLevel)
+		auto* World = GetWorld();
+		if (World != nullptr)
 		{
-			FName FeatureLevelName;
-			GetFeatureLevelName(GRHIFeatureLevel, FeatureLevelName);
-			FeatureLevelText = FText::Format(LOCTEXT("FeatureLevel", "{0}"), FText::FromName(FeatureLevelName));
+			const auto FeatureLevel = World->FeatureLevel;
+			if (FeatureLevel != GMaxRHIFeatureLevel)
+			{
+				FName FeatureLevelName;
+				GetFeatureLevelName(FeatureLevel, FeatureLevelName);
+				FeatureLevelText = FText::Format(LOCTEXT("FeatureLevel", "{0}"), FText::FromName(FeatureLevelName));
+			}
 		}
 	}
 	
@@ -3065,7 +3070,14 @@ EVisibility SLevelViewport::GetCurrentLevelTextVisibility() const
 
 EVisibility SLevelViewport::GetCurrentFeatureLevelPreviewTextVisibility() const
 {
-	return (GRHIFeatureLevel != GMaxRHIFeatureLevel) ? EVisibility::SelfHitTestInvisible : EVisibility::Collapsed;
+	if (GetWorld())
+	{
+		return (GetWorld()->FeatureLevel != GMaxRHIFeatureLevel) ? EVisibility::SelfHitTestInvisible : EVisibility::Collapsed;
+	}
+	else
+	{
+		return EVisibility::Collapsed;
+	}
 }
 
 void SLevelViewport::OnSetViewportConfiguration(FName ConfigurationName)
diff --git a/Engine/Source/Runtime/Engine/Classes/Engine/World.h b/Engine/Source/Runtime/Engine/Classes/Engine/World.h
--- a/Engine/Source/Runtime/Engine/Classes/Engine/World.h
+++ b/Engine/Source/Runtime/Engine/Classes/Engine/World.h
@@ -629,10 +629,12 @@ public:
 
 	/** Change the feature level that this world is current rendering with */
 	void ChangeFeatureLevel(ERHIFeatureLevel::Type InFeatureLevel);
-
 	static void ChangeAllWorldFeatureLevels(ERHIFeatureLevel::Type InFeatureLevel);
 
 private:
+	/** Change the feature level that this world is current rendering with */
+	void ChangeWorldFeatureLevel(ERHIFeatureLevel::Type InFeatureLevel);
+
 	/** List of all the controllers in the world. */
 	TArray<TAutoWeakObjectPtr<class AController> >	ControllerList;
 
diff --git a/Engine/Source/Runtime/Engine/Private/LightMap.cpp b/Engine/Source/Runtime/Engine/Private/LightMap.cpp
--- a/Engine/Source/Runtime/Engine/Private/LightMap.cpp
+++ b/Engine/Source/Runtime/Engine/Private/LightMap.cpp
@@ -1265,7 +1265,7 @@ UTexture2D* FLightMap2D::GetTexture(uint32 BasisIndex)
  */
 bool FLightMap2D::IsValid(uint32 BasisIndex) const
 {
-	return bAllowHighQualityLightMaps ? BasisIndex == 0 : BasisIndex == 1;
+	return AllowHighQualityLightmaps() ? BasisIndex == 0 : BasisIndex == 1;
 }
 
 struct FLegacyLightMapTextureInfo
@@ -1421,14 +1421,14 @@ void FLightMap2D::Serialize(FArchive& Ar)
 
 FLightMapInteraction FLightMap2D::GetInteraction() const
 {
-	int32 LightmapIndex = bAllowHighQualityLightMaps ? 0 : 1;
+	int32 LightmapIndex = AllowHighQualityLightmaps() ? 0 : 1;
 
 	bool bValidTextures = Textures[ LightmapIndex ] && Textures[ LightmapIndex ]->Resource;
 
 	// When the FLightMap2D is first created, the textures aren't set, so that case needs to be handled.
 	if(bValidTextures)
 	{
-		return FLightMapInteraction::Texture(Textures, SkyOcclusionTexture, ScaleVectors, AddVectors, CoordinateScale, CoordinateBias, true);
+		return FLightMapInteraction::Texture(Textures, SkyOcclusionTexture, ScaleVectors, AddVectors, CoordinateScale, CoordinateBias, AllowHighQualityLightmaps());
 	}
 
 	return FLightMapInteraction::None();
diff --git a/Engine/Source/Runtime/Engine/Private/SceneManagement.cpp b/Engine/Source/Runtime/Engine/Private/SceneManagement.cpp
--- a/Engine/Source/Runtime/Engine/Private/SceneManagement.cpp
+++ b/Engine/Source/Runtime/Engine/Private/SceneManagement.cpp
@@ -28,20 +28,16 @@ FLightMapInteraction FLightMapInteraction::Texture(
 	const FVector4* InCoefficientAdds,
 	const FVector2D& InCoordinateScale,
 	const FVector2D& InCoordinateBias,
-	bool bAllowHighQualityLightMaps)
+	bool bUseHighQualityLightMaps)
 {
 	FLightMapInteraction Result;
 	Result.Type = LMIT_Texture;
 
-	bool bUseHighQualityLightMaps = AllowHighQualityLightmaps();
-
 #if ALLOW_LQ_LIGHTMAPS && ALLOW_HQ_LIGHTMAPS
 	// however, if simple and directional are allowed, then we must use the value passed in,
 	// and then cache the number as well
-	bUseHighQualityLightMaps = bUseHighQualityLightMaps && bAllowHighQualityLightMaps;
-
 	Result.bAllowHighQualityLightMaps = bUseHighQualityLightMaps;
-	if (bAllowHighQualityLightMaps)
+	if (bUseHighQualityLightMaps)
 	{
 		Result.NumLightmapCoefficients = NUM_HQ_LIGHTMAP_COEF;
 	}
diff --git a/Engine/Source/Runtime/Engine/Private/SkeletalRenderGPUSkin.cpp b/Engine/Source/Runtime/Engine/Private/SkeletalRenderGPUSkin.cpp
--- a/Engine/Source/Runtime/Engine/Private/SkeletalRenderGPUSkin.cpp
+++ b/Engine/Source/Runtime/Engine/Private/SkeletalRenderGPUSkin.cpp
@@ -1060,7 +1060,7 @@ FPreviousPerBoneMotionBlur::FPreviousPerBoneMotionBlur()
 
 void FPreviousPerBoneMotionBlur::InitResources()
 {
-	if (GRHIFeatureLevel >= ERHIFeatureLevel::SM4)
+	if (GMaxRHIFeatureLevel >= ERHIFeatureLevel::SM4)
 	{
 		for(uint32 i = 0; i < PER_BONE_BUFFER_COUNT; ++i)
 		{
diff --git a/Engine/Source/Runtime/Engine/Private/SystemSettings.cpp b/Engine/Source/Runtime/Engine/Private/SystemSettings.cpp
--- a/Engine/Source/Runtime/Engine/Private/SystemSettings.cpp
+++ b/Engine/Source/Runtime/Engine/Private/SystemSettings.cpp
@@ -340,7 +340,7 @@ void FSystemSettings::Initialize( bool bSetupForEditor )
 
 	ApplyOverrides();
 
-	bInitialUseHighQualityLightmaps = AllowHighQualityLightmaps();
+	bInitialUseHighQualityLightmaps = AllowHighQualityLightmaps(GMaxRHIFeatureLevel);
 
 	IConsoleManager::Get().RegisterConsoleVariableSink(FConsoleCommandDelegate::CreateRaw(this, &FSystemSettings::CVarSink));
 
diff --git a/Engine/Source/Runtime/Engine/Private/UnrealEngine.cpp b/Engine/Source/Runtime/Engine/Private/UnrealEngine.cpp
--- a/Engine/Source/Runtime/Engine/Private/UnrealEngine.cpp
+++ b/Engine/Source/Runtime/Engine/Private/UnrealEngine.cpp
@@ -9785,13 +9785,19 @@ static TAutoConsoleVariable<int32> CVarAllowHighQualityLightMaps(
 	TEXT("If set to 1, allow high quality lightmaps which don't bake in direct lighting of stationary lights"),
 	ECVF_RenderThreadSafe | ECVF_ReadOnly);
 
-bool AllowHighQualityLightmaps()
+
+bool AllowHighQualityLightmaps(ERHIFeatureLevel::Type FeatureLevel)
 {
 	return FPlatformProperties::SupportsHighQualityLightmaps()
-		&& (GRHIFeatureLevel > ERHIFeatureLevel::ES2)
+		&& (FeatureLevel > ERHIFeatureLevel::ES2)
 		&& (CVarAllowHighQualityLightMaps.GetValueOnAnyThread() != 0);
 }
 
+bool AllowHighQualityLightmaps()
+{
+	return AllowHighQualityLightmaps(GRHIFeatureLevel);
+}
+
 // Helper function for changing system resolution via the r.setres console command
 void FSystemResolution::RequestResolutionChange(int32 InResX, int32 InResY, EWindowMode::Type InWindowMode)
 {
diff --git a/Engine/Source/Runtime/Engine/Private/VectorField.cpp b/Engine/Source/Runtime/Engine/Private/VectorField.cpp
--- a/Engine/Source/Runtime/Engine/Private/VectorField.cpp
+++ b/Engine/Source/Runtime/Engine/Private/VectorField.cpp
@@ -759,7 +759,7 @@ public:
 	{
 		check(IsInRenderingThread());
 
-		if (GRHIFeatureLevel == ERHIFeatureLevel::SM5 && AnimatedVectorField && AnimatedVectorField->Texture && AnimatedVectorField->Texture->Resource)
+		if (GetFeatureLevel() == ERHIFeatureLevel::SM5 && AnimatedVectorField && AnimatedVectorField->Texture && AnimatedVectorField->Texture->Resource)
 		{
 			SCOPED_DRAW_EVENT(AnimateVectorField, DEC_PARTICLE);
 
diff --git a/Engine/Source/Runtime/Engine/Private/World.cpp b/Engine/Source/Runtime/Engine/Private/World.cpp
--- a/Engine/Source/Runtime/Engine/Private/World.cpp
+++ b/Engine/Source/Runtime/Engine/Private/World.cpp
@@ -5104,7 +5104,7 @@ void UWorld::GetLandscapeTexturesAndMaterials(ULevel* Level, TArray<UObject*>& O
 	}
 }
 
-void UWorld::ChangeFeatureLevel(ERHIFeatureLevel::Type InFeatureLevel)
+void UWorld::ChangeWorldFeatureLevel(ERHIFeatureLevel::Type InFeatureLevel)
 {
 	if (InFeatureLevel != FeatureLevel)
 	{
@@ -5124,6 +5124,36 @@ void UWorld::ChangeFeatureLevel(ERHIFeatureLevel::Type InFeatureLevel)
 	}
 }
 
+void UWorld::ChangeFeatureLevel(ERHIFeatureLevel::Type InFeatureLevel)
+{
+	if (InFeatureLevel != FeatureLevel)
+	{
+		FlushRenderingCommands();
+
+		FGlobalComponentReregisterContext RecreateComponents;
+
+		GCurrentRHIFeatureLevel = InFeatureLevel;
+		GRHIShaderPlatform = GShaderPlatformForFeatureLevel[InFeatureLevel];
+
+		UMaterial::AllMaterialsCacheResourceShadersForRendering();
+		UMaterialInstance::AllMaterialsCacheResourceShadersForRendering();
+		GetGlobalShaderMap(GRHIShaderPlatform, false);
+		GShaderCompilingManager->ProcessAsyncResults(false, true);
+
+		//invalidate global bound shader states so they will be created with the new shaders the next time they are set (in SetGlobalBoundShaderState)
+		for (TLinkedList<FGlobalBoundShaderStateResource*>::TIterator It(FGlobalBoundShaderStateResource::GetGlobalBoundShaderStateList()); It; It.Next())
+		{
+			BeginUpdateResourceRHI(*It);
+		}
+
+		// try iterating over all worlds here
+		ChangeWorldFeatureLevel(InFeatureLevel);
+
+		FOutputDeviceNull Ar;
+		RecompileShaders(TEXT("CHANGED"), Ar);
+	}
+}
+
 void UWorld::ChangeAllWorldFeatureLevels(ERHIFeatureLevel::Type InFeatureLevel)
 {
 	bool bSomethingChanged = false;
@@ -5166,7 +5196,7 @@ void UWorld::ChangeAllWorldFeatureLevels(ERHIFeatureLevel::Type InFeatureLevel)
 
 			if (World->FeatureLevel != InFeatureLevel)
 			{
-				World->ChangeFeatureLevel(InFeatureLevel);
+				World->ChangeWorldFeatureLevel(InFeatureLevel);
 			}
 		}
 
diff --git a/Engine/Source/Runtime/Engine/Public/ParticleVertexFactory.h b/Engine/Source/Runtime/Engine/Public/ParticleVertexFactory.h
--- a/Engine/Source/Runtime/Engine/Public/ParticleVertexFactory.h
+++ b/Engine/Source/Runtime/Engine/Public/ParticleVertexFactory.h
@@ -59,6 +59,8 @@ public:
 		return bInUse;
 	}
 
+	ERHIFeatureLevel::Type GetFeatureLevel() const { check(HasValidFeatureLevel());  return FRenderResource::GetFeatureLevel(); }
+
 private:
 
 	/** The type of the vertex factory. */
diff --git a/Engine/Source/Runtime/Engine/Public/PrecomputedLightVolume.h b/Engine/Source/Runtime/Engine/Public/PrecomputedLightVolume.h
--- a/Engine/Source/Runtime/Engine/Public/PrecomputedLightVolume.h
+++ b/Engine/Source/Runtime/Engine/Public/PrecomputedLightVolume.h
@@ -154,6 +154,11 @@ public:
 
 	void ApplyWorldOffset(const FVector& InOffset, bool bWorldShift);
 
+	// temporary, for ES2 preview verification.
+	FORCEINLINE bool IsUsingHighQualityLightMap() const
+	{
+		return OctreeForRendering == &HighQualityLightmapOctree;
+	}
 private:
 
 	bool bInitialized;
diff --git a/Engine/Source/Runtime/Engine/Public/UnrealEngine.h b/Engine/Source/Runtime/Engine/Public/UnrealEngine.h
--- a/Engine/Source/Runtime/Engine/Public/UnrealEngine.h
+++ b/Engine/Source/Runtime/Engine/Public/UnrealEngine.h
@@ -407,6 +407,7 @@ struct FCachedSystemScalabilityCVars
 };
 
 ENGINE_API bool AllowHighQualityLightmaps();
+ENGINE_API bool AllowHighQualityLightmaps(ERHIFeatureLevel::Type FeatureLevel);
 
 ENGINE_API const FCachedSystemScalabilityCVars& GetCachedScalabilityCVars();
 
diff --git a/Engine/Source/Runtime/OpenGLDrv/Private/OpenGLCommands.cpp b/Engine/Source/Runtime/OpenGLDrv/Private/OpenGLCommands.cpp
--- a/Engine/Source/Runtime/OpenGLDrv/Private/OpenGLCommands.cpp
+++ b/Engine/Source/Runtime/OpenGLDrv/Private/OpenGLCommands.cpp
@@ -411,7 +411,7 @@ void FOpenGLDynamicRHI::RHISetBoundShaderState( FBoundShaderStateRHIParamRef Bou
 
 void FOpenGLDynamicRHI::RHISetUAVParameter(FComputeShaderRHIParamRef ComputeShaderRHI,uint32 UAVIndex,FUnorderedAccessViewRHIParamRef UnorderedAccessViewRHI)
 {
-	check(GRHIFeatureLevel >= ERHIFeatureLevel::SM5);
+	check(GMaxRHIFeatureLevel >= ERHIFeatureLevel::SM5);
 
 	VERIFY_GL_SCOPE();
 	if(UnorderedAccessViewRHI)
@@ -719,7 +719,7 @@ void FOpenGLDynamicRHI::SetupUAVsForDraw( FOpenGLContextState& ContextState, con
 
 void FOpenGLDynamicRHI::CachedSetupUAVStage( FOpenGLContextState& ContextState, GLint UAVIndex, GLenum Format, GLuint Resource)
 {
-	check(GRHIFeatureLevel >= ERHIFeatureLevel::SM5);
+	check(GMaxRHIFeatureLevel >= ERHIFeatureLevel::SM5);
 
 	if( ContextState.UAVs[UAVIndex].Format == Format && ContextState.Textures[UAVIndex].Resource == Resource)
 	{
@@ -774,7 +774,7 @@ void FOpenGLDynamicRHI::RHISetShaderResourceViewParameter(FVertexShaderRHIParamR
 
 void FOpenGLDynamicRHI::RHISetShaderResourceViewParameter(FComputeShaderRHIParamRef ComputeShaderRHI,uint32 TextureIndex,FShaderResourceViewRHIParamRef SRVRHI)
 {
-	check(GRHIFeatureLevel >= ERHIFeatureLevel::SM5);
+	check(GMaxRHIFeatureLevel >= ERHIFeatureLevel::SM5);
 	VERIFY_GL_SCOPE();
 	check(FOpenGL::SupportsResourceView());
 	DYNAMIC_CAST_OPENGLRESOURCE(ShaderResourceView,SRV);
@@ -795,7 +795,7 @@ void FOpenGLDynamicRHI::RHISetShaderResourceViewParameter(FHullShaderRHIParamRef
 {
 	VALIDATE_BOUND_SHADER(HullShaderRHI);
 
-	check(GRHIFeatureLevel >= ERHIFeatureLevel::SM5);
+	check(GMaxRHIFeatureLevel >= ERHIFeatureLevel::SM5);
 	VERIFY_GL_SCOPE();
 	check(FOpenGL::SupportsResourceView());
 	DYNAMIC_CAST_OPENGLRESOURCE(ShaderResourceView,SRV);
@@ -815,7 +815,7 @@ void FOpenGLDynamicRHI::RHISetShaderResourceViewParameter(FDomainShaderRHIParamR
 {
 	VALIDATE_BOUND_SHADER(DomainShaderRHI);
 
-	check(GRHIFeatureLevel >= ERHIFeatureLevel::SM5);
+	check(GMaxRHIFeatureLevel >= ERHIFeatureLevel::SM5);
 	VERIFY_GL_SCOPE();
 	check(FOpenGL::SupportsResourceView());
 	DYNAMIC_CAST_OPENGLRESOURCE(ShaderResourceView,SRV);
@@ -864,7 +864,7 @@ void FOpenGLDynamicRHI::RHISetShaderTexture(FHullShaderRHIParamRef HullShaderRHI
 {
 	VALIDATE_BOUND_SHADER(HullShaderRHI);
 	
-	check(GRHIFeatureLevel >= ERHIFeatureLevel::SM5);
+	check(GMaxRHIFeatureLevel >= ERHIFeatureLevel::SM5);
 	VERIFY_GL_SCOPE();
 	FOpenGLTextureBase* NewTexture = GetOpenGLTextureFromRHITexture(NewTextureRHI);
 	InternalSetShaderTexture(NewTexture, FOpenGL::GetFirstHullTextureUnit() + TextureIndex, NewTexture->Target, NewTexture->Resource, NewTextureRHI->GetNumMips(), -1);
@@ -874,7 +874,7 @@ void FOpenGLDynamicRHI::RHISetShaderTexture(FDomainShaderRHIParamRef DomainShade
 {
 	VALIDATE_BOUND_SHADER(DomainShaderRHI);
 
-	check(GRHIFeatureLevel >= ERHIFeatureLevel::SM5);
+	check(GMaxRHIFeatureLevel >= ERHIFeatureLevel::SM5);
 	VERIFY_GL_SCOPE();
 	FOpenGLTextureBase* NewTexture = GetOpenGLTextureFromRHITexture(NewTextureRHI);
 	InternalSetShaderTexture(NewTexture, FOpenGL::GetFirstDomainTextureUnit() + TextureIndex, NewTexture->Target, NewTexture->Resource, NewTextureRHI->GetNumMips(), -1);
@@ -918,7 +918,7 @@ void FOpenGLDynamicRHI::RHISetShaderSampler(FHullShaderRHIParamRef HullShaderRHI
 {
 	VALIDATE_BOUND_SHADER(HullShaderRHI);
 
-	check(GRHIFeatureLevel >= ERHIFeatureLevel::SM5);
+	check(GMaxRHIFeatureLevel >= ERHIFeatureLevel::SM5);
 	VERIFY_GL_SCOPE();
 	DYNAMIC_CAST_OPENGLRESOURCE(SamplerState,NewState);
 	if (FOpenGL::SupportsSamplerObjects())
@@ -934,7 +934,7 @@ void FOpenGLDynamicRHI::RHISetShaderSampler(FDomainShaderRHIParamRef DomainShade
 {
 	VALIDATE_BOUND_SHADER(DomainShaderRHI);
 
-	check(GRHIFeatureLevel >= ERHIFeatureLevel::SM5);
+	check(GMaxRHIFeatureLevel >= ERHIFeatureLevel::SM5);
 	VERIFY_GL_SCOPE();
 	DYNAMIC_CAST_OPENGLRESOURCE(SamplerState,NewState);
 	if (FOpenGL::SupportsSamplerObjects())
@@ -965,7 +965,7 @@ void FOpenGLDynamicRHI::RHISetShaderSampler(FGeometryShaderRHIParamRef GeometryS
 
 void FOpenGLDynamicRHI::RHISetShaderTexture(FComputeShaderRHIParamRef ComputeShaderRHI,uint32 TextureIndex,FTextureRHIParamRef NewTextureRHI)
 {
-	check(GRHIFeatureLevel >= ERHIFeatureLevel::SM5);
+	check(GMaxRHIFeatureLevel >= ERHIFeatureLevel::SM5);
 	VERIFY_GL_SCOPE();
 	FOpenGLTextureBase* NewTexture = GetOpenGLTextureFromRHITexture(NewTextureRHI);
 	InternalSetShaderTexture(NewTexture, FOpenGL::GetFirstComputeTextureUnit() + TextureIndex, NewTexture->Target, NewTexture->Resource, NewTextureRHI->GetNumMips(), -1);
@@ -997,7 +997,7 @@ void FOpenGLDynamicRHI::RHISetShaderUniformBuffer(FVertexShaderRHIParamRef Verte
 void FOpenGLDynamicRHI::RHISetShaderUniformBuffer(FHullShaderRHIParamRef HullShaderRHI,uint32 BufferIndex,FUniformBufferRHIParamRef BufferRHI)
 {
 	VALIDATE_BOUND_SHADER(HullShaderRHI);
-	check(GRHIFeatureLevel >= ERHIFeatureLevel::SM5);
+	check(GMaxRHIFeatureLevel >= ERHIFeatureLevel::SM5);
 	PendingState.BoundUniformBuffers[SF_Hull][BufferIndex] = BufferRHI;
 	PendingState.DirtyUniformBuffers[SF_Hull] |= 1 << BufferIndex;
 }
@@ -1005,7 +1005,7 @@ void FOpenGLDynamicRHI::RHISetShaderUniformBuffer(FHullShaderRHIParamRef HullSha
 void FOpenGLDynamicRHI::RHISetShaderUniformBuffer(FDomainShaderRHIParamRef DomainShaderRHI,uint32 BufferIndex,FUniformBufferRHIParamRef BufferRHI)
 {
 	VALIDATE_BOUND_SHADER(DomainShaderRHI);
-	check(GRHIFeatureLevel >= ERHIFeatureLevel::SM5);
+	check(GMaxRHIFeatureLevel >= ERHIFeatureLevel::SM5);
 	PendingState.BoundUniformBuffers[SF_Domain][BufferIndex] = BufferRHI;
 	PendingState.DirtyUniformBuffers[SF_Domain] |= 1 << BufferIndex;
 }
@@ -1019,7 +1019,7 @@ void FOpenGLDynamicRHI::RHISetShaderUniformBuffer(FGeometryShaderRHIParamRef Geo
 
 void FOpenGLDynamicRHI::RHISetShaderSampler(FComputeShaderRHIParamRef ComputeShaderRHI,uint32 SamplerIndex,FSamplerStateRHIParamRef NewStateRHI)
 {
-	check(GRHIFeatureLevel >= ERHIFeatureLevel::SM5);
+	check(GMaxRHIFeatureLevel >= ERHIFeatureLevel::SM5);
 	VERIFY_GL_SCOPE();
 	DYNAMIC_CAST_OPENGLRESOURCE(SamplerState,NewState);
 	FOpenGL::BindSampler(FOpenGL::GetFirstComputeTextureUnit() + SamplerIndex, NewState->Resource);
@@ -1034,7 +1034,7 @@ void FOpenGLDynamicRHI::RHISetShaderUniformBuffer(FPixelShaderRHIParamRef PixelS
 
 void FOpenGLDynamicRHI::RHISetShaderUniformBuffer(FComputeShaderRHIParamRef ComputeShaderRHI,uint32 BufferIndex,FUniformBufferRHIParamRef BufferRHI)
 {
-	check(GRHIFeatureLevel >= ERHIFeatureLevel::SM5);
+	check(GMaxRHIFeatureLevel >= ERHIFeatureLevel::SM5);
 	PendingState.BoundUniformBuffers[SF_Compute][BufferIndex] = BufferRHI;
 	PendingState.DirtyUniformBuffers[SF_Compute] |= 1 << BufferIndex;
 }
@@ -1059,7 +1059,7 @@ void FOpenGLDynamicRHI::RHISetShaderParameter(FHullShaderRHIParamRef HullShaderR
 {
 	VALIDATE_BOUND_SHADER(HullShaderRHI);
 
-	check(GRHIFeatureLevel >= ERHIFeatureLevel::SM5);
+	check(GMaxRHIFeatureLevel >= ERHIFeatureLevel::SM5);
 	VERIFY_GL_SCOPE();
 	PendingState.ShaderParameters[CrossCompiler::SHADER_STAGE_HULL].Set(BufferIndex, BaseIndex, NumBytes, NewValue);
 }
@@ -1068,7 +1068,7 @@ void FOpenGLDynamicRHI::RHISetShaderParameter(FDomainShaderRHIParamRef DomainSha
 {
 	VALIDATE_BOUND_SHADER(DomainShaderRHI);
 
-	check(GRHIFeatureLevel >= ERHIFeatureLevel::SM5);
+	check(GMaxRHIFeatureLevel >= ERHIFeatureLevel::SM5);
 	VERIFY_GL_SCOPE();
 	PendingState.ShaderParameters[CrossCompiler::SHADER_STAGE_DOMAIN].Set(BufferIndex, BaseIndex, NumBytes, NewValue);
 }
@@ -1083,7 +1083,7 @@ void FOpenGLDynamicRHI::RHISetShaderParameter(FGeometryShaderRHIParamRef Geometr
 
 void FOpenGLDynamicRHI::RHISetShaderParameter(FComputeShaderRHIParamRef ComputeShaderRHI,uint32 BufferIndex,uint32 BaseIndex,uint32 NumBytes,const void* NewValue)
 { 
-	check(GRHIFeatureLevel >= ERHIFeatureLevel::SM5);
+	check(GMaxRHIFeatureLevel >= ERHIFeatureLevel::SM5);
 	PendingState.ShaderParameters[CrossCompiler::SHADER_STAGE_COMPUTE].Set(BufferIndex, BaseIndex, NumBytes, NewValue);
 }
 
@@ -2105,7 +2105,7 @@ void FOpenGLDynamicRHI::CommitNonComputeShaderConstants()
 void FOpenGLDynamicRHI::CommitComputeShaderConstants(FComputeShaderRHIParamRef ComputeShaderRHI)
 {
 	VERIFY_GL_SCOPE();
-	check(GRHIFeatureLevel >= ERHIFeatureLevel::SM5);
+	check(GMaxRHIFeatureLevel >= ERHIFeatureLevel::SM5);
 
 	DYNAMIC_CAST_OPENGLRESOURCE(ComputeShader,ComputeShader);
 	const int32 Stage = CrossCompiler::SHADER_STAGE_COMPUTE;
@@ -2892,7 +2892,7 @@ void FOpenGLDynamicRHI::RHIClearMRT(bool bClearColor,int32 NumClearColors,const
 {
 	VERIFY_GL_SCOPE();
 
-	check((GRHIFeatureLevel >= ERHIFeatureLevel::SM5 ) || !PendingState.bFramebufferSetupInvalid);
+	check((GMaxRHIFeatureLevel >= ERHIFeatureLevel::SM5) || !PendingState.bFramebufferSetupInvalid);
 
 	if (bClearColor)
 	{
@@ -3120,7 +3120,7 @@ void FOpenGLDynamicRHI::RHISetComputeShader(FComputeShaderRHIParamRef ComputeSha
 	if ( FOpenGL::SupportsComputeShaders() )
 	{
 		VERIFY_GL_SCOPE();
-		check(GRHIFeatureLevel >= ERHIFeatureLevel::SM5);
+		check(GMaxRHIFeatureLevel >= ERHIFeatureLevel::SM5);
 
 		PendingState.CurrentComputeShader = ComputeShaderRHI;
 
@@ -3145,7 +3145,7 @@ void FOpenGLDynamicRHI::RHIDispatchComputeShader(uint32 ThreadGroupCountX, uint3
 	{
 		VERIFY_GL_SCOPE();
 
-		check(GRHIFeatureLevel >= ERHIFeatureLevel::SM5);
+		check(GMaxRHIFeatureLevel >= ERHIFeatureLevel::SM5);
 
 		FComputeShaderRHIParamRef ComputeShaderRHI = PendingState.CurrentComputeShader;
 		check(ComputeShaderRHI);
@@ -3175,7 +3175,7 @@ void FOpenGLDynamicRHI::RHIDispatchIndirectComputeShader(FVertexBufferRHIParamRe
 	if ( FOpenGL::SupportsComputeShaders() )
 	{
 		VERIFY_GL_SCOPE();
-		check(GRHIFeatureLevel >= ERHIFeatureLevel::SM5)
+		check(GMaxRHIFeatureLevel >= ERHIFeatureLevel::SM5)
 
 		FComputeShaderRHIParamRef ComputeShaderRHI = PendingState.CurrentComputeShader;
 		check(ComputeShaderRHI);
diff --git a/Engine/Source/Runtime/OpenGLDrv/Private/OpenGLDevice.cpp b/Engine/Source/Runtime/OpenGLDrv/Private/OpenGLDevice.cpp
--- a/Engine/Source/Runtime/OpenGLDrv/Private/OpenGLDevice.cpp
+++ b/Engine/Source/Runtime/OpenGLDrv/Private/OpenGLDevice.cpp
@@ -611,7 +611,7 @@ static void InitRHICapabilitiesForGL()
 	GRHIShaderPlatform = FOpenGL::GetShaderPlatform();
 
 	FString FeatureLevelName;
-	GetFeatureLevelName(GRHIFeatureLevel, FeatureLevelName);
+	GetFeatureLevelName(GMaxRHIFeatureLevel, FeatureLevelName);
 	FString ShaderPlatformName = LegacyShaderPlatformToShaderFormat(GRHIShaderPlatform).ToString();
 
 	UE_LOG(LogRHI, Log, TEXT("OpenGL MajorVersion = %d, MinorVersion = %d, ShaderPlatform = %s, FeatureLevel = %s"), MajorVersion, MinorVersion, *ShaderPlatformName, *FeatureLevelName);
@@ -641,7 +641,7 @@ static void InitRHICapabilitiesForGL()
 	GSupportsVertexInstancing = true;
 	GHardwareHiddenSurfaceRemoval = FOpenGL::HasHardwareHiddenSurfaceRemoval();
 
-	GShaderPlatformForFeatureLevel[ERHIFeatureLevel::ES2] = (GRHIFeatureLevel == ERHIFeatureLevel::ES2) ? GRHIShaderPlatform : SP_OPENGL_PCES2;
+	GShaderPlatformForFeatureLevel[ERHIFeatureLevel::ES2] = (GMaxRHIFeatureLevel == ERHIFeatureLevel::ES2) ? GRHIShaderPlatform : SP_OPENGL_PCES2;
 	GShaderPlatformForFeatureLevel[ERHIFeatureLevel::SM3] = SP_NumPlatforms;
 	GShaderPlatformForFeatureLevel[ERHIFeatureLevel::SM4] = SP_OPENGL_SM4;
 	GShaderPlatformForFeatureLevel[ERHIFeatureLevel::SM5] = SP_OPENGL_SM5;
@@ -652,7 +652,7 @@ static void InitRHICapabilitiesForGL()
 	GProjectionSignY = 1.0f;
 
 	// Disable texture streaming on ES2 unless we have the GL_APPLE_copy_texture_levels extension
-	if (GRHIFeatureLevel == ERHIFeatureLevel::ES2 && !FOpenGL::SupportsCopyTextureLevels())
+	if (GMaxRHIFeatureLevel == ERHIFeatureLevel::ES2 && !FOpenGL::SupportsCopyTextureLevels())
 	{
 		GRHISupportsTextureStreaming = false;
 	}
diff --git a/Engine/Source/Runtime/OpenGLDrv/Private/OpenGLRenderTarget.cpp b/Engine/Source/Runtime/OpenGLDrv/Private/OpenGLRenderTarget.cpp
--- a/Engine/Source/Runtime/OpenGLDrv/Private/OpenGLRenderTarget.cpp
+++ b/Engine/Source/Runtime/OpenGLDrv/Private/OpenGLRenderTarget.cpp
@@ -307,8 +307,8 @@ void FOpenGLDynamicRHI::RHICopyToResolveTarget(FTextureRHIParamRef SourceTexture
 
 	VERIFY_GL_SCOPE();
 
-	check(GRHIFeatureLevel >= ERHIFeatureLevel::SM5 || ResolveParams.SourceArrayIndex == 0);
-	check(GRHIFeatureLevel >= ERHIFeatureLevel::SM5 || ResolveParams.DestArrayIndex == 0);
+	check(GMaxRHIFeatureLevel >= ERHIFeatureLevel::SM5 || ResolveParams.SourceArrayIndex == 0);
+	check(GMaxRHIFeatureLevel >= ERHIFeatureLevel::SM5 || ResolveParams.DestArrayIndex == 0);
 
 	FOpenGLTextureBase* SourceTexture = GetOpenGLTextureFromRHITexture(SourceTextureRHI);
 	FOpenGLTextureBase* DestTexture = GetOpenGLTextureFromRHITexture(DestTextureRHI);
@@ -710,7 +710,7 @@ void FOpenGLDynamicRHI::RHIReadSurfaceFloatData(FTextureRHIParamRef TextureRHI,F
 {
 	VERIFY_GL_SCOPE();
 	// Not supported on older APIs
-	check(GRHIFeatureLevel >= ERHIFeatureLevel::SM5 || ArrayIndex == 0);
+	check(GMaxRHIFeatureLevel >= ERHIFeatureLevel::SM5 || ArrayIndex == 0);
 
 	check( FOpenGL::SupportsFloatReadSurface() );
 	FOpenGLTextureBase* Texture = GetOpenGLTextureFromRHITexture(TextureRHI);
@@ -819,7 +819,7 @@ void FOpenGLDynamicRHI::BindPendingFramebuffer( FOpenGLContextState& ContextStat
 {
 	VERIFY_GL_SCOPE();
 
-	check((GRHIFeatureLevel >= ERHIFeatureLevel::SM5) || !PendingState.bFramebufferSetupInvalid);
+	check((GMaxRHIFeatureLevel >= ERHIFeatureLevel::SM5) || !PendingState.bFramebufferSetupInvalid);
 
 	if (ContextState.Framebuffer != PendingState.Framebuffer)
 	{
diff --git a/Engine/Source/Runtime/OpenGLDrv/Private/OpenGLShaders.cpp b/Engine/Source/Runtime/OpenGLDrv/Private/OpenGLShaders.cpp
--- a/Engine/Source/Runtime/OpenGLDrv/Private/OpenGLShaders.cpp
+++ b/Engine/Source/Runtime/OpenGLDrv/Private/OpenGLShaders.cpp
@@ -491,13 +491,13 @@ FGeometryShaderRHIRef FOpenGLDynamicRHI::RHICreateGeometryShader(const TArray<ui
 
 FHullShaderRHIRef FOpenGLDynamicRHI::RHICreateHullShader(const TArray<uint8>& Code)
 {
-	check(GRHIFeatureLevel >= ERHIFeatureLevel::SM5);
+	check(GMaxRHIFeatureLevel >= ERHIFeatureLevel::SM5);
 	return CompileOpenGLShader<FOpenGLHullShader>(Code);
 }
 
 FDomainShaderRHIRef FOpenGLDynamicRHI::RHICreateDomainShader(const TArray<uint8>& Code)
 {
-	check(GRHIFeatureLevel >= ERHIFeatureLevel::SM5);
+	check(GMaxRHIFeatureLevel >= ERHIFeatureLevel::SM5);
 	return CompileOpenGLShader<FOpenGLDomainShader>(Code);
 }
 
@@ -1067,7 +1067,7 @@ static void VerifyUniformBufferLayouts(GLuint Program)
 			glGetActiveUniformBlockiv(Program, BlockIndex, GL_UNIFORM_BLOCK_REFERENCED_BY_VERTEX_SHADER, &ReferencedByVS);
 			glGetActiveUniformBlockiv(Program, BlockIndex, GL_UNIFORM_BLOCK_REFERENCED_BY_FRAGMENT_SHADER, &ReferencedByPS);
 			glGetActiveUniformBlockiv(Program, BlockIndex, GL_UNIFORM_BLOCK_REFERENCED_BY_GEOMETRY_SHADER, &ReferencedByGS);
-			if (GRHIFeatureLevel >= ERHIFeatureLevel::SM5)
+			if (GMaxRHIFeatureLevel >= ERHIFeatureLevel::SM5)
 			{
 				glGetActiveUniformBlockiv(Program, BlockIndex, GL_UNIFORM_BLOCK_REFERENCED_BY_TESS_CONTROL_SHADER, &ReferencedByHS);
 				glGetActiveUniformBlockiv(Program, BlockIndex, GL_UNIFORM_BLOCK_REFERENCED_BY_TESS_EVALUATION_SHADER, &ReferencedByDS);
@@ -1326,7 +1326,7 @@ static FOpenGLLinkedProgram* LinkProgram( const FOpenGLLinkedProgramConfiguratio
 
 FComputeShaderRHIRef FOpenGLDynamicRHI::RHICreateComputeShader(const TArray<uint8>& Code)
 {
-	check(GRHIFeatureLevel >= ERHIFeatureLevel::SM5);
+	check(GMaxRHIFeatureLevel >= ERHIFeatureLevel::SM5);
 	
 	FOpenGLComputeShader* ComputeShader = CompileOpenGLShader<FOpenGLComputeShader>(Code);
 	const ANSICHAR* GlslCode = NULL;
diff --git a/Engine/Source/Runtime/OpenGLDrv/Private/OpenGLTexture.cpp b/Engine/Source/Runtime/OpenGLDrv/Private/OpenGLTexture.cpp
--- a/Engine/Source/Runtime/OpenGLDrv/Private/OpenGLTexture.cpp
+++ b/Engine/Source/Runtime/OpenGLDrv/Private/OpenGLTexture.cpp
@@ -249,7 +249,7 @@ FRHITexture* FOpenGLDynamicRHI::CreateOpenGLTexture(uint32 SizeX,uint32 SizeY,bo
 	check( bArrayTexture != (ArraySize == 1));
 #endif
 
-	if (GRHIFeatureLevel <= ERHIFeatureLevel::ES2)
+	if (GMaxRHIFeatureLevel <= ERHIFeatureLevel::ES2)
 	{
 		// Remove sRGB read flag when not supported
 		Flags &= ~TexCreate_SRGB;
@@ -1295,7 +1295,7 @@ FTexture2DArrayRHIRef FOpenGLDynamicRHI::RHICreateTexture2DArray(uint32 SizeX,ui
 		NumMips = FindMaxMipmapLevel(SizeX, SizeY);
 	}
 
-	if (GRHIFeatureLevel <= ERHIFeatureLevel::ES2)
+	if (GMaxRHIFeatureLevel <= ERHIFeatureLevel::ES2)
 	{
 		// Remove sRGB read flag when not supported
 		Flags &= ~TexCreate_SRGB;
@@ -1404,7 +1404,7 @@ FTexture3DRHIRef FOpenGLDynamicRHI::RHICreateTexture3D(uint32 SizeX,uint32 SizeY
 		NumMips = FindMaxMipmapLevel(SizeX, SizeY, SizeZ);
 	}
 
-	if (GRHIFeatureLevel <= ERHIFeatureLevel::ES2)
+	if (GMaxRHIFeatureLevel <= ERHIFeatureLevel::ES2)
 	{
 		// Remove sRGB read flag when not supported
 		Flags &= ~TexCreate_SRGB;
diff --git a/Engine/Source/Runtime/RHI/Public/RHIUtilities.h b/Engine/Source/Runtime/RHI/Public/RHIUtilities.h
--- a/Engine/Source/Runtime/RHI/Public/RHIUtilities.h
+++ b/Engine/Source/Runtime/RHI/Public/RHIUtilities.h
@@ -18,7 +18,7 @@ struct FRWBuffer
 
 	void Initialize(uint32 BytesPerElement, uint32 NumElements, EPixelFormat Format, uint32 AdditionalUsage = 0)
 	{
-		check(GRHIFeatureLevel == ERHIFeatureLevel::SM5);
+		check(GMaxRHIFeatureLevel == ERHIFeatureLevel::SM5);
 		NumBytes = BytesPerElement * NumElements;
 		FRHIResourceCreateInfo CreateInfo;
 		Buffer = RHICreateVertexBuffer(NumBytes, BUF_UnorderedAccess | BUF_ShaderResource | AdditionalUsage, CreateInfo);
@@ -46,7 +46,7 @@ struct FReadBuffer
 
 	void Initialize(uint32 BytesPerElement, uint32 NumElements, EPixelFormat Format, uint32 AdditionalUsage = 0)
 	{
-		check(GRHIFeatureLevel >= ERHIFeatureLevel::SM4);
+		check(GMaxRHIFeatureLevel >= ERHIFeatureLevel::SM4);
 		NumBytes = BytesPerElement * NumElements;
 		FRHIResourceCreateInfo CreateInfo;
 		Buffer = RHICreateVertexBuffer(NumBytes, BUF_ShaderResource | AdditionalUsage, CreateInfo);
@@ -73,7 +73,7 @@ struct FRWBufferStructured
 
 	void Initialize(uint32 BytesPerElement, uint32 NumElements, uint32 AdditionalUsage = 0, bool bUseUavCounter = false, bool bAppendBuffer = false)
 	{
-		check(GRHIFeatureLevel == ERHIFeatureLevel::SM5);
+		check(GMaxRHIFeatureLevel == ERHIFeatureLevel::SM5);
 		NumBytes = BytesPerElement * NumElements;
 		FRHIResourceCreateInfo CreateInfo;
 		Buffer = RHICreateStructuredBuffer(BytesPerElement, NumBytes, BUF_UnorderedAccess | BUF_ShaderResource | AdditionalUsage, CreateInfo);
@@ -103,7 +103,7 @@ struct FRWBufferByteAddress
 	void Initialize(uint32 InNumBytes, uint32 AdditionalUsage = 0)
 	{
 		NumBytes = InNumBytes;
-		check(GRHIFeatureLevel == ERHIFeatureLevel::SM5);
+		check(GMaxRHIFeatureLevel == ERHIFeatureLevel::SM5);
 		check( NumBytes % 4 == 0 );
 		FRHIResourceCreateInfo CreateInfo;
 		Buffer = RHICreateStructuredBuffer(4, NumBytes, BUF_UnorderedAccess | BUF_ShaderResource | BUF_ByteAddressBuffer | AdditionalUsage, CreateInfo);
@@ -244,7 +244,7 @@ inline void RHICreateTargetableShaderResourceCube(
 	check(TargetableTextureFlags & (TexCreate_RenderTargetable | TexCreate_DepthStencilTargetable));
 
 	// ES2 doesn't support resolve operations.
-    bForceSeparateTargetAndShaderResource &= (GRHIFeatureLevel > ERHIFeatureLevel::ES2);
+	bForceSeparateTargetAndShaderResource &= (GMaxRHIFeatureLevel > ERHIFeatureLevel::ES2);
 
 	if(!bForceSeparateTargetAndShaderResource/* && GSupportsRenderDepthTargetableShaderResources*/)
 	{
diff --git a/Engine/Source/Runtime/RenderCore/Public/RenderResource.h b/Engine/Source/Runtime/RenderCore/Public/RenderResource.h
--- a/Engine/Source/Runtime/RenderCore/Public/RenderResource.h
+++ b/Engine/Source/Runtime/RenderCore/Public/RenderResource.h
@@ -89,7 +89,10 @@ public:
 	// Accessors.
 	FORCEINLINE bool IsInitialized() const { return bInitialized; }
 
-	FORCEINLINE ERHIFeatureLevel::Type GetFeatureLevel() const { return FeatureLevel == ERHIFeatureLevel::Num ? GMaxRHIFeatureLevel : FeatureLevel; }
+protected:
+	// This is used during mobile editor preview refactor, this will eventually be replaced with a parameter to InitRHI() etc..
+	ERHIFeatureLevel::Type GetFeatureLevel() const { return FeatureLevel == ERHIFeatureLevel::Num ? GMaxRHIFeatureLevel : FeatureLevel; }
+	FORCEINLINE bool HasValidFeatureLevel() const { return FeatureLevel < ERHIFeatureLevel::Num; }
 
 private:
 
@@ -285,6 +288,7 @@ struct FMipBiasFade
 class FTexture : public FRenderResource
 {
 public:
+
 	/** The texture's RHI resource. */
 	FTextureRHIRef		TextureRHI;
 
diff --git a/Engine/Source/Runtime/Renderer/Private/BasePassRendering.h b/Engine/Source/Runtime/Renderer/Private/BasePassRendering.h
--- a/Engine/Source/Runtime/Renderer/Private/BasePassRendering.h
+++ b/Engine/Source/Runtime/Renderer/Private/BasePassRendering.h
@@ -877,7 +877,7 @@ void ProcessBasePassMesh(
 				: FLightMapInteraction();
 
 			// force LQ lightmaps based on system settings
-			const bool bAllowHighQualityLightMaps = AllowHighQualityLightmaps() && LightMapInteraction.AllowsHighQualityLightmaps();
+			const bool bAllowHighQualityLightMaps = AllowHighQualityLightmaps(Parameters.FeatureLevel) && LightMapInteraction.AllowsHighQualityLightmaps();
 
 			switch(LightMapInteraction.GetType())
 			{
diff --git a/Engine/Source/Runtime/Renderer/Private/DeferredShadingRenderer.cpp b/Engine/Source/Runtime/Renderer/Private/DeferredShadingRenderer.cpp
--- a/Engine/Source/Runtime/Renderer/Private/DeferredShadingRenderer.cpp
+++ b/Engine/Source/Runtime/Renderer/Private/DeferredShadingRenderer.cpp
@@ -71,7 +71,7 @@ FDeferredShadingSceneRenderer::FDeferredShadingSceneRenderer(const FSceneViewFam
 	{
 		// Use a depth only pass if we are using full blown HQ lightmaps
 		// Otherwise base pass pixel shaders will be cheap and there will be little benefit to rendering a depth only pass
-		if (AllowHighQualityLightmaps() || !ViewFamily.EngineShowFlags.Lighting)
+		if (AllowHighQualityLightmaps(FeatureLevel) || !ViewFamily.EngineShowFlags.Lighting)
 		{
 			EarlyZPassMode = DDM_None;
 		}
@@ -582,8 +582,6 @@ void FDeferredShadingSceneRenderer::Render(FRHICommandListImmediate& RHICmdList)
 
 	SCOPED_DRAW_EVENT(Scene,DEC_SCENE_ITEMS);
 
-	const auto FeatureLevel = Scene->GetFeatureLevel();
-
 	// Initialize global system textures (pass-through if already initialized).
 	GSystemTextures.InitializeTextures(RHICmdList, FeatureLevel);
 
diff --git a/Engine/Source/Runtime/Renderer/Private/DistanceFieldSurfaceCacheLighting.cpp b/Engine/Source/Runtime/Renderer/Private/DistanceFieldSurfaceCacheLighting.cpp
--- a/Engine/Source/Runtime/Renderer/Private/DistanceFieldSurfaceCacheLighting.cpp
+++ b/Engine/Source/Runtime/Renderer/Private/DistanceFieldSurfaceCacheLighting.cpp
@@ -2819,7 +2819,7 @@ bool FDeferredShadingSceneRenderer::RenderDistanceFieldAOSurfaceCache(FRHIComman
 	const FViewInfo& View = Views[0];
 
 	if (GDistanceFieldAO 
-		&& GRHIFeatureLevel >= ERHIFeatureLevel::SM5
+		&& FeatureLevel >= ERHIFeatureLevel::SM5
 		&& DoesPlatformSupportDistanceFieldAO(GRHIShaderPlatform)
 		&& Views.Num() == 1
 		// ViewState is used to cache tile intersection resources which have to be sized based on the view
@@ -3106,7 +3106,7 @@ void FDeferredShadingSceneRenderer::RenderDynamicSkyLighting(FRHICommandListImme
 		&& !Scene->SkyLight->bWantsStaticShadowing
 		&& !Scene->SkyLight->bHasStaticLighting
 		&& ViewFamily.EngineShowFlags.SkyLighting
-		&& GRHIFeatureLevel >= ERHIFeatureLevel::SM4
+		&& FeatureLevel >= ERHIFeatureLevel::SM4
 		&& !IsSimpleDynamicLightingEnabled() 
 		&& !ViewFamily.EngineShowFlags.VisualizeLightCulling)
 	{
diff --git a/Engine/Source/Runtime/Renderer/Private/ForwardBasePassRendering.cpp b/Engine/Source/Runtime/Renderer/Private/ForwardBasePassRendering.cpp
--- a/Engine/Source/Runtime/Renderer/Private/ForwardBasePassRendering.cpp
+++ b/Engine/Source/Runtime/Renderer/Private/ForwardBasePassRendering.cpp
@@ -96,6 +96,9 @@ void FBasePassForwardOpaqueDrawingPolicyFactory::AddStaticMesh(FRHICommandList&
 	// Only draw opaque materials.
 	if( !IsTranslucentBlendMode(BlendMode) )
 	{
+		// following check moved from ProcessBasePassMeshForForwardShading to avoid passing feature level.
+		check(!AllowHighQualityLightmaps(Scene->GetFeatureLevel()));
+
 		ProcessBasePassMeshForForwardShading(
 			RHICmdList,
 			FProcessBasePassMeshParameters(
diff --git a/Engine/Source/Runtime/Renderer/Private/ForwardBasePassRendering.h b/Engine/Source/Runtime/Renderer/Private/ForwardBasePassRendering.h
--- a/Engine/Source/Runtime/Renderer/Private/ForwardBasePassRendering.h
+++ b/Engine/Source/Runtime/Renderer/Private/ForwardBasePassRendering.h
@@ -465,8 +465,6 @@ void ProcessBasePassMeshForForwardShading(
 		? Parameters.Mesh.LCI->GetLightMapInteraction() 
 		: FLightMapInteraction();
 
-	check(!AllowHighQualityLightmaps());
-
 	if (LightMapInteraction.GetType() == LMIT_Texture)
 	{
 		const FShadowMapInteraction ShadowMapInteraction = (Parameters.Mesh.LCI && bIsLitMaterial) 
diff --git a/Engine/Source/Runtime/Renderer/Private/IndirectLightingCache.cpp b/Engine/Source/Runtime/Renderer/Private/IndirectLightingCache.cpp
--- a/Engine/Source/Runtime/Renderer/Private/IndirectLightingCache.cpp
+++ b/Engine/Source/Runtime/Renderer/Private/IndirectLightingCache.cpp
@@ -870,6 +870,7 @@ void FIndirectLightingCache::InterpolateBlock(
 		for (int32 VolumeIndex = 0; VolumeIndex < Scene->PrecomputedLightVolumes.Num(); VolumeIndex++)
 		{
 			const FPrecomputedLightVolume* PrecomputedLightVolume = Scene->PrecomputedLightVolumes[VolumeIndex];
+			check(PrecomputedLightVolume->IsUsingHighQualityLightMap() == AllowHighQualityLightmaps(Scene->GetFeatureLevel()));
 			// Interpolate from the SH volume lighting samples that Lightmass computed
 			// Query using the bounds of all the samples in this block
 			// There will be a performance cliff for large objects which end up intersecting with the entire octree
diff --git a/Engine/Source/Runtime/Renderer/Private/LightMapDensityRendering.cpp b/Engine/Source/Runtime/Renderer/Private/LightMapDensityRendering.cpp
--- a/Engine/Source/Runtime/Renderer/Private/LightMapDensityRendering.cpp
+++ b/Engine/Source/Runtime/Renderer/Private/LightMapDensityRendering.cpp
@@ -118,7 +118,7 @@ bool FLightMapDensityDrawingPolicyFactory::DrawDynamicMesh(
 	bool bIsLitMaterial = (Material->GetShadingModel() != MSM_Unlit);
 	/*const */FLightMapInteraction LightMapInteraction = (Mesh.LCI && bIsLitMaterial) ? Mesh.LCI->GetLightMapInteraction() : FLightMapInteraction();
 	// force simple lightmaps based on system settings
-	bool bAllowHighQualityLightMaps = AllowHighQualityLightmaps() && LightMapInteraction.AllowsHighQualityLightmaps();
+	bool bAllowHighQualityLightMaps = AllowHighQualityLightmaps(View.GetFeatureLevel()) && LightMapInteraction.AllowsHighQualityLightmaps();
 	if (bIsLitMaterial && PrimitiveSceneProxy && (PrimitiveSceneProxy->GetLightMapType() == LMIT_Texture))
 	{
 		// Should this object be texture lightmapped? Ie, is lighting not built for it??
diff --git a/Engine/Source/Runtime/Renderer/Private/PostProcess/RenderTargetPool.cpp b/Engine/Source/Runtime/Renderer/Private/PostProcess/RenderTargetPool.cpp
--- a/Engine/Source/Runtime/Renderer/Private/PostProcess/RenderTargetPool.cpp
+++ b/Engine/Source/Runtime/Renderer/Private/PostProcess/RenderTargetPool.cpp
@@ -305,7 +305,7 @@ bool FRenderTargetPool::FindFreeElement(const FPooledRenderTargetDesc& Desc, TRe
 		if(Desc.TargetableFlags & TexCreate_UAV)
 		{
 			// The render target desc is invalid if a UAV is requested with an RHI that doesn't support the high-end feature level.
-			check(GRHIFeatureLevel == ERHIFeatureLevel::SM5);
+			check(GMaxRHIFeatureLevel == ERHIFeatureLevel::SM5);
 			Found->RenderTargetItem.UAV = RHICreateUnorderedAccessView(Found->RenderTargetItem.TargetableTexture);
 		}
 
diff --git a/Engine/Source/Runtime/Renderer/Private/RendererScene.cpp b/Engine/Source/Runtime/Renderer/Private/RendererScene.cpp
--- a/Engine/Source/Runtime/Renderer/Private/RendererScene.cpp
+++ b/Engine/Source/Runtime/Renderer/Private/RendererScene.cpp
@@ -1582,33 +1582,33 @@ void FScene::UpdateStaticDrawListsForMaterials_RenderThread(FRHICommandListImmed
 
 	// Warning: if any static draw lists are missed here, there will be a crash when trying to render with shaders that have been deleted!
 	TArray<FPrimitiveSceneInfo*> PrimitivesToUpdate;
-
+	auto FeatureLevel = GetFeatureLevel();
 	for (int32 DrawType = 0; DrawType < EBasePass_MAX; DrawType++)
 	{
-		BasePassNoLightMapDrawList[DrawType].GetUsedPrimitivesBasedOnMaterials(Materials, PrimitivesToUpdate);
-		BasePassSimpleDynamicLightingDrawList[DrawType].GetUsedPrimitivesBasedOnMaterials(Materials, PrimitivesToUpdate);
-		BasePassCachedVolumeIndirectLightingDrawList[DrawType].GetUsedPrimitivesBasedOnMaterials(Materials, PrimitivesToUpdate);
-		BasePassCachedPointIndirectLightingDrawList[DrawType].GetUsedPrimitivesBasedOnMaterials(Materials, PrimitivesToUpdate);
-		BasePassHighQualityLightMapDrawList[DrawType].GetUsedPrimitivesBasedOnMaterials(Materials, PrimitivesToUpdate);
-		BasePassDistanceFieldShadowMapLightMapDrawList[DrawType].GetUsedPrimitivesBasedOnMaterials(Materials, PrimitivesToUpdate);
-		BasePassLowQualityLightMapDrawList[DrawType].GetUsedPrimitivesBasedOnMaterials(Materials, PrimitivesToUpdate);
-		BasePassSelfShadowedTranslucencyDrawList[DrawType].GetUsedPrimitivesBasedOnMaterials(Materials, PrimitivesToUpdate);
-		BasePassSelfShadowedCachedPointIndirectTranslucencyDrawList[DrawType].GetUsedPrimitivesBasedOnMaterials(Materials, PrimitivesToUpdate);
-
-		BasePassForForwardShadingNoLightMapDrawList[DrawType].GetUsedPrimitivesBasedOnMaterials(Materials, PrimitivesToUpdate);
-		BasePassForForwardShadingLowQualityLightMapDrawList[DrawType].GetUsedPrimitivesBasedOnMaterials(Materials, PrimitivesToUpdate);
-		BasePassForForwardShadingDistanceFieldShadowMapLightMapDrawList[DrawType].GetUsedPrimitivesBasedOnMaterials(Materials, PrimitivesToUpdate);
-		BasePassForForwardShadingDirectionalLightAndSHIndirectDrawList[DrawType].GetUsedPrimitivesBasedOnMaterials(Materials, PrimitivesToUpdate);
-	}
-
-	PositionOnlyDepthDrawList.GetUsedPrimitivesBasedOnMaterials(Materials, PrimitivesToUpdate);
-	DepthDrawList.GetUsedPrimitivesBasedOnMaterials(Materials, PrimitivesToUpdate);
-	MaskedDepthDrawList.GetUsedPrimitivesBasedOnMaterials(Materials, PrimitivesToUpdate);
-	HitProxyDrawList.GetUsedPrimitivesBasedOnMaterials(Materials, PrimitivesToUpdate);
-	HitProxyDrawList_OpaqueOnly.GetUsedPrimitivesBasedOnMaterials(Materials, PrimitivesToUpdate);
-	VelocityDrawList.GetUsedPrimitivesBasedOnMaterials(Materials, PrimitivesToUpdate);
-	WholeSceneShadowDepthDrawList.GetUsedPrimitivesBasedOnMaterials(Materials, PrimitivesToUpdate);
-	WholeSceneReflectiveShadowMapDrawList.GetUsedPrimitivesBasedOnMaterials(Materials, PrimitivesToUpdate);
+		BasePassNoLightMapDrawList[DrawType].GetUsedPrimitivesBasedOnMaterials(FeatureLevel, Materials, PrimitivesToUpdate);
+		BasePassSimpleDynamicLightingDrawList[DrawType].GetUsedPrimitivesBasedOnMaterials(FeatureLevel, Materials, PrimitivesToUpdate);
+		BasePassCachedVolumeIndirectLightingDrawList[DrawType].GetUsedPrimitivesBasedOnMaterials(FeatureLevel, Materials, PrimitivesToUpdate);
+		BasePassCachedPointIndirectLightingDrawList[DrawType].GetUsedPrimitivesBasedOnMaterials(FeatureLevel, Materials, PrimitivesToUpdate);
+		BasePassHighQualityLightMapDrawList[DrawType].GetUsedPrimitivesBasedOnMaterials(FeatureLevel, Materials, PrimitivesToUpdate);
+		BasePassDistanceFieldShadowMapLightMapDrawList[DrawType].GetUsedPrimitivesBasedOnMaterials(FeatureLevel, Materials, PrimitivesToUpdate);
+		BasePassLowQualityLightMapDrawList[DrawType].GetUsedPrimitivesBasedOnMaterials(FeatureLevel, Materials, PrimitivesToUpdate);
+		BasePassSelfShadowedTranslucencyDrawList[DrawType].GetUsedPrimitivesBasedOnMaterials(FeatureLevel, Materials, PrimitivesToUpdate);
+		BasePassSelfShadowedCachedPointIndirectTranslucencyDrawList[DrawType].GetUsedPrimitivesBasedOnMaterials(FeatureLevel, Materials, PrimitivesToUpdate);
+
+		BasePassForForwardShadingNoLightMapDrawList[DrawType].GetUsedPrimitivesBasedOnMaterials(FeatureLevel, Materials, PrimitivesToUpdate);
+		BasePassForForwardShadingLowQualityLightMapDrawList[DrawType].GetUsedPrimitivesBasedOnMaterials(FeatureLevel, Materials, PrimitivesToUpdate);
+		BasePassForForwardShadingDistanceFieldShadowMapLightMapDrawList[DrawType].GetUsedPrimitivesBasedOnMaterials(FeatureLevel, Materials, PrimitivesToUpdate);
+		BasePassForForwardShadingDirectionalLightAndSHIndirectDrawList[DrawType].GetUsedPrimitivesBasedOnMaterials(FeatureLevel, Materials, PrimitivesToUpdate);
+	}
+
+	PositionOnlyDepthDrawList.GetUsedPrimitivesBasedOnMaterials(FeatureLevel, Materials, PrimitivesToUpdate);
+	DepthDrawList.GetUsedPrimitivesBasedOnMaterials(FeatureLevel, Materials, PrimitivesToUpdate);
+	MaskedDepthDrawList.GetUsedPrimitivesBasedOnMaterials(FeatureLevel, Materials, PrimitivesToUpdate);
+	HitProxyDrawList.GetUsedPrimitivesBasedOnMaterials(FeatureLevel, Materials, PrimitivesToUpdate);
+	HitProxyDrawList_OpaqueOnly.GetUsedPrimitivesBasedOnMaterials(FeatureLevel, Materials, PrimitivesToUpdate);
+	VelocityDrawList.GetUsedPrimitivesBasedOnMaterials(FeatureLevel, Materials, PrimitivesToUpdate);
+	WholeSceneShadowDepthDrawList.GetUsedPrimitivesBasedOnMaterials(FeatureLevel, Materials, PrimitivesToUpdate);
+	WholeSceneReflectiveShadowMapDrawList.GetUsedPrimitivesBasedOnMaterials(FeatureLevel, Materials, PrimitivesToUpdate);
 
 	for (int32 PrimitiveIndex = 0; PrimitiveIndex < PrimitivesToUpdate.Num(); PrimitiveIndex++)
 	{
diff --git a/Engine/Source/Runtime/Renderer/Private/SceneRendering.cpp b/Engine/Source/Runtime/Renderer/Private/SceneRendering.cpp
--- a/Engine/Source/Runtime/Renderer/Private/SceneRendering.cpp
+++ b/Engine/Source/Runtime/Renderer/Private/SceneRendering.cpp
@@ -468,7 +468,7 @@ TUniformBufferRef<FViewUniformShaderParameters> FViewInfo::CreateUniformBuffer(
 	ViewUniformShaderParameters.SceneTextureMinMax = SceneTexMinMax;
 
 	FScene* Scene = (FScene*)Family->Scene;
-	ERHIFeatureLevel::Type FeatureLevel;
+	ERHIFeatureLevel::Type FeatureLevel = Scene == nullptr ? GMaxRHIFeatureLevel : Scene->GetFeatureLevel();
 
 	if (Scene && Scene->SkyLight)
 	{
@@ -482,15 +482,11 @@ TUniformBufferRef<FViewUniformShaderParameters> FViewInfo::CreateUniformBuffer(
 			&& SkyLight->bPrecomputedLightingIsValid;
 
 		ViewUniformShaderParameters.SkyLightParameters = bApplyPrecomputedBentNormalShadowing ? 1 : 0;
-
-		FeatureLevel = Scene->GetFeatureLevel();
 	}
 	else
 	{
 		ViewUniformShaderParameters.SkyLightColor = FLinearColor::Black;
 		ViewUniformShaderParameters.SkyLightParameters = 0;
-
-		FeatureLevel = GMaxRHIFeatureLevel;
 	}
 
 	// Make sure there's no padding since we're going to cast to FVector4*
@@ -632,6 +628,7 @@ FSceneRenderer::FSceneRenderer(const FSceneViewFamily* InViewFamily,FHitProxyCon
 	// (I apologize for the const_cast, but didn't seem worth refactoring just for the freezerendering command)
 	bHasRequestedToggleFreeze = const_cast<FRenderTarget*>(InViewFamily->RenderTarget)->HasToggleFreezeCommand();
 
+	FeatureLevel = Scene->GetFeatureLevel();
 }
 
 bool FSceneRenderer::DoOcclusionQueries() const
diff --git a/Engine/Source/Runtime/Renderer/Private/SceneRendering.h b/Engine/Source/Runtime/Renderer/Private/SceneRendering.h
--- a/Engine/Source/Runtime/Renderer/Private/SceneRendering.h
+++ b/Engine/Source/Runtime/Renderer/Private/SceneRendering.h
@@ -412,8 +412,6 @@ public:
 	FLinearColor LightShaftColorApply;
 	bool bLightShaftUse;
 
-	ERHIFeatureLevel::Type FeatureLevel;
-
 	/** 
 	 * Initialization constructor. Passes all parameters to FSceneView constructor
 	 */
@@ -513,6 +511,9 @@ public:
 	/** Copy from main thread GFrameNumber to be accessible on renderthread side */
 	uint32 FrameNumber;
 
+	/** Feature level being rendered */
+	ERHIFeatureLevel::Type FeatureLevel;
+
 public:
 
 	FSceneRenderer(const FSceneViewFamily* InViewFamily,FHitProxyConsumer* HitProxyConsumer);
diff --git a/Engine/Source/Runtime/Renderer/Private/StaticMeshDrawList.h b/Engine/Source/Runtime/Renderer/Private/StaticMeshDrawList.h
--- a/Engine/Source/Runtime/Renderer/Private/StaticMeshDrawList.h
+++ b/Engine/Source/Runtime/Renderer/Private/StaticMeshDrawList.h
@@ -299,7 +299,7 @@ public:
 	void SortFrontToBack(FVector ViewPosition);
 
 	/** Builds a list of primitives that use the given materials in this static draw list. */
-	void GetUsedPrimitivesBasedOnMaterials(const TArray<const FMaterial*>& Materials, TArray<FPrimitiveSceneInfo*>& PrimitivesToUpdate);
+	void GetUsedPrimitivesBasedOnMaterials(ERHIFeatureLevel::Type FeatureLevel, const TArray<const FMaterial*>& Materials, TArray<FPrimitiveSceneInfo*>& PrimitivesToUpdate);
 
 	/**
 	 * Shifts all meshes bounds by an arbitrary delta.
diff --git a/Engine/Source/Runtime/Renderer/Private/StaticMeshDrawList.inl b/Engine/Source/Runtime/Renderer/Private/StaticMeshDrawList.inl
--- a/Engine/Source/Runtime/Renderer/Private/StaticMeshDrawList.inl
+++ b/Engine/Source/Runtime/Renderer/Private/StaticMeshDrawList.inl
@@ -550,7 +550,7 @@ void TStaticMeshDrawList<DrawingPolicyType>::SortFrontToBack(FVector ViewPositio
 }
 
 template<typename DrawingPolicyType>
-void TStaticMeshDrawList<DrawingPolicyType>::GetUsedPrimitivesBasedOnMaterials(const TArray<const FMaterial*>& Materials, TArray<FPrimitiveSceneInfo*>& PrimitivesToUpdate)
+void TStaticMeshDrawList<DrawingPolicyType>::GetUsedPrimitivesBasedOnMaterials(ERHIFeatureLevel::Type FeatureLevel, const TArray<const FMaterial*>& Materials, TArray<FPrimitiveSceneInfo*>& PrimitivesToUpdate)
 {
 	for (typename TDrawingPolicySet::TIterator DrawingPolicyIt(DrawingPolicySet); DrawingPolicyIt; ++DrawingPolicyIt)
 	{
@@ -566,7 +566,7 @@ void TStaticMeshDrawList<DrawingPolicyType>::GetUsedPrimitivesBasedOnMaterials(c
 
 			if (Proxy)
 			{
-				FMaterial* MaterialResource = Proxy->GetMaterialNoFallback(GRHIFeatureLevel);
+				FMaterial* MaterialResource = Proxy->GetMaterialNoFallback(FeatureLevel);
 
 				if (Materials.Contains(MaterialResource))
 				{
diff --git a/Engine/Source/Runtime/Renderer/Private/VelocityRendering.cpp b/Engine/Source/Runtime/Renderer/Private/VelocityRendering.cpp
--- a/Engine/Source/Runtime/Renderer/Private/VelocityRendering.cpp
+++ b/Engine/Source/Runtime/Renderer/Private/VelocityRendering.cpp
@@ -443,7 +443,8 @@ bool IsMotionBlurEnabled(const FViewInfo& View)
 
 void FDeferredShadingSceneRenderer::RenderVelocities(FRHICommandListImmediate& RHICmdList, const FViewInfo& View, TRefCountPtr<IPooledRenderTarget>& VelocityRT, bool bLastFrame)
 {
-	SCOPE_CYCLE_COUNTER( STAT_RenderVelocities );
+	check(FeatureLevel >= ERHIFeatureLevel::SM4);
+	SCOPE_CYCLE_COUNTER(STAT_RenderVelocities);
 
 	bool bTemporalAA = (View.FinalPostProcessSettings.AntiAliasingMethod == AAM_TemporalAA) && !View.bCameraCut;
 
diff --git a/Engine/Source/Runtime/Renderer/Public/RendererInterface.h b/Engine/Source/Runtime/Renderer/Public/RendererInterface.h
--- a/Engine/Source/Runtime/Renderer/Public/RendererInterface.h
+++ b/Engine/Source/Runtime/Renderer/Public/RendererInterface.h
@@ -199,7 +199,7 @@ public:
 		}
 
 		return Extent.X != 0 && NumMips != 0 && NumSamples >=1 && NumSamples <=16 && Format != PF_Unknown
-			&& ((TargetableFlags & TexCreate_UAV) == 0 || GRHIFeatureLevel == ERHIFeatureLevel::SM5);
+			&& ((TargetableFlags & TexCreate_UAV) == 0 || GMaxRHIFeatureLevel == ERHIFeatureLevel::SM5);
 	}
 
 	/** 
diff --git a/Engine/Source/Runtime/Windows/D3D11RHI/Private/D3D11Device.cpp b/Engine/Source/Runtime/Windows/D3D11RHI/Private/D3D11Device.cpp
--- a/Engine/Source/Runtime/Windows/D3D11RHI/Private/D3D11Device.cpp
+++ b/Engine/Source/Runtime/Windows/D3D11RHI/Private/D3D11Device.cpp
@@ -309,7 +309,7 @@ void FD3D11DynamicRHI::RHIGetSupportedResolution( uint32 &Width, uint32 &Height
 void FD3D11DynamicRHI::GetBestSupportedMSAASetting( DXGI_FORMAT PlatformFormat, uint32 MSAACount, uint32& OutBestMSAACount, uint32& OutMSAAQualityLevels )
 {
 	//  We disable MSAA for Feature level 10
-	if (GRHIFeatureLevel == ERHIFeatureLevel::SM4)
+	if (GMaxRHIFeatureLevel == ERHIFeatureLevel::SM4)
 	{
 		OutBestMSAACount = 1;
 		OutMSAAQualityLevels = 0;
diff --git a/Engine/Source/Runtime/Windows/D3D11RHI/Private/D3D11RHI.cpp b/Engine/Source/Runtime/Windows/D3D11RHI/Private/D3D11RHI.cpp
--- a/Engine/Source/Runtime/Windows/D3D11RHI/Private/D3D11RHI.cpp
+++ b/Engine/Source/Runtime/Windows/D3D11RHI/Private/D3D11RHI.cpp
@@ -182,7 +182,7 @@ FGlobalBoundShaderState LongGPUTaskBoundShaderState;
 
 void FD3D11DynamicRHI::IssueLongGPUTask()
 {
-	if (GRHIFeatureLevel >= ERHIFeatureLevel::SM3)
+	if (GMaxRHIFeatureLevel >= ERHIFeatureLevel::SM3)
 	{
 		int32 LargestViewportIndex = INDEX_NONE;
 		int32 LargestViewportPixels = 0;
diff --git a/Engine/Source/Runtime/Windows/D3D11RHI/Private/D3D11VertexBuffer.cpp b/Engine/Source/Runtime/Windows/D3D11RHI/Private/D3D11VertexBuffer.cpp
--- a/Engine/Source/Runtime/Windows/D3D11RHI/Private/D3D11VertexBuffer.cpp
+++ b/Engine/Source/Runtime/Windows/D3D11RHI/Private/D3D11VertexBuffer.cpp
@@ -24,7 +24,7 @@ FVertexBufferRHIRef FD3D11DynamicRHI::RHICreateVertexBuffer(uint32 Size,uint32 I
 	{
 		Desc.BindFlags |= D3D11_BIND_UNORDERED_ACCESS;
 
-		static bool bRequiresRawView = (GRHIFeatureLevel < ERHIFeatureLevel::SM5);
+		static bool bRequiresRawView = (GMaxRHIFeatureLevel < ERHIFeatureLevel::SM5);
 		if (bRequiresRawView)
 		{
 			Desc.MiscFlags |= D3D11_RESOURCE_MISC_BUFFER_ALLOW_RAW_VIEWS;
