We should disable OpenGL.UseMapBuffer for improved performance under MTGL on OS X, but radr://17662549 prevents it when using GL_TEXTURE_BUFFER's for skinning.

diff --git a/Engine/Source/Runtime/OpenGLDrv/Private/Mac/OpenGLMac.cpp b/Engine/Source/Runtime/OpenGLDrv/Private/Mac/OpenGLMac.cpp
--- a/Engine/Source/Runtime/OpenGLDrv/Private/Mac/OpenGLMac.cpp
+++ b/Engine/Source/Runtime/OpenGLDrv/Private/Mac/OpenGLMac.cpp
@@ -96,13 +96,7 @@ static NSOpenGLContext* CreateContext( NSOpenGLContext* SharedContext )
 	{
 		CGLEnable((CGLContextObj)[Context CGLContextObj], kCGLCEMPEngine);
 		
-		// Disable OpenGL.UseMapBuffer when using MTGL to reduce the number of context synchronisation points.
-		// All calls to glMapBuffer/Range will stall the MTGL thread, even with the unsynchronized bit set, so we want to avoid it.
-		static auto CVar = IConsoleManager::Get().FindConsoleVariable(TEXT("OpenGL.UseMapBuffer"));
-		if(CVar && CVar->GetInt() == 1)
-		{
-			CVar->Set(0);
-		}
+		// @todo: We should disable OpenGL.UseMapBuffer for improved performance under MTGL, but radr://17662549 prevents it when using GL_TEXTURE_BUFFER's for skinning.
 	}
 
 	return Context;
