Equals is a valid case also.

diff --git a/Engine/Source/Runtime/OpenGLDrv/Private/OpenGLCommands.cpp b/Engine/Source/Runtime/OpenGLDrv/Private/OpenGLCommands.cpp
--- a/Engine/Source/Runtime/OpenGLDrv/Private/OpenGLCommands.cpp
+++ b/Engine/Source/Runtime/OpenGLDrv/Private/OpenGLCommands.cpp
@@ -387,8 +387,8 @@ inline void FOpenGLDynamicRHI::UpdateScissorRectInOpenGLContext( FOpenGLContextS
 	if( PendingState.bScissorEnabled &&
 		ContextState.Scissor != PendingState.Scissor )
 	{
-		check(PendingState.Scissor.Min.X < PendingState.Scissor.Max.X);
-		check(PendingState.Scissor.Min.Y < PendingState.Scissor.Max.Y);
+		check(PendingState.Scissor.Min.X <= PendingState.Scissor.Max.X);
+		check(PendingState.Scissor.Min.Y <= PendingState.Scissor.Max.Y);
 		glScissor(PendingState.Scissor.Min.X, PendingState.Scissor.Min.Y, PendingState.Scissor.Max.X - PendingState.Scissor.Min.X, PendingState.Scissor.Max.Y - PendingState.Scissor.Min.Y);
 		ContextState.Scissor = PendingState.Scissor;
 	}
