Use a console variable for MTGL, not a command-line argument and enable it by default.

diff --git a/Engine/Source/Runtime/OpenGLDrv/Private/Mac/OpenGLMac.cpp b/Engine/Source/Runtime/OpenGLDrv/Private/Mac/OpenGLMac.cpp
--- a/Engine/Source/Runtime/OpenGLDrv/Private/Mac/OpenGLMac.cpp
+++ b/Engine/Source/Runtime/OpenGLDrv/Private/Mac/OpenGLMac.cpp
@@ -5,6 +5,38 @@
 #include "MacWindow.h"
 #include "MacTextInputMethodSystem.h"
 
+/*------------------------------------------------------------------------------
+ OpenGL static variables.
+ ------------------------------------------------------------------------------*/
+
+// @todo: remove once Apple fixes radr://16754329 AMD Cards don't always perform FRAMEBUFFER_SRGB if the draw FBO has mixed sRGB & non-SRGB colour attachments
+static TAutoConsoleVariable<int32> CVarMacUseFrameBufferSRGB(
+	TEXT("r.Mac.UseFrameBufferSRGB"),
+	0,
+	TEXT("Flag to toggle use of GL_FRAMEBUFFER_SRGB for better color accuracy.\n"),
+	ECVF_RenderThreadSafe
+	);
+
+// @todo: remove once Apple fixes radr://15553950, TTP# 315197
+static int32 GMacMustFlushTexStorage = 0;
+static FAutoConsoleVariableRef CVarMacMustFlushTexStorage(
+	TEXT("r.Mac.MustFlushTexStorage"),
+	GMacMustFlushTexStorage,
+	TEXT("If true, flush the OpenGL command stream after calls to glTexStorage* to avoid driver errors, do nothing if false (faster, the default)."),
+	ECVF_RenderThreadSafe
+	);
+
+static int32 GMacUseMTGL = 1;
+static FAutoConsoleVariableRef CVarMacUseMTGL(
+	TEXT("r.Mac.UseMTGL"),
+	GMacUseMTGL,
+	TEXT("If true use Apple's multi-threaded OpenGL which parallelises the OpenGL driver with the rendering thread for improved performance, use false to disable. (Default: True)"),
+	ECVF_RenderThreadSafe
+	);
+
+bool GIsRunningOnIntelCard = false; // @todo: remove once Apple fixes radr://16223045 Changes to the GL separate blend state aren't always respected on Intel cards
+static bool GIsEmulatingTimestamp = false; // @todo: Now crashing on Nvidia cards, but not on AMD...
+
 /*------------------------------------------------------------------------------
 	OpenGL context management.
 ------------------------------------------------------------------------------*/
@@ -92,7 +124,7 @@ static NSOpenGLContext* CreateContext( NSOpenGLContext* SharedContext )
 	int32 SurfaceOpacity = 0;
 	[Context setValues: &SurfaceOpacity forParameter: NSOpenGLCPSurfaceOpacity];
 
-	if (FParse::Param(FCommandLine::Get(),TEXT("openglUseMacMTEngine")))
+	if (GMacUseMTGL)
 	{
 		CGLEnable((CGLContextObj)[Context CGLContextObj], kCGLCEMPEngine);
 		
@@ -288,26 +320,6 @@ struct OpenGLContextInfo
 
 extern void OnQueryInvalidation( void );
 
-// @todo: remove once Apple fixes radr://16754329 AMD Cards don't always perform FRAMEBUFFER_SRGB if the draw FBO has mixed sRGB & non-SRGB colour attachments
-static TAutoConsoleVariable<int32> CVarMacUseFrameBufferSRGB(
-		TEXT("r.Mac.UseFrameBufferSRGB"),
-		0,
-		TEXT("Flag to toggle use of GL_FRAMEBUFFER_SRGB for better color accuracy.\n"),
-		ECVF_RenderThreadSafe
-		);
-
-// @todo: remove once Apple fixes radr://15553950, TTP# 315197
-static int32 GMacMustFlushTexStorage = 0;
-static FAutoConsoleVariableRef CVarMacMustFlushTexStorage(
-	TEXT("r.Mac.MustFlushTexStorage"),
-	GMacMustFlushTexStorage,
-	TEXT("If true, flush the OpenGL command stream after calls to glTexStorage* to avoid driver errors, do nothing if false (faster, the default)."),
-	ECVF_RenderThreadSafe
-	);
-
-bool GIsRunningOnIntelCard = false; // @todo: remove once Apple fixes radr://16223045 Changes to the GL separate blend state aren't always respected on Intel cards
-static bool GIsEmulatingTimestamp = false; // @todo: Now crashing on Nvidia cards, but not on AMD...
-
 struct FPlatformOpenGLDevice
 {
 	FPlatformOpenGLContext	SharedContext;
@@ -1393,7 +1405,7 @@ void FMacOpenGL::ProcessExtensions(const FString& ExtensionsString)
 	}
 	
 	// Don't label objects with MTGL, it causes synchronisation of the MTGL thread.
-	if(ExtensionsString.Contains(TEXT("GL_EXT_debug_label")) && !FParse::Param(FCommandLine::Get(),TEXT("openglUseMacMTEngine")))
+	if(ExtensionsString.Contains(TEXT("GL_EXT_debug_label")) && !GMacUseMTGL)
 	{
 		glLabelObjectEXT = (PFNGLLABELOBJECTEXTPROC)dlsym(RTLD_SELF, "glLabelObjectEXT");
 	}
@@ -1799,5 +1811,6 @@ void FMacOpenGL::MacGetQueryObject(GLuint QueryId, EQueryMode QueryMode, GLuint
 
 bool FMacOpenGL::MustFlushTexStorage(void)
 {
-	return GMacMustFlushTexStorage;
+	// @todo There is a bug in Apple's GL with TexStorage calls, on Nvidia and when using MTGL that can see the texture never be created, which then subsequently causes crashes
+	return GMacMustFlushTexStorage || GMacUseMTGL;
 }
