For the benefit of OS X implement support for APPLE_client_storage & APPLE_texture_range and use them for non-renderable textures to eliminate the cost of Lock/Unlock and RHIAsyncReallocateTexture2D.
This only works with uncompressed textures, but helps avoid very slow, synchronous calls to glGetTexImage despite using PBOs. This does not increase memory overhead as client-storage textures use the backing store we provide for the GL engine copy.

diff --git a/Engine/Source/Runtime/OpenGLDrv/Private/Mac/OpenGLMac.h b/Engine/Source/Runtime/OpenGLDrv/Private/Mac/OpenGLMac.h
--- a/Engine/Source/Runtime/OpenGLDrv/Private/Mac/OpenGLMac.h
+++ b/Engine/Source/Runtime/OpenGLDrv/Private/Mac/OpenGLMac.h
@@ -210,7 +210,14 @@ struct FMacOpenGL : public FOpenGL3
 		}
 	}
 	
+	static FORCEINLINE void TextureRange(GLenum Target, GLsizei Length, const GLvoid *Pointer)
+	{
+		glTextureRangeAPPLE(Target, Length, Pointer);
+	}
+	
 	static FORCEINLINE bool SupportsSeamlessCubeMap()					{ return true; }
+	static FORCEINLINE bool SupportsClientStorage()						{ return true; }
+	static FORCEINLINE bool SupportsTextureRange()						{ return true; }
 	
 	static FORCEINLINE EShaderPlatform GetShaderPlatform()
 	{
diff --git a/Engine/Source/Runtime/OpenGLDrv/Private/OpenGLTexture.cpp b/Engine/Source/Runtime/OpenGLDrv/Private/OpenGLTexture.cpp
--- a/Engine/Source/Runtime/OpenGLDrv/Private/OpenGLTexture.cpp
+++ b/Engine/Source/Runtime/OpenGLDrv/Private/OpenGLTexture.cpp
@@ -297,6 +297,9 @@ FRHITexture* FOpenGLDynamicRHI::CreateOpenGLTexture(uint32 SizeX,uint32 SizeY,bo
 	// Use a texture stage that's not likely to be used for draws, to avoid waiting
 	CachedSetupTextureStage(ContextState, FOpenGL::GetMaxCombinedTextureImageUnits() - 1, Target, TextureID, 0, NumMips);
 
+	// For client storage textures we allocate a single backing store buffer.
+	uint8* TextureRange = nullptr;
+	
 	if (NumSamples == 1)
 	{
 		if (!FMath::IsPowerOfTwo(SizeX) || !FMath::IsPowerOfTwo(SizeY))
@@ -338,8 +341,60 @@ FRHITexture* FOpenGLDynamicRHI::CreateOpenGLTexture(uint32 SizeX,uint32 SizeY,bo
 		}
 		else
 		{
+			// Should we use client-storage to improve update time on platforms that require it
+			bool const bRenderable = (Flags & (TexCreate_RenderTargetable|TexCreate_ResolveTargetable|TexCreate_DepthStencilTargetable)) != 0;
+			bool const bUseClientStorage = (FOpenGL::SupportsClientStorage() && !FOpenGL::SupportsTextureView() && !bRenderable && !GLFormat.bCompressed);
+			if(bUseClientStorage)
+			{
+				const bool bIsCubeTexture = Target == GL_TEXTURE_CUBE_MAP;
+				const uint32 TextureSize = CalcTextureSize(SizeX, SizeY, (EPixelFormat)Format, NumMips) * (bIsCubeTexture ? 6 : 1);
+				const GLenum FirstTarget = bIsCubeTexture ? GL_TEXTURE_CUBE_MAP_POSITIVE_X : Target;
+				const uint32 NumTargets = bIsCubeTexture ? 6 : 1;
+				
+				TextureRange = new uint8[TextureSize];
+				check(TextureRange);
+				
+				if(FOpenGL::SupportsTextureRange())
+				{
+					FOpenGL::TextureRange(Target, TextureSize, TextureRange);
+					glTexParameteri(Target, GL_TEXTURE_STORAGE_HINT_APPLE, GL_STORAGE_CACHED_APPLE);
+				}
+				
+				glPixelStorei(GL_UNPACK_CLIENT_STORAGE_APPLE, GL_TRUE);
+				
+				uint8* MipPointer = TextureRange;
+				for(uint32 MipIndex = 0; MipIndex < uint32(NumMips); MipIndex++)
+				{
+					const uint32 MipSize = CalcTextureMipMapSize(SizeX, SizeY, (EPixelFormat)Format, MipIndex);
+					for(uint32 TargetIndex = 0; TargetIndex < NumTargets; TargetIndex++)
+					{
+						glTexImage2D(
+									 FirstTarget + TargetIndex,
+									 MipIndex,
+									 GLFormat.InternalFormat[bSRGB],
+									 FMath::Max<uint32>(1,(SizeX >> MipIndex)),
+									 FMath::Max<uint32>(1,(SizeY >> MipIndex)),
+									 0,
+									 GLFormat.Format,
+									 GLFormat.Type,
+									 MipPointer
+									 );
+						MipPointer += MipSize;
+					}
+				}
+				
+				glPixelStorei(GL_UNPACK_CLIENT_STORAGE_APPLE, GL_FALSE);
+				
+				if(FOpenGL::SupportsTextureRange())
+				{
+					FOpenGL::TextureRange(Target, 0, 0);
+					glTexParameteri(Target, GL_TEXTURE_STORAGE_HINT_APPLE, GL_STORAGE_PRIVATE_APPLE);
+				}
+				
+				// Leave bAllocatedStorage as false, so that the client storage buffers are setup only when the texture is locked
+			}
 			// Try to allocate using TexStorage2D
-			if( FOpenGL::TexStorage2D( Target, NumMips, GLFormat.SizedInternalFormat[bSRGB], SizeX, SizeY, GLFormat.Format, GLFormat.Type, Flags) )
+			else if( FOpenGL::TexStorage2D( Target, NumMips, GLFormat.SizedInternalFormat[bSRGB], SizeX, SizeY, GLFormat.Format, GLFormat.Type, Flags) )
 			{
 				bAllocatedStorage = true;
 			}
@@ -490,12 +545,12 @@ FRHITexture* FOpenGLDynamicRHI::CreateOpenGLTexture(uint32 SizeX,uint32 SizeY,bo
 	FRHITexture* Texture;
 	if (bCubeTexture)
 	{
-		FOpenGLTextureCube* TextureCube = new FOpenGLTextureCube(this,TextureID,Target,Attachment,SizeX,SizeY,0,NumMips,1, ArraySize, (EPixelFormat)Format,true,bAllocatedStorage,Flags);
+		FOpenGLTextureCube* TextureCube = new FOpenGLTextureCube(this,TextureID,Target,Attachment,SizeX,SizeY,0,NumMips,1, ArraySize, (EPixelFormat)Format,true,bAllocatedStorage,Flags,TextureRange);
 		Texture = TextureCube;
 	}
 	else
 	{
-		FOpenGLTexture2D* Texture2D = new FOpenGLTexture2D(this,TextureID,Target,Attachment,SizeX,SizeY,0,NumMips,NumSamples, 1, (EPixelFormat)Format,false,bAllocatedStorage,Flags);
+		FOpenGLTexture2D* Texture2D = new FOpenGLTexture2D(this,TextureID,Target,Attachment,SizeX,SizeY,0,NumMips,NumSamples, 1, (EPixelFormat)Format,false,bAllocatedStorage,Flags,TextureRange);
 		Texture = Texture2D;
 	}
 	OpenGLTextureAllocated( Texture, Flags );
@@ -673,27 +728,70 @@ void* TOpenGLTexture<RHIResourceType>::Lock(uint32 MipIndex,uint32 ArrayIndex,ER
 
 	const int32 BufferIndex = MipIndex * (bCubemap ? 6 : 1) * this->GetEffectiveSizeZ() + ArrayIndex;
 
-	// Standard path with a PBO mirroring ever slice of a texture to allow multiple simulataneous maps
-	bool bBufferExists = true;
-	if (!IsValidRef(PixelBuffers[BufferIndex]))
+	// Should we use client-storage to improve update time on platforms that require it
+	const FOpenGLTextureFormat& GLFormat = GOpenGLTextureFormats[PixelFormat];
+	bool const bRenderable = (this->GetFlags() & (TexCreate_RenderTargetable|TexCreate_ResolveTargetable|TexCreate_DepthStencilTargetable)) != 0;
+	bool const bUseClientStorage = FOpenGL::SupportsClientStorage() && !FOpenGL::SupportsTextureView() && !bRenderable && !this->GetSizeZ() && !GLFormat.bCompressed;
+	if(!bUseClientStorage)
 	{
-		bBufferExists = false;
-		PixelBuffers[BufferIndex] = new FOpenGLPixelBuffer(0, MipBytes, BUF_Dynamic);
-	}
+		// Standard path with a PBO mirroring ever slice of a texture to allow multiple simulataneous maps
+		bool bBufferExists = true;
+		if (!IsValidRef(PixelBuffers[BufferIndex]))
+		{
+			bBufferExists = false;
+			PixelBuffers[BufferIndex] = new FOpenGLPixelBuffer(0, MipBytes, BUF_Dynamic);
+		}
 
-	TRefCountPtr<FOpenGLPixelBuffer> PixelBuffer = PixelBuffers[BufferIndex];
-	check(PixelBuffer->GetSize() == MipBytes);
-	check(!PixelBuffer->IsLocked());
-	
-	// If the buffer already exists & the flags are such that the texture cannot be rendered to & is CPU accessible then we can skip the internal resolve for read locks. This makes HZB occlusion faster.
-	const bool bCPUTexResolved = bBufferExists && (this->GetFlags() & TexCreate_CPUReadback) && !(this->GetFlags() & (TexCreate_RenderTargetable|TexCreate_DepthStencilTargetable));
+		TRefCountPtr<FOpenGLPixelBuffer> PixelBuffer = PixelBuffers[BufferIndex];
+		check(PixelBuffer->GetSize() == MipBytes);
+		check(!PixelBuffer->IsLocked());
+		
+		// If the buffer already exists & the flags are such that the texture cannot be rendered to & is CPU accessible then we can skip the internal resolve for read locks. This makes HZB occlusion faster.
+		const bool bCPUTexResolved = bBufferExists && (this->GetFlags() & TexCreate_CPUReadback) && !(this->GetFlags() & (TexCreate_RenderTargetable|TexCreate_DepthStencilTargetable));
 
-	if( LockMode != RLM_WriteOnly && !bCPUTexResolved && FOpenGL::SupportsPixelBuffers() )
+		if( LockMode != RLM_WriteOnly && !bCPUTexResolved && FOpenGL::SupportsPixelBuffers() )
+		{
+			Resolve(MipIndex, ArrayIndex);
+		}
+
+		result = PixelBuffer->Lock(0, PixelBuffer->GetSize(), LockMode == RLM_ReadOnly, LockMode != RLM_ReadOnly);
+	}
+	else
 	{
-		Resolve(MipIndex, ArrayIndex);
+		// Use APPLE_client_storage to reduce memory usage and improve performance
+		// GL's which support this extension only need copy a pointer, not the memory contents
+		check( FOpenGL::SupportsClientStorage() && !FOpenGL::SupportsTextureView() );
+		if(GetAllocatedStorageForMip(MipIndex,ArrayIndex))
+		{
+			result = ClientStorageBuffers[BufferIndex].Data;
+		}
+		else
+		{
+			// The assumption at present is that this only applies to 2D & cubemap textures
+			// Array, 3D and variants thereof aren't supported.
+			const bool bIsCubeTexture = Target == GL_TEXTURE_CUBE_MAP;
+			const GLenum FirstTarget = bIsCubeTexture ? GL_TEXTURE_CUBE_MAP_POSITIVE_X : Target;
+			const uint32 NumTargets = bIsCubeTexture ? 6 : 1;
+			
+			uint8* MipPointer = TextureRange;
+			for(uint32 MipIndex = 0; MipIndex < FOpenGLTextureBase::NumMips; MipIndex++)
+			{
+				const uint32 MipSize = CalcTextureMipMapSize(this->GetSizeX(), this->GetSizeY(), PixelFormat, MipIndex);
+				for(uint32 TargetIndex = 0; TargetIndex < NumTargets; TargetIndex++)
+				{
+					const int32 ClientIndex = (MipIndex * NumTargets) + TargetIndex;
+					ClientStorageBuffers[ClientIndex].Data = MipPointer;
+					ClientStorageBuffers[ClientIndex].Size = MipSize;
+					ClientStorageBuffers[ClientIndex].bReadOnly = false;
+					MipPointer += MipSize;
+					SetAllocatedStorageForMip(MipIndex,TargetIndex);
+				}
+				
+			}
+			
+			result = ClientStorageBuffers[BufferIndex].Data;
+		}
 	}
-
-	result = PixelBuffer->Lock(0, PixelBuffer->GetSize(), LockMode == RLM_ReadOnly, LockMode != RLM_ReadOnly);
 	
 	return result;
 }
@@ -757,12 +855,15 @@ void TOpenGLTexture<RHIResourceType>::Unlock(uint32 MipIndex,uint32 ArrayIndex)
 	SCOPE_CYCLE_COUNTER(STAT_OpenGLUnlockTextureTime);
 
 	const int32 BufferIndex = MipIndex * (bCubemap ? 6 : 1) * this->GetEffectiveSizeZ() + ArrayIndex;
-	check(IsValidRef(PixelBuffers[BufferIndex]));
-
 	TRefCountPtr<FOpenGLPixelBuffer> PixelBuffer = PixelBuffers[BufferIndex];
 	const FOpenGLTextureFormat& GLFormat = GOpenGLTextureFormats[this->GetFormat()];
 	const bool bSRGB = (this->GetFlags() & TexCreate_SRGB) != 0;
 
+	// Should we use client-storage to improve update time on platforms that require it
+	bool const bRenderable = (this->GetFlags() & (TexCreate_RenderTargetable|TexCreate_ResolveTargetable|TexCreate_DepthStencilTargetable)) != 0;
+	bool const bUseClientStorage = FOpenGL::SupportsClientStorage() && !FOpenGL::SupportsTextureView() && !FOpenGL::SupportsTextureView() && !bRenderable && !this->GetSizeZ() && !GLFormat.bCompressed;
+	check(bUseClientStorage || IsValidRef(PixelBuffers[BufferIndex]));
+	
 #if PLATFORM_ANDROID
 	// check for FloatRGBA to RGBA8 conversion needed
 	if (this->GetFormat() == PF_FloatRGBA && GLFormat.Type == GL_UNSIGNED_BYTE)
@@ -825,8 +926,8 @@ void TOpenGLTexture<RHIResourceType>::Unlock(uint32 MipIndex,uint32 ArrayIndex)
 		return;
 	}
 #endif
-
-	if ( FOpenGL::SupportsPixelBuffers() )
+	
+	if ( !bUseClientStorage && FOpenGL::SupportsPixelBuffers() )
 	{
 		// Code path for PBO per slice
 		check(IsValidRef(PixelBuffers[BufferIndex]));
@@ -944,8 +1045,23 @@ void TOpenGLTexture<RHIResourceType>::Unlock(uint32 MipIndex,uint32 ArrayIndex)
 
 		CachedBindPixelUnpackBuffer( 0 );
 		
+		uint32 LockedSize = 0;
+		void* LockedBuffer = nullptr;
+		
+		if(bUseClientStorage)
+		{
+			glPixelStorei(GL_UNPACK_CLIENT_STORAGE_APPLE, GL_TRUE);
+			LockedSize = ClientStorageBuffers[BufferIndex].Size;
+			LockedBuffer = ClientStorageBuffers[BufferIndex].Data;
+		}
+		else
+		{
+			LockedSize = PixelBuffer->GetSize();
+			LockedBuffer = PixelBuffer->GetLockedBuffer();
+		}
 		if (GLFormat.bCompressed)
 		{
+			glPixelStorei(GL_UNPACK_ALIGNMENT, 1);
 			if (GetAllocatedStorageForMip(MipIndex,ArrayIndex))
 			{
 				glCompressedTexSubImage2D(
@@ -956,8 +1072,8 @@ void TOpenGLTexture<RHIResourceType>::Unlock(uint32 MipIndex,uint32 ArrayIndex)
 					FMath::Max<uint32>(1,(this->GetSizeX() >> MipIndex)),
 					FMath::Max<uint32>(1,(this->GetSizeY() >> MipIndex)),
 					GLFormat.InternalFormat[bSRGB],
-					PixelBuffer->GetSize(),
-					PixelBuffer->GetLockedBuffer());
+					LockedSize,
+					LockedBuffer);
 			}
 			else
 			{
@@ -968,30 +1084,51 @@ void TOpenGLTexture<RHIResourceType>::Unlock(uint32 MipIndex,uint32 ArrayIndex)
 					FMath::Max<uint32>(1,(this->GetSizeX() >> MipIndex)),
 					FMath::Max<uint32>(1,(this->GetSizeY() >> MipIndex)),
 					0,
-					PixelBuffer->GetSize(),
-					PixelBuffer->GetLockedBuffer());
+					LockedSize,
+					LockedBuffer);
 					SetAllocatedStorageForMip(MipIndex,ArrayIndex);
 			}
+			glPixelStorei(GL_UNPACK_ALIGNMENT, 4);
 		}
 		else
 		{
-			// All construction paths should have called TexStorage2D or TexImage2D. So we will
-			// always call TexSubImage2D.
-			check(GetAllocatedStorageForMip(MipIndex,ArrayIndex) == true);
-			glTexSubImage2D(
-				bCubemap ? GL_TEXTURE_CUBE_MAP_POSITIVE_X + ArrayIndex : Target,
-				MipIndex,
-				0,
-				0,
-				FMath::Max<uint32>(1,(this->GetSizeX() >> MipIndex)),
-				FMath::Max<uint32>(1,(this->GetSizeY() >> MipIndex)),
-				GLFormat.Format,
-				GLFormat.Type,
-				PixelBuffer->GetLockedBuffer());
+			if (GetAllocatedStorageForMip(MipIndex,ArrayIndex))
+			{
+				glTexSubImage2D(
+					bCubemap ? GL_TEXTURE_CUBE_MAP_POSITIVE_X + ArrayIndex : Target,
+					MipIndex,
+					0,
+					0,
+					FMath::Max<uint32>(1,(this->GetSizeX() >> MipIndex)),
+					FMath::Max<uint32>(1,(this->GetSizeY() >> MipIndex)),
+					GLFormat.Format,
+					GLFormat.Type,
+					LockedBuffer);
+			}
+			else
+			{
+				glTexImage2D(
+					bCubemap ? GL_TEXTURE_CUBE_MAP_POSITIVE_X + ArrayIndex : Target,
+					MipIndex,
+					GLFormat.InternalFormat[bSRGB],
+					FMath::Max<uint32>(1,(this->GetSizeX() >> MipIndex)),
+					FMath::Max<uint32>(1,(this->GetSizeY() >> MipIndex)),
+					0,
+					GLFormat.Format,
+					GLFormat.Type,
+					LockedBuffer);
+				SetAllocatedStorageForMip(MipIndex,ArrayIndex);
+			}
+		}
+		if(bUseClientStorage)
+		{
+			glPixelStorei(GL_UNPACK_CLIENT_STORAGE_APPLE, GL_TRUE);
+		}
+		else
+		{
+			// Unlock "PixelBuffer" and free the temp memory after the texture upload.
+			PixelBuffer->Unlock();
 		}
-
-		// Unlock "PixelBuffer" and free the temp memory after the texture upload.
-		PixelBuffer->Unlock();
 	}
 
 	// No need to restore texture stage; leave it like this,
@@ -1250,6 +1387,12 @@ TOpenGLTexture<RHIResourceType>::~TOpenGLTexture()
 		OpenGLRHI->InvalidateTextureResourceInCache( Resource );
 		glDeleteTextures( 1, &Resource );
 	}
+	
+	if(TextureRange)
+	{
+		delete [] TextureRange;
+		TextureRange = nullptr;
+	}
 
 	ReleaseOpenGLFramebuffers(OpenGLRHI,this);
 }
@@ -1381,7 +1524,7 @@ FTexture2DArrayRHIRef FOpenGLDynamicRHI::RHICreateTexture2DArray(uint32 SizeX,ui
 			: GL_COLOR_ATTACHMENT0);
 	}
 
-	FOpenGLTexture2DArray* Texture = new FOpenGLTexture2DArray(this,TextureID,Target,Attachment,SizeX,SizeY,SizeZ,NumMips,1, SizeZ, (EPixelFormat)Format,false,true,Flags);
+	FOpenGLTexture2DArray* Texture = new FOpenGLTexture2DArray(this,TextureID,Target,Attachment,SizeX,SizeY,SizeZ,NumMips,1, SizeZ, (EPixelFormat)Format,false,true,Flags,nullptr);
 	OpenGLTextureAllocated( Texture, Flags );
 
 	// No need to restore texture stage; leave it like this,
@@ -1489,7 +1632,7 @@ FTexture3DRHIRef FOpenGLDynamicRHI::RHICreateTexture3D(uint32 SizeX,uint32 SizeY
 			: GL_COLOR_ATTACHMENT0);
 	}
 
-	FOpenGLTexture3D* Texture = new FOpenGLTexture3D(this,TextureID,Target,Attachment,SizeX,SizeY,SizeZ,NumMips,1,1, (EPixelFormat)Format,false,true,Flags);
+	FOpenGLTexture3D* Texture = new FOpenGLTexture3D(this,TextureID,Target,Attachment,SizeX,SizeY,SizeZ,NumMips,1,1, (EPixelFormat)Format,false,true,Flags,nullptr);
 	OpenGLTextureAllocated( Texture, Flags );
 
 	// No need to restore texture stage; leave it like this,
@@ -1650,7 +1793,12 @@ FTexture2DRHIRef FOpenGLDynamicRHI::RHIAsyncReallocateTexture2D(FTexture2DRHIPar
 	const uint32 BlockSizeX = GPixelFormats[Texture2D->GetFormat()].BlockSizeX;
 	const uint32 BlockSizeY = GPixelFormats[Texture2D->GetFormat()].BlockSizeY;
 	const uint32 NumBytesPerBlock = GPixelFormats[Texture2D->GetFormat()].BlockBytes;
-
+	
+	// Should we use client-storage to improve update time on platforms that require it
+	const bool bCompressed = GOpenGLTextureFormats[Texture2D->GetFormat()].bCompressed;
+	bool const bRenderable = (Texture2D->GetFlags() & (TexCreate_RenderTargetable|TexCreate_ResolveTargetable|TexCreate_DepthStencilTargetable)) != 0;
+	const bool bUseClientStorage = (FOpenGL::SupportsClientStorage() && !FOpenGL::SupportsTextureView() && !bRenderable && !bCompressed);
+	
 	// Use the GPU to asynchronously copy the old mip-maps into the new texture.
 	const uint32 NumSharedMips = FMath::Min(Texture2D->GetNumMips(),NewTexture2D->GetNumMips());
 	const uint32 SourceMipOffset = Texture2D->GetNumMips()    - NumSharedMips;
@@ -1666,7 +1814,7 @@ FTexture2DRHIRef FOpenGLDynamicRHI::RHIAsyncReallocateTexture2D(FTexture2DRHIPar
 	{
 		FOpenGL::CopyTextureLevels( NewTexture2D->Resource, Texture2D->Resource, SourceMipOffset, NumSharedMips);
 	}
-	else if (FOpenGL::SupportsPixelBuffers())
+	else if (FOpenGL::SupportsPixelBuffers() && !bUseClientStorage)
 	{
 		FOpenGLTexture2D *NewOGLTexture2D = (FOpenGLTexture2D*)NewTexture2D;
 		FOpenGLTexture2D *OGLTexture2D = (FOpenGLTexture2D*)Texture2D;
diff --git a/Engine/Source/Runtime/OpenGLDrv/Public/OpenGL.h b/Engine/Source/Runtime/OpenGLDrv/Public/OpenGL.h
--- a/Engine/Source/Runtime/OpenGLDrv/Public/OpenGL.h
+++ b/Engine/Source/Runtime/OpenGLDrv/Public/OpenGL.h
@@ -133,6 +133,8 @@ public:
 	static FORCEINLINE bool SupportsVertexAttribBinding()				{ return false; }
 	static FORCEINLINE bool SupportsBufferStorage()						{ return false; }
 	static FORCEINLINE bool SupportsDepthBoundsTest()					{ return false; }
+	static FORCEINLINE bool SupportsClientStorage()						{ return false; }
+	static FORCEINLINE bool SupportsTextureRange()						{ return false; }
 	static FORCEINLINE bool SupportsTextureNPOT()						{ return true; }
 	static FORCEINLINE bool HasHardwareHiddenSurfaceRemoval()			{ return false; }
 	static FORCEINLINE bool AmdWorkaround()								{ return false; }
@@ -287,6 +289,7 @@ public:
 	static FORCEINLINE void VertexBindingDivisor(GLuint BindingIndex, GLuint Divisor) UGL_REQUIRED_VOID
 	static FORCEINLINE void BufferStorage(GLenum Target, GLsizeiptr Size, const void *Data, GLbitfield Flags) UGL_REQUIRED_VOID
 	static FORCEINLINE void DepthBounds(GLfloat Min, GLfloat Max) UGL_REQUIRED_VOID
+	static FORCEINLINE void TextureRange(GLenum Target, GLsizei Length, const GLvoid *Pointer) UGL_OPTIONAL_VOID
 
 	static FPlatformOpenGLDevice*	CreateDevice() UGL_REQUIRED(NULL)
 	static FPlatformOpenGLContext*	CreateContext( FPlatformOpenGLDevice* Device, void* WindowHandle ) UGL_REQUIRED(NULL)
@@ -561,3 +564,17 @@ protected:
 #ifndef GL_GPU_DISJOINT_EXT
 #define GL_GPU_DISJOINT_EXT							0x8FBB
 #endif
+
+#ifndef GL_APPLE_client_storage
+#define GL_UNPACK_CLIENT_STORAGE_APPLE				0x85B2
+#endif
+
+#ifndef GL_APPLE_texture_range
+#define GL_TEXTURE_RANGE_LENGTH_APPLE     0x85B7
+#define GL_TEXTURE_RANGE_POINTER_APPLE    0x85B8
+#define GL_TEXTURE_STORAGE_HINT_APPLE     0x85BC
+#define GL_TEXTURE_MINIMIZE_STORAGE_APPLE 0x85B6
+#define GL_STORAGE_PRIVATE_APPLE          0x85BD
+#define GL_STORAGE_CACHED_APPLE           0x85BE
+#define GL_STORAGE_SHARED_APPLE           0x85BF
+#endif
diff --git a/Engine/Source/Runtime/OpenGLDrv/Public/OpenGLResources.h b/Engine/Source/Runtime/OpenGLDrv/Public/OpenGLResources.h
--- a/Engine/Source/Runtime/OpenGLDrv/Public/OpenGLResources.h
+++ b/Engine/Source/Runtime/OpenGLDrv/Public/OpenGLResources.h
@@ -803,14 +803,6 @@ private:
 	uint32 bIsPowerOfTwo	: 1;
 };
 
-// Temporary Android GL4 WAR
-// Use host pointers instead of PBOs for texture uploads
-#if PLATFORM_ANDROIDGL4
-#define USE_PBO 0
-#else
-#define USE_PBO 1
-#endif
-
 // Textures.
 template<typename BaseType>
 class TOpenGLTexture : public BaseType, public FOpenGLTextureBase
@@ -832,7 +824,8 @@ public:
 		EPixelFormat InFormat,
 		bool bInCubemap,
 		bool bInAllocatedStorage,
-		uint32 InFlags
+		uint32 InFlags,
+		uint8* InTextureRange
 		)
 	: BaseType(InSizeX,InSizeY,InSizeZ,InNumMips,InNumSamples, InArraySize, InFormat,InFlags)
 	, FOpenGLTextureBase(
@@ -844,12 +837,11 @@ public:
 		)
 	, BaseLevel(0)
 	, bCubemap(bInCubemap)
+	, TextureRange(InTextureRange)
 	{
 		PixelBuffers.AddZeroed(this->GetNumMips() * (bCubemap ? 6 : 1) * GetEffectiveSizeZ());
 		bAllocatedStorage.Init(bInAllocatedStorage, this->GetNumMips() * (bCubemap ? 6 : 1));
-#if !USE_PBO
-		TempBuffers.AddZeroed(this->GetNumMips() * (bCubemap ? 6 : 1) * GetEffectiveSizeZ());
-#endif
+		ClientStorageBuffers.AddZeroed(this->GetNumMips() * (bCubemap ? 6 : 1) * GetEffectiveSizeZ());
 	}
 
 	virtual ~TOpenGLTexture();
@@ -912,18 +904,17 @@ private:
 
 	TArray< TRefCountPtr<FOpenGLPixelBuffer> > PixelBuffers;
 
-#if !USE_PBO
-	struct FTempBuffer
+	/** Backing store for all client storage buffers for platforms and textures types where this is faster than PBOs */
+	uint8* TextureRange;
+
+	/** Client storage buffers for platforms and textures types where this is faster than PBOs */
+	struct FOpenGLClientStore
 	{
-		void* Data;
+		uint8* Data;
 		uint32 Size;
 		bool bReadOnly;
-
-		FTempBuffer() : Data(0), Size(0), bReadOnly(false)
-		{}
 	};
-	TArray< FTempBuffer> TempBuffers;
-#endif
+	TArray<FOpenGLClientStore> ClientStorageBuffers;
 
 	uint32 GetEffectiveSizeZ( void ) { return this->GetSizeZ() ? this->GetSizeZ() : 1; }
 
