Wrap OpenGL buffer and texture deletion calls so that on OS X they can be handled differently.

- On OS X glDeleteBuffers and glDeleteTextures will implictly finish the current command stream if they need to. This causes stalls that idle the GPU, so we want to avoid that.
- OpenGL wrappers now implement DeleteTextures and DeleteBuffers.
- On OS X DeleteTextures pools the texture names into one of three buffers so that textures are deleted at least two frames after their last use.
- The texture deletion occurs after flushBuffer in BlitToViewport and the min/max number of textures to delete in a single invocation can be overidden by the r.Mac.MinTexturesToDeletePerFrame & r.Mac.MaxTexturesToDeletePerFrame.
- On OS X buffers aren't actually deleted, instead their backing store is orphaned (BufferData with size 0) as this is substantially faster.

diff --git a/Engine/Source/Runtime/OpenGLDrv/Private/Mac/OpenGLMac.cpp b/Engine/Source/Runtime/OpenGLDrv/Private/Mac/OpenGLMac.cpp
--- a/Engine/Source/Runtime/OpenGLDrv/Private/Mac/OpenGLMac.cpp
+++ b/Engine/Source/Runtime/OpenGLDrv/Private/Mac/OpenGLMac.cpp
@@ -34,6 +34,26 @@ static FAutoConsoleVariableRef CVarMacUseMTGL(
 	ECVF_RenderThreadSafe
 	);
 
+static int32 GMacMinTexturesToDeletePerFrame = 8;
+static FAutoConsoleVariableRef CVarMacMinTexturesToDeletePerFrame(
+	TEXT("r.Mac.MinTexturesToDeletePerFrame"),
+	GMacMinTexturesToDeletePerFrame,
+	TEXT("Specifies how many textures are required to be waiting before a call to glDeleteTextures is made in order to amortize the cost. (Default: 8)"),
+	ECVF_RenderThreadSafe
+	);
+
+static int32 GMacMaxTexturesToDeletePerFrame = INT_MAX;
+static FAutoConsoleVariableRef CVarMacMaxTexturesToDeletePerFrame(
+	TEXT("r.Mac.MaxTexturesToDeletePerFrame"),
+	GMacMaxTexturesToDeletePerFrame,
+	TEXT("Specifies the maximum number of textures that can be deleted in the single call to glDeleteTextures each frame. (Default: INT_MAX)"),
+	ECVF_RenderThreadSafe
+	);
+
+static const uint32 GMacTexturePoolNum = 3;
+static TArray<GLuint> GMacTexturesToDelete[GMacTexturePoolNum];
+static TArray<GLuint> GMacBuffers;
+
 bool GIsRunningOnIntelCard = false; // @todo: remove once Apple fixes radr://16223045 Changes to the GL separate blend state aren't always respected on Intel cards
 static bool GIsEmulatingTimestamp = false; // @todo: Now crashing on Nvidia cards, but not on AMD...
 
@@ -464,6 +484,18 @@ struct FPlatformOpenGLDevice
 			
 			glDeleteTextures(1, &SharedContextCompositeTexture);
 			
+			for(uint32 i = 0; i < GMacTexturePoolNum; i++)
+			{
+				if(GMacTexturesToDelete[i].Num())
+				{
+					glDeleteTextures(GMacTexturesToDelete[i].Num(), GMacTexturesToDelete[i].GetData());
+					GMacTexturesToDelete[i].Reset();
+				}
+			}
+			
+			glDeleteBuffers(GMacBuffers.Num(), GMacBuffers.GetData());
+			GMacBuffers.Reset();
+			
 			if(GIsEmulatingTimestamp)
 			{
 				SharedContext.RunningQueries.Empty();
@@ -809,7 +841,14 @@ bool PlatformBlitToViewport( FPlatformOpenGLDevice* Device, const FOpenGLViewpor
 				}
 				
 				[Context->OpenGLContext flushBuffer];
-
+				
+				TArray<GLuint>& TexturesToDelete = GMacTexturesToDelete[(GFrameNumberRenderThread - (GMacTexturePoolNum - 1)) % GMacTexturePoolNum];
+				if(TexturesToDelete.Num() && TexturesToDelete.Num() > GMacMinTexturesToDeletePerFrame)
+				{
+					uint32 Num = FMath::Min(TexturesToDelete.Num(), GMacMaxTexturesToDeletePerFrame);
+					glDeleteTextures(Num, TexturesToDelete.GetData());
+					TexturesToDelete.RemoveAt(0, Num, false);
+				}
 				
 				glEnable(GL_FRAMEBUFFER_SRGB);
 				
@@ -1814,3 +1853,58 @@ bool FMacOpenGL::MustFlushTexStorage(void)
 	// @todo There is a bug in Apple's GL with TexStorage calls, on Nvidia and when using MTGL that can see the texture never be created, which then subsequently causes crashes
 	return GMacMustFlushTexStorage || GMacUseMTGL;
 }
+
+void FMacOpenGL::GenBuffers( GLsizei n, GLuint *buffers)
+{
+#if USE_OPENGL_NAME_CACHE
+	if( n < GMacBuffers.Num() )
+	{
+		FMemory::Memcpy( buffers, GMacBuffers.GetData(), sizeof(GLuint)*n);
+		GMacBuffers.RemoveAt(0, n, false);
+	}
+	else if(n < OPENGL_NAME_CACHE_SIZE)
+	{
+		GLsizei Leftover = GMacBuffers.Num();
+		if(Leftover)
+		{
+			FMemory::Memcpy( buffers, GMacBuffers.GetData(), sizeof(GLuint)*Leftover);
+			GMacBuffers.RemoveAt(0, Leftover, false);
+		}
+		
+		GMacBuffers.Reserve(OPENGL_NAME_CACHE_SIZE);
+		GLuint Buffer[OPENGL_NAME_CACHE_SIZE] = {0};
+		glGenBuffers( OPENGL_NAME_CACHE_SIZE, Buffer);
+		GMacBuffers.Append(Buffer, OPENGL_NAME_CACHE_SIZE);
+		
+		n -= Leftover;
+		buffers += Leftover;
+		
+		FMemory::Memcpy( buffers, GMacBuffers.GetData(), sizeof(GLuint)*n);
+		GMacBuffers.RemoveAt(0, n, false);
+	}
+	else
+	{
+		glGenBuffers(n, buffers);
+	}
+#else
+	glGenBuffers( n, buffers);
+#endif
+}
+
+void FMacOpenGL::DeleteBuffers(GLsizei Number, const GLuint* Buffers)
+{
+	for(uint32 i = 0; i < Number; i++)
+	{
+		glBindBuffer(GL_COPY_WRITE_BUFFER, Buffers[i]);
+		glBufferData(GL_COPY_WRITE_BUFFER, 0, nullptr, GL_STATIC_DRAW);
+		GMacBuffers.Add(Buffers[i]);
+	}
+	glBindBuffer(GL_COPY_WRITE_BUFFER, 0);
+}
+
+void FMacOpenGL::DeleteTextures(GLsizei Number, const GLuint* Textures)
+{
+	TArray<GLuint>& TexturesToDelete = GMacTexturesToDelete[GFrameNumberRenderThread % GMacTexturePoolNum];
+	TexturesToDelete.Append(Textures, Number);
+}
+
diff --git a/Engine/Source/Runtime/OpenGLDrv/Private/Mac/OpenGLMac.h b/Engine/Source/Runtime/OpenGLDrv/Private/Mac/OpenGLMac.h
--- a/Engine/Source/Runtime/OpenGLDrv/Private/Mac/OpenGLMac.h
+++ b/Engine/Source/Runtime/OpenGLDrv/Private/Mac/OpenGLMac.h
@@ -148,6 +148,12 @@ struct FMacOpenGL : public FOpenGL3
 			return false;
 		}
 	}
+	
+	static void GenBuffers( GLsizei n, GLuint *buffers);
+	
+	static void DeleteBuffers(GLsizei Number, const GLuint* Buffers);
+	
+	static void DeleteTextures(GLsizei Number, const GLuint* Textures);
 
 	static FORCEINLINE bool SupportsSeparateAlphaBlend()				{ return bSupportsDrawBuffersBlend; }
 	
diff --git a/Engine/Source/Runtime/OpenGLDrv/Private/OpenGLDevice.cpp b/Engine/Source/Runtime/OpenGLDrv/Private/OpenGLDevice.cpp
--- a/Engine/Source/Runtime/OpenGLDrv/Private/OpenGLDevice.cpp
+++ b/Engine/Source/Runtime/OpenGLDrv/Private/OpenGLDevice.cpp
@@ -1092,7 +1092,7 @@ void FOpenGLDynamicRHI::Cleanup()
 	// Release zero-filled dummy uniform buffer, if it exists.
 	if (PendingState.ZeroFilledDummyUniformBuffer)
 	{
-		glDeleteBuffers(1, &PendingState.ZeroFilledDummyUniformBuffer);
+		FOpenGL::DeleteBuffers(1, &PendingState.ZeroFilledDummyUniformBuffer);
 		PendingState.ZeroFilledDummyUniformBuffer = 0;
 		DecrementBufferMemory(GL_UNIFORM_BUFFER, false, ZERO_FILLED_DUMMY_UNIFORM_BUFFER_SIZE);
 	}
diff --git a/Engine/Source/Runtime/OpenGLDrv/Private/OpenGLRenderTarget.cpp b/Engine/Source/Runtime/OpenGLDrv/Private/OpenGLRenderTarget.cpp
--- a/Engine/Source/Runtime/OpenGLDrv/Private/OpenGLRenderTarget.cpp
+++ b/Engine/Source/Runtime/OpenGLDrv/Private/OpenGLRenderTarget.cpp
@@ -811,7 +811,7 @@ void FOpenGLDynamicRHI::RHIRead3DSurfaceFloatData(FTextureRHIParamRef TextureRHI
 	glBindTexture(GL_TEXTURE_3D, (TextureState.Target == GL_TEXTURE_3D) ? TextureState.Resource : 0);
 	glActiveTexture( GL_TEXTURE0 + ContextState.ActiveTexture );
 	glDeleteFramebuffers( 1, &SourceFramebuffer);
-	glDeleteTextures( 1, &TempTexture );
+	FOpenGL::DeleteTextures( 1, &TempTexture );
 	ContextState.Framebuffer = (GLuint)-1;
 }
 
diff --git a/Engine/Source/Runtime/OpenGLDrv/Private/OpenGLTexture.cpp b/Engine/Source/Runtime/OpenGLDrv/Private/OpenGLTexture.cpp
--- a/Engine/Source/Runtime/OpenGLDrv/Private/OpenGLTexture.cpp
+++ b/Engine/Source/Runtime/OpenGLDrv/Private/OpenGLTexture.cpp
@@ -1385,7 +1385,7 @@ TOpenGLTexture<RHIResourceType>::~TOpenGLTexture()
 	if( Resource != 0 )
 	{
 		OpenGLRHI->InvalidateTextureResourceInCache( Resource );
-		glDeleteTextures( 1, &Resource );
+		FOpenGL::DeleteTextures( 1, &Resource );
 	}
 	
 	if(TextureRange)
diff --git a/Engine/Source/Runtime/OpenGLDrv/Private/OpenGLUAV.cpp b/Engine/Source/Runtime/OpenGLDrv/Private/OpenGLUAV.cpp
--- a/Engine/Source/Runtime/OpenGLDrv/Private/OpenGLUAV.cpp
+++ b/Engine/Source/Runtime/OpenGLDrv/Private/OpenGLUAV.cpp
@@ -39,7 +39,7 @@ FOpenGLShaderResourceView::~FOpenGLShaderResourceView()
 	if (Resource && OwnsResource)
 	{
 		OpenGLRHI->InvalidateTextureResourceInCache( Resource );
-		glDeleteTextures(1, &Resource);
+		FOpenGL::DeleteTextures(1, &Resource);
 	}
 }
 
@@ -101,7 +101,7 @@ FOpenGLVertexBufferUnorderedAccessView::~FOpenGLVertexBufferUnorderedAccessView(
 	if (Resource)
 	{
 		OpenGLRHI->InvalidateTextureResourceInCache( Resource );
-		glDeleteTextures(1, &Resource);
+		FOpenGL::DeleteTextures(1, &Resource);
 	}
 }
 
diff --git a/Engine/Source/Runtime/OpenGLDrv/Private/OpenGLUniformBuffer.cpp b/Engine/Source/Runtime/OpenGLDrv/Private/OpenGLUniformBuffer.cpp
--- a/Engine/Source/Runtime/OpenGLDrv/Private/OpenGLUniformBuffer.cpp
+++ b/Engine/Source/Runtime/OpenGLDrv/Private/OpenGLUniformBuffer.cpp
@@ -228,7 +228,7 @@ void BeginFrame_UniformBufferPoolCleanup()
 					}
 					else
 					{
-						glDeleteBuffers( 1, &PoolEntry.Buffer );
+						FOpenGL::DeleteBuffers( 1, &PoolEntry.Buffer );
 					}
 					GLUniformBufferPool[BucketIndex][StreamedIndex].RemoveAtSwap(EntryIndex);
 
@@ -381,7 +381,7 @@ FOpenGLUniformBuffer::~FOpenGLUniformBuffer()
 		}
 		else
 		{
-			glDeleteBuffers(1, &Resource);
+			FOpenGL::DeleteBuffers(1, &Resource);
 			DecrementBufferMemory(GL_UNIFORM_BUFFER, /*bIsStructuredBuffer=*/ false, AllocatedSize);
 		}
 	}
diff --git a/Engine/Source/Runtime/OpenGLDrv/Public/OpenGL3.h b/Engine/Source/Runtime/OpenGLDrv/Public/OpenGL3.h
--- a/Engine/Source/Runtime/OpenGLDrv/Public/OpenGL3.h
+++ b/Engine/Source/Runtime/OpenGLDrv/Public/OpenGL3.h
@@ -479,6 +479,16 @@ struct FOpenGL3 : public FOpenGLBase
 		glGenTextures( n, textures);
 #endif
 	}
+	
+	static FORCEINLINE void DeleteBuffers(GLsizei Number, const GLuint* Buffers)
+	{
+		glDeleteBuffers(Number, Buffers);
+	}
+	
+	static FORCEINLINE void DeleteTextures(GLsizei Number, const GLuint* Textures)
+	{
+		glDeleteTextures(Number, Textures);
+	}
 
 	static FORCEINLINE void CompressedTexSubImage3D(GLenum Target, GLint Level, GLint XOffset, GLint YOffset, GLint ZOffset, GLsizei Width, GLsizei Height, GLsizei Depth, GLenum Format, GLsizei ImageSize, const GLvoid* PixelData)
 	{
diff --git a/Engine/Source/Runtime/OpenGLDrv/Public/OpenGLES2.h b/Engine/Source/Runtime/OpenGLDrv/Public/OpenGLES2.h
--- a/Engine/Source/Runtime/OpenGLDrv/Public/OpenGLES2.h
+++ b/Engine/Source/Runtime/OpenGLDrv/Public/OpenGLES2.h
@@ -333,6 +333,16 @@ struct FOpenGLES2 : public FOpenGLBase
 	{
 		glGenTextures( n, textures);
 	}
+	
+	static FORCEINLINE void DeleteBuffers(GLsizei Number, const GLuint* Buffers)
+	{
+		glDeleteBuffers(Number, Buffers);
+	}
+	
+	static FORCEINLINE void DeleteTextures(GLsizei Number, const GLuint* Textures)
+	{
+		glDeleteTextures(Number, Textures);
+	}
 
 protected:
 	/** GL_OES_vertex_array_object */
diff --git a/Engine/Source/Runtime/OpenGLDrv/Public/OpenGLResources.h b/Engine/Source/Runtime/OpenGLDrv/Public/OpenGLResources.h
--- a/Engine/Source/Runtime/OpenGLDrv/Public/OpenGLResources.h
+++ b/Engine/Source/Runtime/OpenGLDrv/Public/OpenGLResources.h
@@ -143,7 +143,7 @@ public:
 		VERIFY_GL_SCOPE();
 		if (Resource != 0 && BaseType::OnDelete(Resource,RealSize,bStreamDraw,0))
 		{
-			glDeleteBuffers(1, &Resource);
+			FOpenGL::DeleteBuffers(1, &Resource);
 			DecrementBufferMemory(Type, BaseType::IsStructuredBuffer(), RealSize);
 		}
 		if (LockBuffer != NULL)
