#UE4 Fix for bogus assert when using ScreenAlignedUVs (really anything that makes use of the entries with higher offset of any cbuffer)

diff --git a/Engine/Source/Runtime/OpenGLDrv/Private/OpenGLShaders.cpp b/Engine/Source/Runtime/OpenGLDrv/Private/OpenGLShaders.cpp
--- a/Engine/Source/Runtime/OpenGLDrv/Private/OpenGLShaders.cpp
+++ b/Engine/Source/Runtime/OpenGLDrv/Private/OpenGLShaders.cpp
@@ -1915,7 +1915,7 @@ void FOpenGLShaderParameterCache::CommitPackedUniformBuffers(FOpenGLLinkedProgra
 				const FOpenGLUniformBufferCopyInfo& Info = UniformBuffersCopyInfo[InfoIndex];
 				if (Info.SourceUBIndex == BufferIndex)
 				{
-					check((Info.SourceOffsetInFloats + Info.SizeInFloats) * sizeof(float) <= (uint32)GlobalUniformArraySize);
+					check((Info.DestOffsetInFloats + Info.SizeInFloats) * sizeof(float) <= (uint32)GlobalUniformArraySize);
 
 					float* RESTRICT ScratchMem = (float*)PackedGlobalUniforms[Info.DestUBTypeIndex];
 					ScratchMem += Info.DestOffsetInFloats;
