Linux: minor fixes from PR #306.

diff --git a/Engine/Source/Runtime/OpenGLDrv/Private/OpenGLDevice.cpp b/Engine/Source/Runtime/OpenGLDrv/Private/OpenGLDevice.cpp
--- a/Engine/Source/Runtime/OpenGLDrv/Private/OpenGLDevice.cpp
+++ b/Engine/Source/Runtime/OpenGLDrv/Private/OpenGLDevice.cpp
@@ -418,7 +418,7 @@ void InitDebugContext()
 
 	// this is to suppress feeding back of the debug markers and groups to the log, since those originate in the app anyways...
 #if ENABLE_OPENGL_DEBUG_GROUPS && GL_ARB_debug_output && GL_KHR_debug
-	if(glDebugMessageControlARB)
+	if(glDebugMessageControlARB && bDebugOutputInitialized)
 	{
 		glDebugMessageControlARB(GL_DEBUG_SOURCE_APPLICATION_ARB, GL_DEBUG_TYPE_MARKER, GL_DONT_CARE, 0, NULL, GL_FALSE);
 		glDebugMessageControlARB(GL_DEBUG_SOURCE_APPLICATION_ARB, GL_DEBUG_TYPE_PUSH_GROUP, GL_DONT_CARE, 0, NULL, GL_FALSE);
diff --git a/Engine/Source/Runtime/OpenGLDrv/Private/OpenGLShaders.cpp b/Engine/Source/Runtime/OpenGLDrv/Private/OpenGLShaders.cpp
--- a/Engine/Source/Runtime/OpenGLDrv/Private/OpenGLShaders.cpp
+++ b/Engine/Source/Runtime/OpenGLDrv/Private/OpenGLShaders.cpp
@@ -10,7 +10,7 @@
 
 #define CHECK_FOR_GL_SHADERS_TO_REPLACE 0
 
-#if PLATFORM_WINDOWS
+#if PLATFORM_WINDOWS || PLATFORM_LINUX
 #include <mmintrin.h>
 #elif PLATFORM_MAC
 #include <xmmintrin.h>
@@ -48,7 +48,7 @@ static bool VerifyCompiledShader(GLuint Shader, const ANSICHAR* GlslCode )
 			CompileLog = (ANSICHAR *)FMemory::Malloc(LogLength);
 			glGetShaderInfoLog(Shader, LogLength, NULL, CompileLog);
 		}
-	
+
 #if DEBUG_GL_SHADERS
 		if (GlslCode)
 		{
@@ -425,10 +425,10 @@ ShaderType* CompileOpenGLShader(const TArray<uint8>& Code)
 		const GLint ShaderSourceLen[4] = { (GLint)(strlen(VersionString)), (GLint)(strlen(ExtensionString)), (GLint)(strlen(Prologue)), (GLint)(strlen(GlslCode)) };
 		glShaderSource(Resource, 4, (const GLchar**)&ShaderSourceStrings, ShaderSourceLen);		
 		glCompileShader(Resource);
-#else
+#else // PLATFORM_ANDROID || PLATFORM_HTML5
 		glShaderSource(Resource, 1, (const GLchar**)&GlslCode, &GlslCodeLength);
 		glCompileShader(Resource);
-#endif
+#endif // PLATFORM_ANDROID || PLATFORM_HTML5
 		GetOpenGLCompiledShaderCache().Add(Key,Resource);
 	}
 
