Linux: disable sRGB when blitting to viewport.

diff --git a/Engine/Source/Runtime/OpenGLDrv/Private/Linux/OpenGLLinux.cpp b/Engine/Source/Runtime/OpenGLDrv/Private/Linux/OpenGLLinux.cpp
--- a/Engine/Source/Runtime/OpenGLDrv/Private/Linux/OpenGLLinux.cpp
+++ b/Engine/Source/Runtime/OpenGLDrv/Private/Linux/OpenGLLinux.cpp
@@ -357,6 +357,7 @@ bool PlatformBlitToViewport(FPlatformOpenGLDevice* Device,
 		glDrawBuffer( GL_BACK );
 		glBindFramebuffer( GL_READ_FRAMEBUFFER, Context->ViewportFramebuffer );
 		glReadBuffer( GL_COLOR_ATTACHMENT0 );
+		glDisable(GL_FRAMEBUFFER_SRGB);
 
 		glBlitFramebuffer(	0, 0, BackbufferSizeX, BackbufferSizeY,
 							0, BackbufferSizeY, BackbufferSizeX, 0,
@@ -379,6 +380,7 @@ bool PlatformBlitToViewport(FPlatformOpenGLDevice* Device,
 
 			SDL_GL_SwapWindow( Context->hWnd );
 
+			glEnable(GL_FRAMEBUFFER_SRGB);
 			REPORT_GL_END_BUFFER_EVENT_FOR_FRAME_DUMP();
 //			INITIATE_GL_FRAME_DUMP_EVERY_X_CALLS( 1000 );
 		}
