Wrap glFlush so that on OS X we always call glFlushRenderAPPLE which is faster, while still sufficient for these calls.

diff --git a/Engine/Source/Runtime/OpenGLDrv/Private/Mac/OpenGLMac.h b/Engine/Source/Runtime/OpenGLDrv/Private/Mac/OpenGLMac.h
--- a/Engine/Source/Runtime/OpenGLDrv/Private/Mac/OpenGLMac.h
+++ b/Engine/Source/Runtime/OpenGLDrv/Private/Mac/OpenGLMac.h
@@ -221,6 +221,11 @@ struct FMacOpenGL : public FOpenGL3
 		glTextureRangeAPPLE(Target, Length, Pointer);
 	}
 	
+	static FORCEINLINE void Flush()
+	{
+		glFlushRenderAPPLE();
+	}
+	
 	static FORCEINLINE bool SupportsSeamlessCubeMap()					{ return true; }
 	static FORCEINLINE bool SupportsClientStorage()						{ return true; }
 	static FORCEINLINE bool SupportsTextureRange()						{ return true; }
diff --git a/Engine/Source/Runtime/OpenGLDrv/Private/OpenGLCommands.cpp b/Engine/Source/Runtime/OpenGLDrv/Private/OpenGLCommands.cpp
--- a/Engine/Source/Runtime/OpenGLDrv/Private/OpenGLCommands.cpp
+++ b/Engine/Source/Runtime/OpenGLDrv/Private/OpenGLCommands.cpp
@@ -1300,7 +1300,6 @@ void FOpenGLDynamicRHI::SetPendingBlendStateForActiveRenderTargets( FOpenGLConte
 					if (GIsRunningOnIntelCard)
 					{
 						// Flush the separate blend state changes through GL or on Intel cards under OS X the state change may be silently ignored.
-						// We may wish to consider using GL_APPLE_flush_render for this workaround if it is still faster to call glFlushRenderAPPLE than plain glFlush.
 						if (CachedRenderTargetBlendState.ColorSourceBlendFactor != RenderTargetBlendState.ColorSourceBlendFactor
 							|| CachedRenderTargetBlendState.ColorDestBlendFactor != RenderTargetBlendState.ColorDestBlendFactor
 							|| CachedRenderTargetBlendState.AlphaSourceBlendFactor != RenderTargetBlendState.AlphaSourceBlendFactor
@@ -1308,7 +1307,7 @@ void FOpenGLDynamicRHI::SetPendingBlendStateForActiveRenderTargets( FOpenGLConte
 							|| CachedRenderTargetBlendState.ColorBlendOperation != RenderTargetBlendState.ColorBlendOperation
 							|| CachedRenderTargetBlendState.AlphaBlendOperation != RenderTargetBlendState.AlphaBlendOperation)
 						{
-							glFlush();
+							FOpenGL::Flush();
 						}
 					}
 #endif
diff --git a/Engine/Source/Runtime/OpenGLDrv/Private/OpenGLDevice.cpp b/Engine/Source/Runtime/OpenGLDrv/Private/OpenGLDevice.cpp
--- a/Engine/Source/Runtime/OpenGLDrv/Private/OpenGLDevice.cpp
+++ b/Engine/Source/Runtime/OpenGLDrv/Private/OpenGLDevice.cpp
@@ -943,7 +943,7 @@ void FOpenGLDynamicRHI::Init()
 #endif
 
 	// Flush here since we might be switching to a different context/thread for rendering
-	glFlush();
+	FOpenGL::Flush();
 
 	FHardwareInfo::RegisterHardwareInfo( NAME_RHI, TEXT( "OpenGL" ) );
 
diff --git a/Engine/Source/Runtime/OpenGLDrv/Private/OpenGLQuery.cpp b/Engine/Source/Runtime/OpenGLDrv/Private/OpenGLQuery.cpp
--- a/Engine/Source/Runtime/OpenGLDrv/Private/OpenGLQuery.cpp
+++ b/Engine/Source/Runtime/OpenGLDrv/Private/OpenGLQuery.cpp
@@ -254,7 +254,7 @@ void FOpenGLEventQuery::IssueEvent()
 		Sync = 0;
 	}
 	Sync = FOpenGL::FenceSync(GL_SYNC_GPU_COMMANDS_COMPLETE, 0);
-	glFlush();
+	FOpenGL::Flush();
 	
 	checkSlow(FOpenGL::IsSync(Sync));
 
diff --git a/Engine/Source/Runtime/OpenGLDrv/Public/OpenGL.h b/Engine/Source/Runtime/OpenGLDrv/Public/OpenGL.h
--- a/Engine/Source/Runtime/OpenGLDrv/Public/OpenGL.h
+++ b/Engine/Source/Runtime/OpenGLDrv/Public/OpenGL.h
@@ -290,6 +290,7 @@ public:
 	static FORCEINLINE void BufferStorage(GLenum Target, GLsizeiptr Size, const void *Data, GLbitfield Flags) UGL_REQUIRED_VOID
 	static FORCEINLINE void DepthBounds(GLfloat Min, GLfloat Max) UGL_REQUIRED_VOID
 	static FORCEINLINE void TextureRange(GLenum Target, GLsizei Length, const GLvoid *Pointer) UGL_OPTIONAL_VOID
+	static FORCEINLINE void Flush() UGL_REQUIRED_VOID
 
 	static FPlatformOpenGLDevice*	CreateDevice() UGL_REQUIRED(NULL)
 	static FPlatformOpenGLContext*	CreateContext( FPlatformOpenGLDevice* Device, void* WindowHandle ) UGL_REQUIRED(NULL)
diff --git a/Engine/Source/Runtime/OpenGLDrv/Public/OpenGL3.h b/Engine/Source/Runtime/OpenGLDrv/Public/OpenGL3.h
--- a/Engine/Source/Runtime/OpenGLDrv/Public/OpenGL3.h
+++ b/Engine/Source/Runtime/OpenGLDrv/Public/OpenGL3.h
@@ -489,6 +489,11 @@ struct FOpenGL3 : public FOpenGLBase
 	{
 		glDeleteTextures(Number, Textures);
 	}
+	
+	static FORCEINLINE void Flush()
+	{
+		glFlush();
+	}
 
 	static FORCEINLINE void CompressedTexSubImage3D(GLenum Target, GLint Level, GLint XOffset, GLint YOffset, GLint ZOffset, GLsizei Width, GLsizei Height, GLsizei Depth, GLenum Format, GLsizei ImageSize, const GLvoid* PixelData)
 	{
diff --git a/Engine/Source/Runtime/OpenGLDrv/Public/OpenGLES2.h b/Engine/Source/Runtime/OpenGLDrv/Public/OpenGLES2.h
--- a/Engine/Source/Runtime/OpenGLDrv/Public/OpenGLES2.h
+++ b/Engine/Source/Runtime/OpenGLDrv/Public/OpenGLES2.h
@@ -343,6 +343,11 @@ struct FOpenGLES2 : public FOpenGLBase
 	{
 		glDeleteTextures(Number, Textures);
 	}
+	
+	static FORCEINLINE void Flush()
+	{
+		glFlush();
+	}
 
 protected:
 	/** GL_OES_vertex_array_object */
