Engine: Fix unused expression (missing UE_LOG)

diff --git a/Engine/Source/Runtime/OpenGLDrv/Private/OpenGLDrv.cpp b/Engine/Source/Runtime/OpenGLDrv/Private/OpenGLDrv.cpp
--- a/Engine/Source/Runtime/OpenGLDrv/Private/OpenGLDrv.cpp
+++ b/Engine/Source/Runtime/OpenGLDrv/Private/OpenGLDrv.cpp
@@ -337,7 +337,7 @@ void FOpenGLEventNodeFrame::LogDisjointQuery()
 	}
 	else
 	{
-		TEXT("Profiled range \"disjoinness\" could not be determined due to lack of disjoing timer query functionality on this platform.");
+		UE_LOG(LogRHI, Warning, TEXT("Profiled range \"disjoinness\" could not be determined due to lack of disjoint timer query functionality on this platform."));
 	}
 }
 
