Fix crash with opengl UBOs when ES2 preview mode is toggled.
ES2 emulation now uses UBOs instead of packed uniform blocks.
Fixed issue with ogl shader compiler introducing '__' into padding variable names.

diff --git a/Engine/Source/Developer/ShaderFormatOpenGL/Private/OpenGLShaderCompiler.cpp b/Engine/Source/Developer/ShaderFormatOpenGL/Private/OpenGLShaderCompiler.cpp
--- a/Engine/Source/Developer/ShaderFormatOpenGL/Private/OpenGLShaderCompiler.cpp
+++ b/Engine/Source/Developer/ShaderFormatOpenGL/Private/OpenGLShaderCompiler.cpp
@@ -1611,14 +1611,11 @@ void CompileShader_Windows_OGL(const FShaderCompilerInput& Input,FShaderCompiler
 		}
 
 		uint32 CCFlags = HLSLCC_NoPreprocess | HLSLCC_PackUniforms | HLSLCC_DX11ClipSpace;
-		if (IsES2Platform(Version))
+		if (IsES2Platform(Version) && !IsPCES2Platform(Version))
 		{
 			CCFlags |= HLSLCC_FlattenUniformBuffers | HLSLCC_FlattenUniformBufferStructures;
 			// Currently only enabled for ES2, as there are still features to implement for SM4+ (atomics, global store, UAVs, etc)
-			if (!IsPCES2Platform(Version))
-			{
-				CCFlags |= HLSLCC_ApplyCommonSubexpressionElimination;
-			}
+			CCFlags |= HLSLCC_ApplyCommonSubexpressionElimination;
 		}
 
 		if (bDumpDebugInfo)
diff --git a/Engine/Source/Developer/ShaderFormatOpenGL/Private/ShaderFormatOpenGL.cpp b/Engine/Source/Developer/ShaderFormatOpenGL/Private/ShaderFormatOpenGL.cpp
--- a/Engine/Source/Developer/ShaderFormatOpenGL/Private/ShaderFormatOpenGL.cpp
+++ b/Engine/Source/Developer/ShaderFormatOpenGL/Private/ShaderFormatOpenGL.cpp
@@ -22,10 +22,10 @@ class FShaderFormatGLSL : public IShaderFormat
 	{
 		/** Version for shader format, this becomes part of the DDC key. */
 		UE_SHADER_GLSL_150_VER = 23,
-		UE_SHADER_GLSL_150_MAC_VER = 2,
+		UE_SHADER_GLSL_150_MAC_VER = 3,
 		UE_SHADER_GLSL_430_VER = 53,
 		UE_SHADER_GLSL_ES2_VER = 12,
-		UE_SHADER_GLSL_150ES2_VER = 13,
+		UE_SHADER_GLSL_150ES2_VER = 14,
 		UE_SHADER_GLSL_ES2_VER_WEBGL  = 13,
 		UE_SHADER_GLSL_ES2_IOS_VER  = 3,
 	}; 
diff --git a/Engine/Source/Runtime/OpenGLDrv/Private/OpenGLCommands.cpp b/Engine/Source/Runtime/OpenGLDrv/Private/OpenGLCommands.cpp
--- a/Engine/Source/Runtime/OpenGLDrv/Private/OpenGLCommands.cpp
+++ b/Engine/Source/Runtime/OpenGLDrv/Private/OpenGLCommands.cpp
@@ -1486,7 +1486,7 @@ void FOpenGLDynamicRHI::RHISetRenderTargets(
 
 	FOpenGLTextureBase* NewDepthStencilRT = GetOpenGLTextureFromRHITexture(NewDepthStencilTargetRHI);
 
-	if (IsES2Platform(GRHIShaderPlatform))
+	if (IsES2Platform(GRHIShaderPlatform) && !IsPCPlatform(GRHIShaderPlatform))
 	{
 		// @todo-mobile
 
@@ -2079,13 +2079,13 @@ void FOpenGLDynamicRHI::CommitNonComputeShaderConstants()
 	VERIFY_GL_SCOPE();
 
 	FOpenGLLinkedProgram* LinkedProgram = PendingState.BoundShaderState->LinkedProgram;
-	if (IsES2Platform(GRHIShaderPlatform))
+	if (GUseEmulatedUniformBuffers)
 	{
 		PendingState.ShaderParameters[CrossCompiler::SHADER_STAGE_VERTEX].CommitPackedUniformBuffers(LinkedProgram, CrossCompiler::SHADER_STAGE_VERTEX, PendingState.BoundUniformBuffers[SF_Vertex], PendingState.BoundShaderState->VertexShader->UniformBuffersCopyInfo);
 	}
 	PendingState.ShaderParameters[CrossCompiler::SHADER_STAGE_VERTEX].CommitPackedGlobals(LinkedProgram, CrossCompiler::SHADER_STAGE_VERTEX);
 
-	if (IsES2Platform(GRHIShaderPlatform))
+	if (GUseEmulatedUniformBuffers)
 	{
 		PendingState.ShaderParameters[CrossCompiler::SHADER_STAGE_PIXEL].CommitPackedUniformBuffers(LinkedProgram, CrossCompiler::SHADER_STAGE_PIXEL, PendingState.BoundUniformBuffers[SF_Pixel], PendingState.BoundShaderState->PixelShader->UniformBuffersCopyInfo);
 	}
@@ -2093,7 +2093,7 @@ void FOpenGLDynamicRHI::CommitNonComputeShaderConstants()
 
 	if (PendingState.BoundShaderState->GeometryShader)
 	{
-		if (IsES2Platform(GRHIShaderPlatform))
+		if (GUseEmulatedUniformBuffers)
 		{
 			PendingState.ShaderParameters[CrossCompiler::SHADER_STAGE_GEOMETRY].CommitPackedUniformBuffers(LinkedProgram, CrossCompiler::SHADER_STAGE_GEOMETRY, PendingState.BoundUniformBuffers[SF_Geometry], PendingState.BoundShaderState->GeometryShader->UniformBuffersCopyInfo);
 		}
diff --git a/Engine/Source/Runtime/OpenGLDrv/Private/OpenGLDevice.cpp b/Engine/Source/Runtime/OpenGLDrv/Private/OpenGLDevice.cpp
--- a/Engine/Source/Runtime/OpenGLDrv/Private/OpenGLDevice.cpp
+++ b/Engine/Source/Runtime/OpenGLDrv/Private/OpenGLDevice.cpp
@@ -19,6 +19,9 @@ FOpenGLTextureFormat GOpenGLTextureFormats[PF_MAX];
 /** Device is necessary for vertex buffers, so they can reach the global device on destruction, and tell it to reset vertex array caches */
 static FOpenGLDynamicRHI* PrivateOpenGLDevicePtr = NULL;
 
+/** true if we're not using UBOs. (ES2) */
+bool GUseEmulatedUniformBuffers;
+
 void OnQueryCreation( FOpenGLRenderQuery* Query )
 {
 	check(PrivateOpenGLDevicePtr);
@@ -609,6 +612,9 @@ static void InitRHICapabilitiesForGL()
 	GMaxRHIFeatureLevelValue = GRHIFeatureLevelValue = FOpenGL::GetFeatureLevel();
 	GRHIShaderPlatformValue = FOpenGL::GetShaderPlatform();
 
+	// Emulate uniform buffers on ES2, unless we're on a desktop platform emulating ES2.
+	GUseEmulatedUniformBuffers = IsES2Platform(GRHIShaderPlatform) && !IsPCPlatform(GRHIShaderPlatform);
+
 	FString FeatureLevelName;
 	GetFeatureLevelName(GMaxRHIFeatureLevel, FeatureLevelName);
 	FString ShaderPlatformName = LegacyShaderPlatformToShaderFormat(GRHIShaderPlatform).ToString();
diff --git a/Engine/Source/Runtime/OpenGLDrv/Private/OpenGLDrvPrivate.h b/Engine/Source/Runtime/OpenGLDrv/Private/OpenGLDrvPrivate.h
--- a/Engine/Source/Runtime/OpenGLDrv/Private/OpenGLDrvPrivate.h
+++ b/Engine/Source/Runtime/OpenGLDrv/Private/OpenGLDrvPrivate.h
@@ -481,4 +481,6 @@ public:
 
 extern TGlobalResource<FVector4VertexDeclaration> GOpenGLVector4VertexDeclaration;
 
+extern bool GUseEmulatedUniformBuffers;
+
 #endif // __OPENGLDRVPRIVATE_H__
diff --git a/Engine/Source/Runtime/OpenGLDrv/Private/OpenGLShaders.cpp b/Engine/Source/Runtime/OpenGLDrv/Private/OpenGLShaders.cpp
--- a/Engine/Source/Runtime/OpenGLDrv/Private/OpenGLShaders.cpp
+++ b/Engine/Source/Runtime/OpenGLDrv/Private/OpenGLShaders.cpp
@@ -1596,7 +1596,7 @@ void FOpenGLDynamicRHI::BindPendingShaderState( FOpenGLContextState& ContextStat
 #endif
 	}
 
-	if (!IsES2Platform(GRHIShaderPlatform))
+	if (!GUseEmulatedUniformBuffers)
 	{
 		int32 NextUniformBufferIndex = OGL_FIRST_UNIFORM_BUFFER;
 
@@ -1736,7 +1736,7 @@ void FOpenGLDynamicRHI::BindPendingComputeShaderState(FOpenGLContextState& Conte
 		ForceUniformBindingUpdate = true;
 	}
 
-	if (!IsES2Platform(GRHIShaderPlatform))
+	if (!GUseEmulatedUniformBuffers)
 	{
 		BindUniformBufferBase(
 			ContextState,
diff --git a/Engine/Source/Runtime/OpenGLDrv/Private/OpenGLUniformBuffer.cpp b/Engine/Source/Runtime/OpenGLDrv/Private/OpenGLUniformBuffer.cpp
--- a/Engine/Source/Runtime/OpenGLDrv/Private/OpenGLUniformBuffer.cpp
+++ b/Engine/Source/Runtime/OpenGLDrv/Private/OpenGLUniformBuffer.cpp
@@ -60,7 +60,7 @@ static TArray<uint32> UniformBufferSizeBuckets;
 static inline bool IsSuballocatingUBOs()
 {
 #if SUBALLOCATED_CONSTANT_BUFFER
-	if (!IsES2Platform(GRHIShaderPlatform))
+	if (!GUseEmulatedUniformBuffers)
 	{
 		return OpenGLConsoleVariables::RequestedUBOPoolSize != 0;
 	}
@@ -205,7 +205,6 @@ void BeginFrame_UniformBufferPoolCleanup()
 	if (!IsSuballocatingUBOs())
 	{
 	// Clean a limited number of old entries to reduce hitching when leaving a large level
-	const bool bIsES2 = IsES2Platform(GRHIShaderPlatform);
 	for( int32 StreamedIndex = 0; StreamedIndex < 2; ++StreamedIndex)
 	{
 			for (int32 BucketIndex = 0; BucketIndex < UniformBufferSizeBuckets.Num(); BucketIndex++)
@@ -222,7 +221,7 @@ void BeginFrame_UniformBufferPoolCleanup()
 					DEC_DWORD_STAT(STAT_OpenGLNumFreeUniformBuffers);
 					DEC_MEMORY_STAT_BY(STAT_OpenGLFreeUniformBufferMemory, PoolEntry.CreatedSize);
 					DecrementBufferMemory(GL_UNIFORM_BUFFER, /*bIsStructuredBuffer=*/ false, PoolEntry.CreatedSize);
-					if (bIsES2)
+					if (GUseEmulatedUniformBuffers)
 					{
 						UniformBufferDataFactory.Destroy(PoolEntry.Buffer);
 					}
@@ -485,7 +484,7 @@ FUniformBufferRHIRef FOpenGLDynamicRHI::RHICreateUniformBuffer(const void* Conte
 				AllocatedResource = FreeBufferEntry.Buffer;
 				AllocatedSize = FreeBufferEntry.CreatedSize;
 
-				if (IsES2Platform(GRHIShaderPlatform))
+				if (GUseEmulatedUniformBuffers)
 				{
 					EmulatedUniformDataRef = UniformBufferDataFactory.Get(AllocatedResource);
 				}
@@ -507,7 +506,7 @@ FUniformBufferRHIRef FOpenGLDynamicRHI::RHICreateUniformBuffer(const void* Conte
 			AllocatedSize = (SizeOfBufferToAllocate > 0) ? SizeOfBufferToAllocate : Layout.ConstantBufferSize;
 
 			// Nothing usable was found in the free pool, or we're not pooling, so create a new uniform buffer
-			if (IsES2Platform(GRHIShaderPlatform))
+			if (GUseEmulatedUniformBuffers)
 			{
 				EmulatedUniformDataRef = UniformBufferDataFactory.Create(AllocatedSize, AllocatedResource);
 			}
@@ -525,7 +524,7 @@ FUniformBufferRHIRef FOpenGLDynamicRHI::RHICreateUniformBuffer(const void* Conte
 			IncrementBufferMemory(GL_UNIFORM_BUFFER, /*bIsStructuredBuffer=*/ false, AllocatedSize);
 		}
 
-		check(!IsES2Platform(GRHIShaderPlatform) || (IsValidRef(EmulatedUniformDataRef) && (EmulatedUniformDataRef->Data.Num() * EmulatedUniformDataRef->Data.GetTypeSize() == AllocatedSize)));
+		check(!GUseEmulatedUniformBuffers || (IsValidRef(EmulatedUniformDataRef) && (EmulatedUniformDataRef->Data.Num() * EmulatedUniformDataRef->Data.GetTypeSize() == AllocatedSize)));
 
 		// Copy the contents of the uniform buffer.
 		if (IsValidRef(EmulatedUniformDataRef))
diff --git a/Engine/Source/Runtime/ShaderCore/Private/ShaderParameters.cpp b/Engine/Source/Runtime/ShaderCore/Private/ShaderParameters.cpp
--- a/Engine/Source/Runtime/ShaderCore/Private/ShaderParameters.cpp
+++ b/Engine/Source/Runtime/ShaderCore/Private/ShaderParameters.cpp
@@ -231,7 +231,7 @@ static void CreateHLSLUniformBufferStructMembersDeclaration(FUniformBufferDecl&
 				{
 					if(bExplicitPadding)
 					{
-						Decl.ConstantBufferMembers += FString::Printf(TEXT("\tfloat1 _%s_PrePadding%u;\r\n"), *NamePrefix, HLSLBaseOffset);
+						Decl.ConstantBufferMembers += FString::Printf(TEXT("\tfloat1 _%sPrePadding%u;\r\n"), *NamePrefix, HLSLBaseOffset);
 					}
 					HLSLBaseOffset += 4;
 				};
