Amortise the cost of deleting OpenGL buffers on OS X rather than trying to reuse the names - doing so causes a crash opening the ShooterGame project in the Editor when using MTGL.

diff --git a/Engine/Source/Runtime/OpenGLDrv/Private/Mac/OpenGLMac.cpp b/Engine/Source/Runtime/OpenGLDrv/Private/Mac/OpenGLMac.cpp
--- a/Engine/Source/Runtime/OpenGLDrv/Private/Mac/OpenGLMac.cpp
+++ b/Engine/Source/Runtime/OpenGLDrv/Private/Mac/OpenGLMac.cpp
@@ -52,9 +52,18 @@ static FAutoConsoleVariableRef CVarMacMaxTexturesToDeletePerFrame(
 	ECVF_RenderThreadSafe
 	);
 
+static int32 GMacMinBuffersToDelete = OPENGL_NAME_CACHE_SIZE;
+static FAutoConsoleVariableRef CVarMacMinBuffersToDelete(
+	TEXT("r.Mac.MinBuffersToDelete"),
+	GMacMinBuffersToDelete,
+	TEXT("Specifies how many buffers are required to be waiting before a call to glDeleteBuffers is made in order to amortize the cost. (Default: 1024)"),
+	ECVF_RenderThreadSafe
+	);
+
 static const uint32 GMacTexturePoolNum = 3;
 static TArray<GLuint> GMacTexturesToDelete[GMacTexturePoolNum];
 static TArray<GLuint> GMacBuffers;
+static TArray<GLuint> GMacBuffersToDelete;
 
 bool GIsRunningOnIntelCard = false; // @todo: remove once Apple fixes radr://16223045 Changes to the GL separate blend state aren't always respected on Intel cards
 static bool GIsEmulatingTimestamp = false; // @todo: Now crashing on Nvidia cards, but not on AMD...
@@ -498,6 +507,9 @@ struct FPlatformOpenGLDevice
 			glDeleteBuffers(GMacBuffers.Num(), GMacBuffers.GetData());
 			GMacBuffers.Reset();
 			
+			glDeleteBuffers(GMacBuffersToDelete.Num(), GMacBuffersToDelete.GetData());
+			GMacBuffersToDelete.Reset();
+			
 			if(GIsEmulatingTimestamp)
 			{
 				SharedContext.RunningQueries.Empty();
@@ -1911,13 +1923,26 @@ void FMacOpenGL::GenBuffers( GLsizei n, GLuint *buffers)
 
 void FMacOpenGL::DeleteBuffers(GLsizei Number, const GLuint* Buffers)
 {
-	for(uint32 i = 0; i < Number; i++)
+	GMacBuffersToDelete.Append(Buffers, Number);
+	if(GMacBuffersToDelete.Num() >= GMacMinBuffersToDelete)
 	{
-		glBindBuffer(GL_COPY_WRITE_BUFFER, Buffers[i]);
-		glBufferData(GL_COPY_WRITE_BUFFER, 0, nullptr, GL_STATIC_DRAW);
-		GMacBuffers.Add(Buffers[i]);
+		glDeleteBuffers(GMacBuffersToDelete.Num(), GMacBuffersToDelete.GetData());
+		GMacBuffersToDelete.Empty();
+	}
+	else
+	{
+		for(uint32 i = 0; i < Number; i++)
+		{
+			glBindBuffer(GL_COPY_WRITE_BUFFER, Buffers[i]);
+#if UE_BUILD_DEBUG
+			GLint bLocked = GL_FALSE;
+			glGetBufferParameteriv(GL_COPY_WRITE_BUFFER, GL_BUFFER_MAPPED, &bLocked);
+			check(!bLocked);
+#endif
+			glBufferData(GL_COPY_WRITE_BUFFER, 0, nullptr, GL_STATIC_DRAW);
+		}
+		glBindBuffer(GL_COPY_WRITE_BUFFER, 0);
 	}
-	glBindBuffer(GL_COPY_WRITE_BUFFER, 0);
 }
 
 void FMacOpenGL::DeleteTextures(GLsizei Number, const GLuint* Textures)
