UE4 - Integration (GL)

diff --git a/Engine/Source/Runtime/OpenGLDrv/Private/OpenGLDrv.cpp b/Engine/Source/Runtime/OpenGLDrv/Private/OpenGLDrv.cpp
--- a/Engine/Source/Runtime/OpenGLDrv/Private/OpenGLDrv.cpp
+++ b/Engine/Source/Runtime/OpenGLDrv/Private/OpenGLDrv.cpp
@@ -376,7 +376,7 @@ public:
 
 	static bool ShouldCache(EShaderPlatform Platform)
 	{
-		return IsFeatureLevelSupported(Platform, ERHIFeatureLevel::SM3);
+		return IsFeatureLevelSupported(Platform, ERHIFeatureLevel::SM4);
 	}
 };
 
diff --git a/Engine/Source/Runtime/OpenGLDrv/Public/OpenGL.h b/Engine/Source/Runtime/OpenGLDrv/Public/OpenGL.h
--- a/Engine/Source/Runtime/OpenGLDrv/Public/OpenGL.h
+++ b/Engine/Source/Runtime/OpenGLDrv/Public/OpenGL.h
@@ -171,6 +171,8 @@ public:
 	static FORCEINLINE bool IsDebugContent()					{ return false; }
 	static FORCEINLINE void InitDebugContext()					{ }
 
+	static FORCEINLINE int32 GetReadHalfFloatPixelsEnum() UGL_REQUIRED(0)
+
 
 	// Silently ignored if not implemented:
 	static FORCEINLINE void QueryTimestampCounter(GLuint QueryID) UGL_OPTIONAL_VOID
diff --git a/Engine/Source/Runtime/OpenGLDrv/Public/OpenGL3.h b/Engine/Source/Runtime/OpenGLDrv/Public/OpenGL3.h
--- a/Engine/Source/Runtime/OpenGLDrv/Public/OpenGL3.h
+++ b/Engine/Source/Runtime/OpenGLDrv/Public/OpenGL3.h
@@ -574,6 +574,8 @@ struct FOpenGL3 : public FOpenGLBase
 
 	static void ProcessQueryGLInt();
 	static void ProcessExtensions(const FString& ExtensionsString);
+	
+	static FORCEINLINE int32 GetReadHalfFloatPixelsEnum() { return GL_HALF_FLOAT; }
 
 protected:
 	static GLsizei NextTextureName;
diff --git a/Engine/Source/Runtime/OpenGLDrv/Public/OpenGL4.h b/Engine/Source/Runtime/OpenGLDrv/Public/OpenGL4.h
--- a/Engine/Source/Runtime/OpenGLDrv/Public/OpenGL4.h
+++ b/Engine/Source/Runtime/OpenGLDrv/Public/OpenGL4.h
@@ -99,6 +99,8 @@ struct FOpenGL4 : public FOpenGL3
 
 	static uint64 GetVideoMemorySize();
 
+	static FORCEINLINE int32 GetReadHalfFloatPixelsEnum() { return GL_HALF_FLOAT; }
+
 protected:
 	static GLint MaxComputeTextureImageUnits;
 	static GLint MaxComputeUniformComponents;
diff --git a/Engine/Source/Runtime/OpenGLDrv/Public/OpenGLES2.h b/Engine/Source/Runtime/OpenGLDrv/Public/OpenGLES2.h
--- a/Engine/Source/Runtime/OpenGLDrv/Public/OpenGLES2.h
+++ b/Engine/Source/Runtime/OpenGLDrv/Public/OpenGLES2.h
@@ -44,6 +44,8 @@ typedef GLfloat GLdouble;
 #undef UGL_TIME_ELAPSED
 #define UGL_TIME_ELAPSED		GL_TIME_ELAPSED_EXT
 
+#define GL_HALF_FLOAT_OES 0x8D61
+
 
 
 struct FOpenGLES2 : public FOpenGLBase
@@ -64,7 +66,7 @@ struct FOpenGLES2 : public FOpenGLBase
 	static bool SupportsDisjointTimeQueries();
 	static FORCEINLINE bool SupportsBlitFramebuffer()					{ return bSupportsNVFrameBufferBlit; }
 	static FORCEINLINE bool SupportsDepthStencilRead()					{ return false; }
-	static FORCEINLINE bool SupportsFloatReadSurface()					{ return false; }
+	static FORCEINLINE bool SupportsFloatReadSurface()					{ return SupportsColorBufferHalfFloat(); }
 	static FORCEINLINE bool SupportsMultipleRenderTargets()				{ return false; }
 	static FORCEINLINE bool SupportsMultisampledTextures()				{ return false; }
 	static FORCEINLINE bool SupportsFences()							{ return false; }
@@ -109,6 +111,8 @@ struct FOpenGLES2 : public FOpenGLBase
 	static FORCEINLINE bool RequiresDontEmitPrecisionForTextureSamplers() { return bRequiresDontEmitPrecisionForTextureSamplers; }
 	static FORCEINLINE bool RequiresTextureCubeLodEXTToTextureCubeLodDefine() { return bRequiresTextureCubeLodEXTToTextureCubeLodDefine; }
 
+	static FORCEINLINE int32 GetReadHalfFloatPixelsEnum()				{ return GL_HALF_FLOAT_OES; }
+
 
 	// On iOS both glMapBufferOES() and glBufferSubData() for immediate vertex and index data
 	// is the slow path (they both hit GPU sync and data cache flush in driver according to profiling in driver symbols).
