Change the buffer deletion routine on Mac OS X again in an effort to eliminate the crashes that QA were seeing.

diff --git a/Engine/Source/Runtime/OpenGLDrv/Private/Mac/OpenGLMac.cpp b/Engine/Source/Runtime/OpenGLDrv/Private/Mac/OpenGLMac.cpp
--- a/Engine/Source/Runtime/OpenGLDrv/Private/Mac/OpenGLMac.cpp
+++ b/Engine/Source/Runtime/OpenGLDrv/Private/Mac/OpenGLMac.cpp
@@ -1923,7 +1923,29 @@ void FMacOpenGL::GenBuffers( GLsizei n, GLuint *buffers)
 
 void FMacOpenGL::DeleteBuffers(GLsizei Number, const GLuint* Buffers)
 {
-	GMacBuffersToDelete.Append(Buffers, Number);
+#if UE_BUILD_DEBUG
+	GLint Buffer = 0;
+	glGetIntegerv(GL_COPY_WRITE_BUFFER, &Buffer);
+#endif
+	
+	uint32 Index = 0;
+	
+#if USE_OPENGL_NAME_CACHE
+	while ( GMacBuffers.Num() < USE_OPENGL_NAME_CACHE && Index > Number )
+	{
+#if UE_BUILD_DEBUG
+		glBindBuffer(GL_COPY_WRITE_BUFFER, Buffers[Index]);
+		GLint bLocked = GL_FALSE;
+		glGetBufferParameteriv(GL_COPY_WRITE_BUFFER, GL_BUFFER_MAPPED, &bLocked);
+		check(!bLocked);
+#endif
+		
+		GMacBuffers.Add(Buffers[Index]);
+		Index++;
+	}
+#endif
+	
+	GMacBuffersToDelete.Append(&Buffers[Index], Number-Index);
 	if(GMacBuffersToDelete.Num() >= GMacMinBuffersToDelete)
 	{
 		glDeleteBuffers(GMacBuffersToDelete.Num(), GMacBuffersToDelete.GetData());
@@ -1931,18 +1953,20 @@ void FMacOpenGL::DeleteBuffers(GLsizei Number, const GLuint* Buffers)
 	}
 	else
 	{
-		for(uint32 i = 0; i < Number; i++)
+		for(uint32 i = Index; i < Number; i++)
 		{
-			glBindBuffer(GL_COPY_WRITE_BUFFER, Buffers[i]);
 #if UE_BUILD_DEBUG
+			glBindBuffer(GL_COPY_WRITE_BUFFER, Buffers[i]);
 			GLint bLocked = GL_FALSE;
 			glGetBufferParameteriv(GL_COPY_WRITE_BUFFER, GL_BUFFER_MAPPED, &bLocked);
 			check(!bLocked);
 #endif
-			glBufferData(GL_COPY_WRITE_BUFFER, 0, nullptr, GL_STATIC_DRAW);
 		}
-		glBindBuffer(GL_COPY_WRITE_BUFFER, 0);
 	}
+	
+#if UE_BUILD_DEBUG
+	glBindBuffer(GL_COPY_WRITE_BUFFER, 0);
+#endif
 }
 
 void FMacOpenGL::DeleteTextures(GLsizei Number, const GLuint* Textures)
