Simplfy glsl code fixups in #defines for android/html5. provide a different #def'd code block for html5.

#ttp 345305

diff --git a/Engine/Source/Runtime/OpenGLDrv/Private/OpenGLShaders.cpp b/Engine/Source/Runtime/OpenGLDrv/Private/OpenGLShaders.cpp
--- a/Engine/Source/Runtime/OpenGLDrv/Private/OpenGLShaders.cpp
+++ b/Engine/Source/Runtime/OpenGLDrv/Private/OpenGLShaders.cpp
@@ -340,6 +340,9 @@ ShaderType* CompileOpenGLShader(const TArray<uint8>& Code)
 				VersionCodePointer[index] = ' ';
 			}
 		}
+#endif
+
+#if PLATFORM_ANDROID 
 
 		const ANSICHAR* ExtensionString = "";
 		
@@ -385,7 +388,6 @@ ShaderType* CompileOpenGLShader(const TArray<uint8>& Code)
 				ReplaceShaderSubstring(const_cast<ANSICHAR*>(GlslCode), "varying", "in");
 			}
 		}
-#if PLATFORM_ANDROID
 		else if ( (TypeEnum == GL_FRAGMENT_SHADER) &&
 			FOpenGL::RequiresDontEmitPrecisionForTextureSamplers() )
 		{
@@ -401,18 +403,11 @@ ShaderType* CompileOpenGLShader(const TArray<uint8>& Code)
 		{
 			Prologue = "#define textureCubeLodEXT textureCubeLod \n";
 		}
-#endif
 		else if(!FOpenGL::SupportsShaderTextureLod() || !FOpenGL::SupportsShaderTextureCubeLod())
 		{
 			Prologue = 	"#define texture2DLodEXT(a, b, c) texture2D(a, b) \n"
 						"#define textureCubeLodEXT(a, b, c) textureCube(a, b) \n";
 		}
-#if !PLATFORM_ANDROID
-		else if(!FOpenGL::SupportsTextureCubeLodEXT())
-		{
-			Prologue = "#define textureCubeLodEXT textureCubeLod \n";
-		}
-#endif
 		else 
 		{
 			Prologue = "";
@@ -425,10 +420,33 @@ ShaderType* CompileOpenGLShader(const TArray<uint8>& Code)
 		const GLint ShaderSourceLen[4] = { (GLint)(strlen(VersionString)), (GLint)(strlen(ExtensionString)), (GLint)(strlen(Prologue)), (GLint)(strlen(GlslCode)) };
 		glShaderSource(Resource, 4, (const GLchar**)&ShaderSourceStrings, ShaderSourceLen);		
 		glCompileShader(Resource);
-#else // PLATFORM_ANDROID || PLATFORM_HTML5
+#endif 
+
+
+	// HTML5 use case is much simpler, use a separate chunk of code from android. 
+#if PLATFORM_HTML5
+
+		ANSICHAR Prologue[1024] = ""; 
+
+		if (!FOpenGL::SupportsShaderTextureLod()) 
+		{
+			strcat(Prologue, "#define  DONTEMITEXTENSIONSHADERTEXTURELODENABLE \n");
+			strcat(Prologue, "#define texture2DLodEXT(a, b, c) texture2D(a, b) \n");
+			strcat(Prologue, "#define textureCubeLodEXT(a, b, c) textureCube(a, b)\n");
+		}
+
+		const GLchar* ShaderSourceStrings[3] = { VersionString, Prologue, GlslCode };
+		const GLint ShaderSourceLen[3] = { (GLint)(strlen(VersionString)), (GLint)(strlen(Prologue)), (GLint)(strlen(GlslCode)) };
+
+		glShaderSource(Resource, 3, (const GLchar**)&ShaderSourceStrings, ShaderSourceLen);		
+		glCompileShader(Resource);
+
+#endif 
+
+#if !PLATFORM_HTML5 && !PLATFORM_ANDROID
 		glShaderSource(Resource, 1, (const GLchar**)&GlslCode, &GlslCodeLength);
 		glCompileShader(Resource);
-#endif // PLATFORM_ANDROID || PLATFORM_HTML5
+#endif // !PLATFORM_ANDROID && !PLATFORM_HTML5
 		GetOpenGLCompiledShaderCache().Add(Key,Resource);
 	}
 
