UE4 - Remove GL error on iOS

diff --git a/Engine/Source/Runtime/OpenGLDrv/Private/IOS/IOSOpenGL.h b/Engine/Source/Runtime/OpenGLDrv/Private/IOS/IOSOpenGL.h
--- a/Engine/Source/Runtime/OpenGLDrv/Private/IOS/IOSOpenGL.h
+++ b/Engine/Source/Runtime/OpenGLDrv/Private/IOS/IOSOpenGL.h
@@ -105,6 +105,8 @@ struct FIOSOpenGL : public FOpenGLES2
 	{
 		return SP_OPENGL_ES2_IOS;
 	}
+
+	static FORCEINLINE bool SupportsFramebufferSRGBEnable()				{ return false; }
 };
 
 typedef FIOSOpenGL FOpenGL;
diff --git a/Engine/Source/Runtime/OpenGLDrv/Private/OpenGLDrv.cpp b/Engine/Source/Runtime/OpenGLDrv/Private/OpenGLDrv.cpp
--- a/Engine/Source/Runtime/OpenGLDrv/Private/OpenGLDrv.cpp
+++ b/Engine/Source/Runtime/OpenGLDrv/Private/OpenGLDrv.cpp
@@ -517,6 +517,10 @@ void InitDefaultGLContextState(void)
 	// Intel HD4000 under <= 10.8.4 requires GL_DITHER disabled or dithering will occur on any channel < 8bits.
 	// No other driver does this but we don't need GL_DITHER on anyway.
 	glDisable(GL_DITHER);
-	// Render targets with TexCreate_SRGB should do sRGB conversion like in D3D11
-	glEnable(GL_FRAMEBUFFER_SRGB);
+
+	if (FOpenGL::SupportsFramebufferSRGBEnable())
+	{
+		// Render targets with TexCreate_SRGB should do sRGB conversion like in D3D11
+		glEnable(GL_FRAMEBUFFER_SRGB);
+	}
 }
diff --git a/Engine/Source/Runtime/OpenGLDrv/Public/OpenGL.h b/Engine/Source/Runtime/OpenGLDrv/Public/OpenGL.h
--- a/Engine/Source/Runtime/OpenGLDrv/Public/OpenGL.h
+++ b/Engine/Source/Runtime/OpenGLDrv/Public/OpenGL.h
@@ -114,6 +114,7 @@ public:
 	static FORCEINLINE bool SupportsASTC()								{ return bSupportsASTC; }
 	static FORCEINLINE bool SupportsETC1()								{ return false; }
 	static FORCEINLINE bool SupportsETC2()								{ return false; }
+	static FORCEINLINE bool SupportsFramebufferSRGBEnable()				{ return true; }
 	static FORCEINLINE bool SupportsCombinedDepthStencilAttachment()	{ return true; }
 	static FORCEINLINE bool SupportsFastBufferData()					{ return true; }
 	static FORCEINLINE bool SupportsCopyImage()							{ return bSupportsCopyImage; }
