Fix an error in timestamp emulation that would result in incorrect values. Perform the resize & blit operations on a dedicated thread to avoid a glFlush call on Nvidia cards to improve performance. This is unlikely a real solution to the problem, which appears to be in the OpenGL stack rather than our code.

diff --git a/Engine/Source/Runtime/OpenGLDrv/Private/Mac/OpenGLMac.cpp b/Engine/Source/Runtime/OpenGLDrv/Private/Mac/OpenGLMac.cpp
--- a/Engine/Source/Runtime/OpenGLDrv/Private/Mac/OpenGLMac.cpp
+++ b/Engine/Source/Runtime/OpenGLDrv/Private/Mac/OpenGLMac.cpp
@@ -5,6 +5,7 @@
 #include "MacWindow.h"
 #include "MacTextInputMethodSystem.h"
 #include "CocoaTextView.h"
+#include "CocoaThread.h"
 #include "CocoaWindow.h"
 
 /*------------------------------------------------------------------------------
@@ -351,8 +352,26 @@ struct OpenGLContextInfo
 
 extern void OnQueryInvalidation( void );
 
+// Provides a dedicated thread to perform viewport resize & blit operations.
+@interface FCocoaOpenGLThread : NSThread
+@property (readonly) NSRunLoop* RunLoop;
+@end
+
+@implementation FCocoaOpenGLThread
+@synthesize RunLoop;
+- (void)main
+{
+	RunLoop = [NSRunLoop currentRunLoop];
+	while (!GIsRequestingExit)
+	{
+		[[NSRunLoop currentRunLoop] runMode:NSDefaultRunLoopMode beforeDate:[NSDate distantFuture]];
+	}
+}
+@end
+
 struct FPlatformOpenGLDevice
 {
+	FCocoaOpenGLThread*	BlitThread;
 	FPlatformOpenGLContext	SharedContext;
 	GLuint				SharedContextCompositeVertexShader;
 	GLuint				SharedContextCompositeFragmentShader;
@@ -376,6 +395,9 @@ struct FPlatformOpenGLDevice
 		: NumUsedContexts(0)
 	{
 		SCOPED_AUTORELEASE_POOL;
+		
+		BlitThread = [FCocoaOpenGLThread new];
+		[BlitThread start];
 
 		extern void InitDebugContext();
 		ContextUsageGuard = new FCriticalSection;
@@ -525,6 +547,8 @@ struct FPlatformOpenGLDevice
 		[RenderingContext.OpenGLContext release];
 		[SharedContext.OpenGLContext release];
 		
+		[BlitThread release];
+		
 		delete ContextUsageGuard;
 	}
 };
@@ -697,205 +721,214 @@ void* PlatformGetWindow(FPlatformOpenGLContext* Context, void** AddParam)
 
 bool PlatformBlitToViewport( FPlatformOpenGLDevice* Device, const FOpenGLViewport& Viewport, uint32 BackbufferSizeX, uint32 BackbufferSizeY, bool bPresent,bool bLockToVsync, int32 SyncInterval )
 {
-	FPlatformOpenGLContext* const Context = Viewport.GetGLContext();
-
-	check(Context && Context->OpenGLView);
-
+	// Don't flush on Nvidia - this causes a tremendous slowdown & doesn't appear to be necessary to avoid rendering errors
+	if ( !IsRHIDeviceNVIDIA() )
 	{
-		FScopeLock ScopeLock(Device->ContextUsageGuard);
+		// Other GPU vendors do need to flush, or the blit to screen may occur before rendering to FBO has completed.
+		glFlushRenderAPPLE();
+	}
+	// Use a separate dedicated thread to blit to the screen to avoid expense of swapping contexts.
+	PerformBlockOnRunLoop(^{
+		FPlatformOpenGLContext* const Context = Viewport.GetGLContext();
+		check(Context && Context->OpenGLView);
 		{
-			FScopeContext ScopeContext(Context->OpenGLContext);
-			
-			if(Context->OpenGLView.bNeedsUpdate)
+			FScopeLock ScopeLock(Device->ContextUsageGuard);
 			{
-				Context->OpenGLView.bNeedsUpdate = false;
-				[Context->OpenGLContext update];
-			}
-			
-            if (Viewport.GetCustomPresent())
-            {
-				SCOPED_AUTORELEASE_POOL;
-                
-				// Clear the Alpha channel
-				glColorMask(GL_FALSE,GL_FALSE,GL_FALSE,GL_TRUE);
-				glClearColor(0.f,0.f,0.f,1.f);
-				glClear(GL_COLOR_BUFFER_BIT);
-				glColorMask(GL_TRUE,GL_TRUE,GL_TRUE,GL_TRUE);
-				glClearColor(0.f,0.f,0.f,0.f);
-
-                glDisable(GL_FRAMEBUFFER_SRGB);
-                Viewport.GetCustomPresent()->Present(SyncInterval);
-                glEnable(GL_FRAMEBUFFER_SRGB);
-                return false;
-            }
-			// OpenGL state necessary for blit is set up in PlatformResizeGLContext(), and should be correct here,
-			// as viewport contexts aren't bound at any other occasion.
-			glBindFramebuffer(GL_DRAW_FRAMEBUFFER, 0);
-			glDrawBuffer(GL_BACK);
-			glBindFramebuffer(GL_READ_FRAMEBUFFER, Context->ViewportFramebuffer);
-			glReadBuffer(GL_COLOR_ATTACHMENT0);
-			glDisable(GL_FRAMEBUFFER_SRGB);
-
-            glBlitFramebuffer(
-				0, 0, BackbufferSizeX, BackbufferSizeY,
-				0, BackbufferSizeY, BackbufferSizeX, 0,
-				GL_COLOR_BUFFER_BIT,
-				GL_NEAREST
-            );
-
-			if (bPresent)
-			{
-				SCOPED_AUTORELEASE_POOL;
-
-				// Clear the Alpha channel
-				glColorMask(GL_FALSE,GL_FALSE,GL_FALSE,GL_TRUE);
-				glClearColor(0.f,0.f,0.f,1.f);
-				glClear(GL_COLOR_BUFFER_BIT);
-				glColorMask(GL_TRUE,GL_TRUE,GL_TRUE,GL_TRUE);
-				glClearColor(0.f,0.f,0.f,0.f);
-				
-				bool bRoundedBlit = false;
+				FScopeContext ScopeContext(Context->OpenGLContext);
 				
-				NSWindow* Window  = [Context->OpenGLView window];
-				NSView* SuperView = [[Window contentView] superview];
-				FCocoaWindow* CocoaWindow = [Window isKindOfClass:[FCocoaWindow class]] ? (FCocoaWindow*)Window : nil;
-				if([SuperView respondsToSelector:@selector(roundedCornerRadius)] && CocoaWindow)
+				if(Context->OpenGLView.bNeedsUpdate)
 				{
-					bool bRoundedCorners = [CocoaWindow roundedCorners];
-					bool bFullWindowRendering = ([CocoaWindow styleMask] & (NSTexturedBackgroundWindowMask));
-					EWindowMode::Type WindowMode = [CocoaWindow windowMode];
-					bRoundedBlit = (WindowMode == EWindowMode::Windowed && bRoundedCorners && bFullWindowRendering);
+					Context->OpenGLView.bNeedsUpdate = false;
+					[Context->OpenGLContext update];
 				}
 				
-				if(bRoundedBlit)
+				if (Viewport.GetCustomPresent())
 				{
-					bool bGenerateTexture = !Device->SharedContextCompositeTextureSizeX && !Device->SharedContextCompositeTextureSizeY;
-					glBindSampler(0, 0);
-					glActiveTexture(GL_TEXTURE0);
-					glBindTexture (GL_TEXTURE_2D, Device->SharedContextCompositeTexture);
-					glTexParameteri(GL_TEXTURE_2D, GL_TEXTURE_BASE_LEVEL, 0);
-					glTexParameteri(GL_TEXTURE_2D, GL_TEXTURE_MAX_LEVEL, 0);
-					glTexParameteri(GL_TEXTURE_2D, GL_TEXTURE_MIN_FILTER, GL_NEAREST);
-					glTexParameteri(GL_TEXTURE_2D, GL_TEXTURE_MAG_FILTER, GL_NEAREST);
-					glTexParameteri(GL_TEXTURE_2D, GL_TEXTURE_WRAP_S, GL_CLAMP_TO_EDGE);
-					glTexParameteri(GL_TEXTURE_2D, GL_TEXTURE_WRAP_T, GL_CLAMP_TO_EDGE);
+					SCOPED_AUTORELEASE_POOL;
 					
-					if(bGenerateTexture)
+					// Clear the Alpha channel
+					glColorMask(GL_FALSE,GL_FALSE,GL_FALSE,GL_TRUE);
+					glClearColor(0.f,0.f,0.f,1.f);
+					glClear(GL_COLOR_BUFFER_BIT);
+					glColorMask(GL_TRUE,GL_TRUE,GL_TRUE,GL_TRUE);
+					glClearColor(0.f,0.f,0.f,0.f);
+
+					glDisable(GL_FRAMEBUFFER_SRGB);
+					Viewport.GetCustomPresent()->Present(SyncInterval);
+					glEnable(GL_FRAMEBUFFER_SRGB);
+				}
+				else
+				{
+					// OpenGL state necessary for blit is set up in PlatformResizeGLContext(), and should be correct here,
+					// as viewport contexts aren't bound at any other occasion.
+					glBindFramebuffer(GL_DRAW_FRAMEBUFFER, 0);
+					glDrawBuffer(GL_BACK);
+					glBindFramebuffer(GL_READ_FRAMEBUFFER, Context->ViewportFramebuffer);
+					glReadBuffer(GL_COLOR_ATTACHMENT0);
+					glDisable(GL_FRAMEBUFFER_SRGB);
+
+					glBlitFramebuffer(
+						0, 0, BackbufferSizeX, BackbufferSizeY,
+						0, BackbufferSizeY, BackbufferSizeX, 0,
+						GL_COLOR_BUFFER_BIT,
+						GL_NEAREST
+					);
+
+					if (bPresent)
 					{
-						uint32 Size = 32;
-						Device->SharedContextCompositeTextureSizeX = Size;
-						Device->SharedContextCompositeTextureSizeY = Size;
-						NSImage* MaskImage = [[[NSImage alloc] initWithSize:NSMakeSize(Size*2, Size*2)] autorelease];
+						SCOPED_AUTORELEASE_POOL;
+
+						// Clear the Alpha channel
+						glColorMask(GL_FALSE,GL_FALSE,GL_FALSE,GL_TRUE);
+						glClearColor(0.f,0.f,0.f,1.f);
+						glClear(GL_COLOR_BUFFER_BIT);
+						glColorMask(GL_TRUE,GL_TRUE,GL_TRUE,GL_TRUE);
+						glClearColor(0.f,0.f,0.f,0.f);
+						
+						bool bRoundedBlit = false;
+						
+						NSWindow* Window  = [Context->OpenGLView window];
+						NSView* SuperView = [[Window contentView] superview];
+						FCocoaWindow* CocoaWindow = [Window isKindOfClass:[FCocoaWindow class]] ? (FCocoaWindow*)Window : nil;
+						if([SuperView respondsToSelector:@selector(roundedCornerRadius)] && CocoaWindow)
 						{
-							[MaskImage lockFocus];
-							{
-								NSGraphicsContext* CurrentContext = [NSGraphicsContext currentContext];
-								[CurrentContext saveGraphicsState];
-								[CurrentContext setShouldAntialias: NO];
-								
-								[[NSColor clearColor] set];
-								[[NSBezierPath bezierPathWithRect:NSMakeRect(0.f, 0.f, Size*2, Size*2)] fill];
-								
-								[[NSColor colorWithDeviceRed:0.0f green:0.0f blue:0.0f alpha:1.0f] set];
-								float Radius = [SuperView roundedCornerRadius] * 1.6f;
-								[[NSBezierPath bezierPathWithRoundedRect:NSMakeRect(0.f, 0.f, Size*2, Size*2) xRadius: Radius yRadius: Radius] fill];
-								
-								[CurrentContext restoreGraphicsState];
-							}
-							[MaskImage unlockFocus];
+							bool bRoundedCorners = [CocoaWindow roundedCorners];
+							bool bFullWindowRendering = ([CocoaWindow styleMask] & (NSTexturedBackgroundWindowMask));
+							EWindowMode::Type WindowMode = [CocoaWindow windowMode];
+							bRoundedBlit = (WindowMode == EWindowMode::Windowed && bRoundedCorners && bFullWindowRendering);
 						}
 						
-						NSRect SrcRect = NSMakeRect(0, Size, Size, Size);
-						NSRect DestRect = NSMakeRect(0, 0, Size, Size);
-						NSImage* CornerImage = [[[NSImage alloc] initWithSize:DestRect.size] autorelease];
+						if(bRoundedBlit)
 						{
-							[CornerImage lockFocus];
+							bool bGenerateTexture = !Device->SharedContextCompositeTextureSizeX && !Device->SharedContextCompositeTextureSizeY;
+							glBindSampler(0, 0);
+							glActiveTexture(GL_TEXTURE0);
+							glBindTexture (GL_TEXTURE_2D, Device->SharedContextCompositeTexture);
+							glTexParameteri(GL_TEXTURE_2D, GL_TEXTURE_BASE_LEVEL, 0);
+							glTexParameteri(GL_TEXTURE_2D, GL_TEXTURE_MAX_LEVEL, 0);
+							glTexParameteri(GL_TEXTURE_2D, GL_TEXTURE_MIN_FILTER, GL_NEAREST);
+							glTexParameteri(GL_TEXTURE_2D, GL_TEXTURE_MAG_FILTER, GL_NEAREST);
+							glTexParameteri(GL_TEXTURE_2D, GL_TEXTURE_WRAP_S, GL_CLAMP_TO_EDGE);
+							glTexParameteri(GL_TEXTURE_2D, GL_TEXTURE_WRAP_T, GL_CLAMP_TO_EDGE);
+							
+							if(bGenerateTexture)
 							{
-								NSGraphicsContext* CurrentContext = [NSGraphicsContext currentContext];
-								[CurrentContext saveGraphicsState];
+								uint32 Size = 32;
+								Device->SharedContextCompositeTextureSizeX = Size;
+								Device->SharedContextCompositeTextureSizeY = Size;
+								NSImage* MaskImage = [[[NSImage alloc] initWithSize:NSMakeSize(Size*2, Size*2)] autorelease];
+								{
+									[MaskImage lockFocus];
+									{
+										NSGraphicsContext* CurrentContext = [NSGraphicsContext currentContext];
+										[CurrentContext saveGraphicsState];
+										[CurrentContext setShouldAntialias: NO];
+										
+										[[NSColor clearColor] set];
+										[[NSBezierPath bezierPathWithRect:NSMakeRect(0.f, 0.f, Size*2, Size*2)] fill];
+										
+										[[NSColor colorWithDeviceRed:0.0f green:0.0f blue:0.0f alpha:1.0f] set];
+										float Radius = [SuperView roundedCornerRadius] * 1.6f;
+										[[NSBezierPath bezierPathWithRoundedRect:NSMakeRect(0.f, 0.f, Size*2, Size*2) xRadius: Radius yRadius: Radius] fill];
+										
+										[CurrentContext restoreGraphicsState];
+									}
+									[MaskImage unlockFocus];
+								}
 								
-								[[NSColor clearColor] set];
-								[[NSBezierPath bezierPathWithRect:DestRect] fill];
+								NSRect SrcRect = NSMakeRect(0, Size, Size, Size);
+								NSRect DestRect = NSMakeRect(0, 0, Size, Size);
+								NSImage* CornerImage = [[[NSImage alloc] initWithSize:DestRect.size] autorelease];
+								{
+									[CornerImage lockFocus];
+									{
+										NSGraphicsContext* CurrentContext = [NSGraphicsContext currentContext];
+										[CurrentContext saveGraphicsState];
+										
+										[[NSColor clearColor] set];
+										[[NSBezierPath bezierPathWithRect:DestRect] fill];
+										
+										[MaskImage drawInRect:DestRect fromRect:SrcRect operation:NSCompositeSourceOver fraction:1 respectFlipped:YES hints:nil];
+										
+										[CurrentContext restoreGraphicsState];
+									}
+									[CornerImage unlockFocus];
+								}
 								
-								[MaskImage drawInRect:DestRect fromRect:SrcRect operation:NSCompositeSourceOver fraction:1 respectFlipped:YES hints:nil];
+								CGImageRef CGImage = [CornerImage CGImageForProposedRect:nil context:nil hints:nil];
+								check(CGImage);
+								NSBitmapImageRep* ImageRep = [[[NSBitmapImageRep alloc] initWithCGImage:CGImage] autorelease];
+								check(ImageRep);
 								
-								[CurrentContext restoreGraphicsState];
+								GLenum format = [ImageRep hasAlpha] ? GL_RGBA : GL_RGB;
+								glTexImage2D (GL_TEXTURE_2D, 0, format, [ImageRep size].width, [ImageRep size].height, 0, format, GL_UNSIGNED_BYTE, [ImageRep bitmapData]);
 							}
-							[CornerImage unlockFocus];
+							
+							glUseProgram(Device->SharedContextCompositeProgram);
+							
+							glUniform1i(Device->SharedContextWindowTextureUniform, 0);
+							
+							glUniform1i(Device->SharedContextTextureDirectionUniform, 0);
+							
+							glColorMask(GL_FALSE,GL_FALSE,GL_FALSE,GL_TRUE);
+							
+							glViewport(0, BackbufferSizeY-Device->SharedContextCompositeTextureSizeY, Device->SharedContextCompositeTextureSizeX, Device->SharedContextCompositeTextureSizeY);
+							
+							glDrawArraysInstanced(GL_TRIANGLES, 0, 6, 1);
+							
+							glUniform1i(Device->SharedContextTextureDirectionUniform, 1);
+							
+							glViewport(BackbufferSizeX-Device->SharedContextCompositeTextureSizeX, BackbufferSizeY-Device->SharedContextCompositeTextureSizeY, Device->SharedContextCompositeTextureSizeX, Device->SharedContextCompositeTextureSizeY);
+							
+							glDrawArraysInstanced(GL_TRIANGLES, 0, 6, 1);
+							
+							glViewport(0, 0, BackbufferSizeX, BackbufferSizeY);
+							
+							glColorMask(GL_TRUE,GL_TRUE,GL_TRUE,GL_TRUE);
+						}
+
+						int32 RealSyncInterval = bLockToVsync ? SyncInterval : 0;
+
+						if (Context->SyncInterval != RealSyncInterval)
+						{
+							[Context->OpenGLContext setValues: &RealSyncInterval forParameter: NSOpenGLCPSwapInterval];
+							Context->SyncInterval = RealSyncInterval;
 						}
 						
-						CGImageRef CGImage = [CornerImage CGImageForProposedRect:nil context:nil hints:nil];
-						check(CGImage);
-						NSBitmapImageRep* ImageRep = [[[NSBitmapImageRep alloc] initWithCGImage:CGImage] autorelease];
-						check(ImageRep);
+						if(!GMacEnableCocoaScreenUpdates)
+						{
+							GMacEnableCocoaScreenUpdates = true;
+							NSEnableScreenUpdates();
+						}
 						
-						GLenum format = [ImageRep hasAlpha] ? GL_RGBA : GL_RGB;
-						glTexImage2D (GL_TEXTURE_2D, 0, format, [ImageRep size].width, [ImageRep size].height, 0, format, GL_UNSIGNED_BYTE, [ImageRep bitmapData]);
+						[Context->OpenGLContext flushBuffer];
+						
+						TArray<GLuint>& TexturesToDelete = GMacTexturesToDelete[(GFrameNumberRenderThread - (GMacTexturePoolNum - 1)) % GMacTexturePoolNum];
+						if(TexturesToDelete.Num() && TexturesToDelete.Num() > GMacMinTexturesToDeletePerFrame)
+						{
+							uint32 Num = FMath::Min(TexturesToDelete.Num(), GMacMaxTexturesToDeletePerFrame);
+							glDeleteTextures(Num, TexturesToDelete.GetData());
+							TexturesToDelete.RemoveAt(0, Num, false);
+						}
+						
+						glEnable(GL_FRAMEBUFFER_SRGB);
+						
+						REPORT_GL_END_BUFFER_EVENT_FOR_FRAME_DUMP();
+		//				INITIATE_GL_FRAME_DUMP_EVERY_X_CALLS( 1000 );
+						
+						if(CocoaWindow)
+						{
+							// Using dispatch was safer - but during loading we aren't responsive on the main thread
+							// this makes it impossible to open the window & so it never appears to play the loading screen.
+							// Just hope that this will work OK...
+							[CocoaWindow performDeferredOrderFront];
+						}
 					}
-					
-					glUseProgram(Device->SharedContextCompositeProgram);
-					
-					glUniform1i(Device->SharedContextWindowTextureUniform, 0);
-					
-					glUniform1i(Device->SharedContextTextureDirectionUniform, 0);
-					
-					glColorMask(GL_FALSE,GL_FALSE,GL_FALSE,GL_TRUE);
-					
-					glViewport(0, BackbufferSizeY-Device->SharedContextCompositeTextureSizeY, Device->SharedContextCompositeTextureSizeX, Device->SharedContextCompositeTextureSizeY);
-					
-					glDrawArraysInstanced(GL_TRIANGLES, 0, 6, 1);
-					
-					glUniform1i(Device->SharedContextTextureDirectionUniform, 1);
-					
-					glViewport(BackbufferSizeX-Device->SharedContextCompositeTextureSizeX, BackbufferSizeY-Device->SharedContextCompositeTextureSizeY, Device->SharedContextCompositeTextureSizeX, Device->SharedContextCompositeTextureSizeY);
-					
-					glDrawArraysInstanced(GL_TRIANGLES, 0, 6, 1);
-					
-					glViewport(0, 0, BackbufferSizeX, BackbufferSizeY);
-					
-					glColorMask(GL_TRUE,GL_TRUE,GL_TRUE,GL_TRUE);
-				}
-
-				int32 RealSyncInterval = bLockToVsync ? SyncInterval : 0;
-
-				if (Context->SyncInterval != RealSyncInterval)
-				{
-					[Context->OpenGLContext setValues: &RealSyncInterval forParameter: NSOpenGLCPSwapInterval];
-					Context->SyncInterval = RealSyncInterval;
-				}
-				
-				if(!GMacEnableCocoaScreenUpdates)
-				{
-					GMacEnableCocoaScreenUpdates = true;
-					NSEnableScreenUpdates();
-				}
-				
-				[Context->OpenGLContext flushBuffer];
-				
-				TArray<GLuint>& TexturesToDelete = GMacTexturesToDelete[(GFrameNumberRenderThread - (GMacTexturePoolNum - 1)) % GMacTexturePoolNum];
-				if(TexturesToDelete.Num() && TexturesToDelete.Num() > GMacMinTexturesToDeletePerFrame)
-				{
-					uint32 Num = FMath::Min(TexturesToDelete.Num(), GMacMaxTexturesToDeletePerFrame);
-					glDeleteTextures(Num, TexturesToDelete.GetData());
-					TexturesToDelete.RemoveAt(0, Num, false);
-				}
-				
-				glEnable(GL_FRAMEBUFFER_SRGB);
-				
-				REPORT_GL_END_BUFFER_EVENT_FOR_FRAME_DUMP();
-//				INITIATE_GL_FRAME_DUMP_EVERY_X_CALLS( 1000 );
-				
-				if(CocoaWindow)
-				{
-					// Using dispatch was safer - but during loading we aren't responsive on the main thread
-					// this makes it impossible to open the window & so it never appears to play the loading screen.
-					// Just hope that this will work OK...
-					[CocoaWindow performDeferredOrderFront];
 				}
 			}
 		}
-	}
-	return true;
+	}, Device->BlitThread.RunLoop, NSDefaultRunLoopMode, NSDefaultRunLoopMode, true);
+	return !(Viewport.GetCustomPresent());
 }
 
 void PlatformRenderingContextSetup(FPlatformOpenGLDevice* Device)
@@ -969,62 +1002,71 @@ void PlatformRebindResources(FPlatformOpenGLDevice* Device)
 
 void PlatformResizeGLContext( FPlatformOpenGLDevice* Device, FPlatformOpenGLContext* Context, uint32 SizeX, uint32 SizeY, bool bFullscreen, bool bWasFullscreen, GLenum BackBufferTarget, GLuint BackBufferResource)
 {
-	FScopeLock ScopeLock(Device->ContextUsageGuard);
+	// Don't flush on Nvidia - this causes a tremendous slowdown & doesn't appear to be necessary to avoid rendering errors
+	if ( !IsRHIDeviceNVIDIA() )
 	{
-		FScopeContext ScopeContext(Context->OpenGLContext);
+		// Other GPU vendors do need to flush, or the blit to screen may occur before rendering to FBO has completed.
+		glFlushRenderAPPLE();
+	}
+	// Use a separate dedicated thread to blit to the screen to avoid expense of swapping contexts.
+	PerformBlockOnRunLoop(^{
+		FScopeLock ScopeLock(Device->ContextUsageGuard);
+		{
+			FScopeContext ScopeContext(Context->OpenGLContext);
 
-		SCOPED_AUTORELEASE_POOL;
-		
-		// Cache & clear the drawable view before resizing the context
-		// otherwise backstore size changes won't be respected.
-		NSView* View = [Context->OpenGLContext view];
-		[Context->OpenGLContext clearDrawable];
-		
-		// Resize the context to the desired resolution directly rather than
-		// trusting the OS. We must do this because of our subversion of AppKit
-		// message handling. If we don't do this ourselves then the OS will not
-		// resize the backing store when the window size changes to adjust the resolution
-		// when exiting fullscreen. This fails despite the transition from fullscreen
-		// having completed according to the notifications.
-		GLint dim[2];
-		dim[0] = SizeX;
-		dim[1] = SizeY;
-		CGLError Err = CGLSetParameter((CGLContextObj)[Context->OpenGLContext CGLContextObj], kCGLCPSurfaceBackingSize, &dim[0]);
-		Err = CGLEnable((CGLContextObj)[Context->OpenGLContext CGLContextObj], kCGLCESurfaceBackingSize);
-		
-		// Restore the drawable view to make it display again.
-		[Context->OpenGLContext setView: View];
-		
-		[Context->OpenGLContext update];
+			SCOPED_AUTORELEASE_POOL;
+			
+			// Cache & clear the drawable view before resizing the context
+			// otherwise backstore size changes won't be respected.
+			NSView* View = [Context->OpenGLContext view];
+			[Context->OpenGLContext clearDrawable];
+			
+			// Resize the context to the desired resolution directly rather than
+			// trusting the OS. We must do this because of our subversion of AppKit
+			// message handling. If we don't do this ourselves then the OS will not
+			// resize the backing store when the window size changes to adjust the resolution
+			// when exiting fullscreen. This fails despite the transition from fullscreen
+			// having completed according to the notifications.
+			GLint dim[2];
+			dim[0] = SizeX;
+			dim[1] = SizeY;
+			CGLError Err = CGLSetParameter((CGLContextObj)[Context->OpenGLContext CGLContextObj], kCGLCPSurfaceBackingSize, &dim[0]);
+			Err = CGLEnable((CGLContextObj)[Context->OpenGLContext CGLContextObj], kCGLCESurfaceBackingSize);
+			
+			// Restore the drawable view to make it display again.
+			[Context->OpenGLContext setView: View];
+			
+			[Context->OpenGLContext update];
 
-		if (Context->ViewportFramebuffer == 0)
-		{
-			glGenFramebuffers(1, &Context->ViewportFramebuffer);
-			check(Context->ViewportFramebuffer);
-		}
-		glBindFramebuffer(GL_FRAMEBUFFER, Context->ViewportFramebuffer);
-		glFramebufferTexture2D(GL_FRAMEBUFFER, GL_COLOR_ATTACHMENT0, BackBufferTarget, BackBufferResource, 0);
-#if UE_BUILD_DEBUG
-		glReadBuffer(GL_COLOR_ATTACHMENT0);
-		glDrawBuffer(GL_COLOR_ATTACHMENT0);
-		GLenum CompleteResult = glCheckFramebufferStatus(GL_FRAMEBUFFER);
-		if (CompleteResult != GL_FRAMEBUFFER_COMPLETE)
-		{
-			UE_LOG(LogRHI, Fatal,TEXT("Framebuffer not complete. Status = 0x%x"), CompleteResult);
-		}
-#endif
+			if (Context->ViewportFramebuffer == 0)
+			{
+				glGenFramebuffers(1, &Context->ViewportFramebuffer);
+				check(Context->ViewportFramebuffer);
+			}
+			glBindFramebuffer(GL_FRAMEBUFFER, Context->ViewportFramebuffer);
+			glFramebufferTexture2D(GL_FRAMEBUFFER, GL_COLOR_ATTACHMENT0, BackBufferTarget, BackBufferResource, 0);
+	#if UE_BUILD_DEBUG
+			glReadBuffer(GL_COLOR_ATTACHMENT0);
+			glDrawBuffer(GL_COLOR_ATTACHMENT0);
+			GLenum CompleteResult = glCheckFramebufferStatus(GL_FRAMEBUFFER);
+			if (CompleteResult != GL_FRAMEBUFFER_COMPLETE)
+			{
+				UE_LOG(LogRHI, Fatal,TEXT("Framebuffer not complete. Status = 0x%x"), CompleteResult);
+			}
+	#endif
 
-		// Clear new buffer to black
-		glViewport(0, 0, SizeX, SizeY);
-		static GLfloat ZeroColor[4] = { 0.0f, 0.0f, 0.0f, 0.0f };
-		glClearBufferfv(GL_COLOR, 0, ZeroColor );
+			// Clear new buffer to black
+			glViewport(0, 0, SizeX, SizeY);
+			static GLfloat ZeroColor[4] = { 0.0f, 0.0f, 0.0f, 0.0f };
+			glClearBufferfv(GL_COLOR, 0, ZeroColor );
 
-		// Set up the state for framebuffer blit
-		glBindFramebuffer(GL_DRAW_FRAMEBUFFER, 0);
-		glDrawBuffer(GL_BACK);
-		glBindFramebuffer(GL_READ_FRAMEBUFFER, Context->ViewportFramebuffer);
-		glReadBuffer(GL_COLOR_ATTACHMENT0);
-	}
+			// Set up the state for framebuffer blit
+			glBindFramebuffer(GL_DRAW_FRAMEBUFFER, 0);
+			glDrawBuffer(GL_BACK);
+			glBindFramebuffer(GL_READ_FRAMEBUFFER, Context->ViewportFramebuffer);
+			glReadBuffer(GL_COLOR_ATTACHMENT0);
+		}
+	}, Device->BlitThread.RunLoop, NSDefaultRunLoopMode, NSDefaultRunLoopMode, true);
 }
 
 void PlatformGetSupportedResolution(uint32 &Width, uint32 &Height)
@@ -1683,6 +1725,11 @@ void FMacOpenGLQuery::End()
 		Finish->End();
 		Current->Begin();
 		Context->LastTimer = Current;
+		
+		if(Start == Finish)
+		{
+			Start = Start->Previous;
+		}
 	}
 	else if(Target != GL_TIMESTAMP)
 	{
