Implement optional & off-by-default support for automatic graphics switching on OS X, rather than just relying on the OS, with a GUI plugin for the editor.
- The very large OpenGLMac file has been broken up:
     - MacOpenGLContext provides the FPlatformOpenGLContext wrapper around NSOpenGLContext & various helpers.
     - MacOpenGLQuery provides the GL_TIMESTAMP emulation wrappers.
- The OpenGL/DisplayMask ini setting specifies an explicit display mask for all contexts just as before.
- The OpenGL/RendererID ini setting specifies an explicit rendering device to use for all contexts, but this is just the initial setting.
- The OpenGL/bUseMultipleRenderers ini setting allows contexts to contain multiple GPUs & render windows using the GPU for the appropriate display, only useful for desktops with multiple GPUs, in this mode the explicit rendering device is only used for offscreen rendering. To qualify all GPUs must provide the same OpenGL version.
- The OpenGL/bAllowAutomaticGraphicsSwitching ini setting turns on automatic graphics switching support which allows for switching between discrete & integrated GPUs which provide the same OpenGL version.
- The r.Mac.ExplicitRendererID cvar allows changing the explicit rendering device used when is OpenGL/bAllowAutomaticGraphicsSwitching enabled, rather than the default which updates the contexts to use the online device, whatever that may be.
- Renderer dependent variables are updated whenever the renderer changes to ensure rendering consistency.
- The plugin provides an editor preferences panel & for laptops using bAllowAutomaticGraphicsSwitching a live drop-down in the LevelEditor window to select the current GPU.

diff --git a/Engine/Plugins/Editor/MacGraphicsSwitching/MacGraphicsSwitching.uplugin b/Engine/Plugins/Editor/MacGraphicsSwitching/MacGraphicsSwitching.uplugin
new file mode 100644
--- /dev/null
+++ b/Engine/Plugins/Editor/MacGraphicsSwitching/MacGraphicsSwitching.uplugin
@@ -0,0 +1,22 @@
+{
+	"FileVersion" : 3,
+	
+	"FriendlyName" : "Mac Graphics Switching",
+	"Version" : 1,
+	"VersionName" : "1.0",
+	"CreatedBy" : "Epic Games, Inc.",
+	"CreatedByURL" : "http://epicgames.com",
+	"EngineVersion" : "4.5.0",
+	"Description" : "Provides support for Mac OS X's automatic graphic switching and multiple graphics device configuration.",
+	"Category" : "Editor.Mac.Graphics",
+	"EnabledByDefault" : true,
+	
+	"Modules" :
+	[
+		{
+			"Name" : "MacGraphicsSwitching",
+			"Type" : "Editor",
+			"WhitelistPlatforms" : [ "Mac" ]
+		}
+	]
+}
\ No newline at end of file
diff --git a/Engine/Plugins/Editor/MacGraphicsSwitching/Source/MacGraphicsSwitching/MacGraphicsSwitching.Build.cs b/Engine/Plugins/Editor/MacGraphicsSwitching/Source/MacGraphicsSwitching/MacGraphicsSwitching.Build.cs
new file mode 100644
--- /dev/null
+++ b/Engine/Plugins/Editor/MacGraphicsSwitching/Source/MacGraphicsSwitching/MacGraphicsSwitching.Build.cs
@@ -0,0 +1,51 @@
+// Copyright 1998-2014 Epic Games, Inc. All Rights Reserved.
+
+namespace UnrealBuildTool.Rules
+{
+	public class MacGraphicsSwitching : ModuleRules
+	{
+        public MacGraphicsSwitching(TargetInfo Target)
+		{
+            PublicDependencyModuleNames.AddRange(
+				new string[]
+				{
+					"Core",
+					"CoreUObject",
+                    "EditorStyle",
+					"Engine",
+                    "InputCore",
+					"LevelEditor",
+					"Slate"
+				}
+			);
+
+            PrivateDependencyModuleNames.AddRange(
+				new string[]
+				{
+					"PropertyEditor",
+					"SlateCore",
+					"UnrealEd"
+				}
+			);
+
+			DynamicallyLoadedModuleNames.AddRange(
+				new string[] 
+				{
+					"LevelEditor",
+					"MainFrame",
+					"Settings"
+				}
+			);
+
+			PrivateIncludePathModuleNames.AddRange(
+				new string[] 
+				{
+					"MainFrame",
+					"Settings"
+				}
+			);
+
+			AddThirdPartyPrivateStaticDependencies(Target, "OpenGL");
+		}
+	}
+}
\ No newline at end of file
diff --git a/Engine/Plugins/Editor/MacGraphicsSwitching/Source/MacGraphicsSwitching/Private/MacGraphicsSwitchingModule.cpp b/Engine/Plugins/Editor/MacGraphicsSwitching/Source/MacGraphicsSwitching/Private/MacGraphicsSwitchingModule.cpp
new file mode 100644
--- /dev/null
+++ b/Engine/Plugins/Editor/MacGraphicsSwitching/Source/MacGraphicsSwitching/Private/MacGraphicsSwitchingModule.cpp
@@ -0,0 +1,201 @@
+// Copyright 1998-2014 Epic Games, Inc. All Rights Reserved.
+
+#include "MacGraphicsSwitchingModule.h"
+#include "IMacGraphicsSwitchingModule.h"
+#include "MacGraphicsSwitchingSettings.h"
+#include "MacGraphicsSwitchingSettingsDetails.h"
+#include "MacGraphicsSwitchingStyle.h"
+#include "MacGraphicsSwitchingWidget.h"
+#include "Runtime/Core/Public/Features/IModularFeatures.h"
+#include "Runtime/SlateCore/Public/Rendering/SlateRenderer.h"
+#include "Developer/Settings/Public/Settings.h"
+#include "Editor/MainFrame/Public/Interfaces/IMainFrameModule.h"
+#include "LevelEditor.h"
+
+#import <IOKit/ps/IOPowerSources.h>
+#import <IOKit/ps/IOPSKeys.h>
+#import <OpenGL/OpenGL.h>
+
+#define LOCTEXT_NAMESPACE "MacGraphicsSwitching"
+
+class FMacGraphicsSwitchingModule : public IMacGraphicsSwitchingModule
+{
+public:
+
+	/** IModuleInterface implementation */
+	virtual void StartupModule() override;
+	virtual void ShutdownModule() override;
+	
+	/** IMacGraphicsSwitchingModule implementation */
+	virtual bool AllowAutomaticGraphicsSwitching() override;
+	virtual bool AllowMultipleGPUs() override;
+	
+private:
+	void Initialize( TSharedPtr<SWindow> InRootWindow, bool bIsNewProjectWindow );
+	void AddGraphicsSwitcher( FToolBarBuilder& ToolBarBuilder );
+	
+private:
+	TSharedPtr< FExtender > NotificationBarExtender;
+	bool bAllowAutomaticGraphicsSwitching;
+	bool bAllowMultiGPUs;
+};
+
+IMPLEMENT_MODULE(FMacGraphicsSwitchingModule, MacGraphicsSwitching)
+
+void FMacGraphicsSwitchingModule::StartupModule()
+{
+	bAllowMultiGPUs = true;
+	bAllowAutomaticGraphicsSwitching = false;
+	
+	CFTypeRef PowerSourcesInfo = IOPSCopyPowerSourcesInfo();
+	if (PowerSourcesInfo)
+	{
+		CFArrayRef PowerSourcesArray = IOPSCopyPowerSourcesList(PowerSourcesInfo);
+		for (CFIndex Index = 0; Index < CFArrayGetCount(PowerSourcesArray); Index++)
+		{
+			CFTypeRef PowerSource = CFArrayGetValueAtIndex(PowerSourcesArray, Index);
+			NSDictionary* Description = (NSDictionary*)IOPSGetPowerSourceDescription(PowerSourcesInfo, PowerSource);
+			if ([(NSString*)[Description objectForKey: @kIOPSTypeKey] isEqualToString: @kIOPSInternalBatteryType])
+			{
+				bAllowAutomaticGraphicsSwitching = true;
+				bAllowMultiGPUs = false;
+				break;
+			}
+		}
+		CFRelease(PowerSourcesArray);
+		CFRelease(PowerSourcesInfo);
+	}
+	
+	if(bAllowMultiGPUs || bAllowAutomaticGraphicsSwitching)
+	{
+		GLint NumberOfRenderers = 0;
+		CGLRendererInfoObj RendererInfo;
+		CGLQueryRendererInfo(0xffffffff, &RendererInfo, &NumberOfRenderers);
+		if(RendererInfo)
+		{
+			TArray<uint32> HardwareRenderers;
+			for(uint32 i = 0; i < NumberOfRenderers; i++)
+			{
+				GLint bAccelerated = 0;
+				CGLDescribeRenderer(RendererInfo, i, kCGLRPAccelerated, &bAccelerated);
+				if(bAccelerated)
+				{
+					HardwareRenderers.Push(i);
+				}
+			}
+			
+			bool bRenderersMatch = true;
+			GLint MajorGLVersion = 0;
+			CGLDescribeRenderer(RendererInfo, 0, kCGLRPMajorGLVersion, &MajorGLVersion);
+			for(uint32 i = 1; i < HardwareRenderers.Num(); i++)
+			{
+				GLint OtherMajorGLVersion = 0;
+				CGLDescribeRenderer(RendererInfo, i, kCGLRPMajorGLVersion, &OtherMajorGLVersion);
+				bRenderersMatch &= MajorGLVersion == OtherMajorGLVersion;
+			}
+			
+			CGLDestroyRendererInfo(RendererInfo);
+			
+			// Can only automatically switch when the renderers use the same major GL version, or errors will occur.
+			bAllowAutomaticGraphicsSwitching &= bRenderersMatch;
+			bAllowMultiGPUs &= bRenderersMatch;
+		}
+	}
+	
+	// Register settings
+	ISettingsModule* SettingsModule = ISettingsModule::Get();
+	if( ensure( SettingsModule != nullptr ) )
+	{
+		SettingsModule->RegisterSettings( "Editor", "Plugins", "MacGraphicsSwitching",
+										 LOCTEXT( "MacGraphicsSwitchingSettingsName", "Graphics Switching"),
+										 LOCTEXT( "MacGraphicsSwitchingSettingsDescription", "Settings for Mac OS X graphics switching"),
+										 GetMutableDefault< UMacGraphicsSwitchingSettings >()
+										 );
+		
+		FPropertyEditorModule& PropertyModule = FModuleManager::LoadModuleChecked<FPropertyEditorModule>("PropertyEditor");
+		PropertyModule.RegisterCustomClassLayout("MacGraphicsSwitchingSettings", FOnGetDetailCustomizationInstance::CreateStatic(&FMacGraphicsSwitchingSettingsDetails::MakeInstance));
+		PropertyModule.NotifyCustomizationModuleChanged();
+		
+		IMainFrameModule& MainFrameModule = FModuleManager::LoadModuleChecked<IMainFrameModule>(TEXT("MainFrame"));
+		MainFrameModule.OnMainFrameCreationFinished().AddRaw( this, &FMacGraphicsSwitchingModule::Initialize );
+	}
+}
+
+void FMacGraphicsSwitchingModule::Initialize( TSharedPtr<SWindow> InRootWindow, bool bIsNewProjectWindow )
+{
+	if( !bIsNewProjectWindow )
+	{
+		IMainFrameModule& MainFrameModule = FModuleManager::LoadModuleChecked<IMainFrameModule>(TEXT("MainFrame"));
+		MainFrameModule.OnMainFrameCreationFinished().RemoveAll(this);
+
+		FMacGraphicsSwitchingStyle::Initialize();
+		
+		bool MacUseAutomaticGraphicsSwitching = false;
+		if ( GConfig->GetBool(TEXT("/Script/MacGraphicsSwitching.MacGraphicsSwitchingSettings"), TEXT("bAllowAutomaticGraphicsSwitching"), MacUseAutomaticGraphicsSwitching, GEditorGameAgnosticIni) && MacUseAutomaticGraphicsSwitching )
+		{
+			if ( AllowAutomaticGraphicsSwitching() )
+			{
+				NotificationBarExtender = MakeShareable( new FExtender() );
+				NotificationBarExtender->AddToolBarExtension("Start",
+															 EExtensionHook::After,
+															 nullptr,
+															 FToolBarExtensionDelegate::CreateRaw(this, &FMacGraphicsSwitchingModule::AddGraphicsSwitcher));
+				
+				FLevelEditorModule& LevelEditorModule = FModuleManager::LoadModuleChecked<FLevelEditorModule>("LevelEditor");
+				LevelEditorModule.GetNotificationBarExtensibilityManager()->AddExtender( NotificationBarExtender );
+				LevelEditorModule.BroadcastNotificationBarChanged();
+			}
+			else
+			{
+				GConfig->SetBool(TEXT("/Script/MacGraphicsSwitching.MacGraphicsSwitchingSettings"), TEXT("bAllowAutomaticGraphicsSwitching"), false, GEditorGameAgnosticIni);
+			}
+		}
+	}
+}
+
+void FMacGraphicsSwitchingModule::AddGraphicsSwitcher( FToolBarBuilder& ToolBarBuilder )
+{
+	ToolBarBuilder.AddWidget( SNew( SMacGraphicsSwitchingWidget ).bLiveSwitching(true) );
+}
+
+void FMacGraphicsSwitchingModule::ShutdownModule()
+{
+	// Unregister settings
+	ISettingsModule* SettingsModule = ISettingsModule::Get();
+	if( SettingsModule != nullptr )
+	{
+		SettingsModule->UnregisterSettings( "Editor", "Plugins", "MacGraphicsSwitching" );
+	}
+	
+	if ( FModuleManager::Get().IsModuleLoaded("LevelEditor") && NotificationBarExtender.IsValid() )
+	{
+		FLevelEditorModule& LevelEditorModule = FModuleManager::GetModuleChecked<FLevelEditorModule>("LevelEditor");
+		LevelEditorModule.GetNotificationBarExtensibilityManager()->RemoveExtender( NotificationBarExtender );
+	}
+	
+	if ( FModuleManager::Get().IsModuleLoaded("MainFrame") )
+	{
+		FMacGraphicsSwitchingStyle::Shutdown();
+		
+		IMainFrameModule& MainFrameModule = FModuleManager::GetModuleChecked<IMainFrameModule>(TEXT("MainFrame"));
+		MainFrameModule.OnMainFrameCreationFinished().RemoveAll(this);
+	}
+	
+	if(FModuleManager::Get().IsModuleLoaded("PropertyEditor"))
+	{
+		FPropertyEditorModule& PropertyEditorModule = FModuleManager::GetModuleChecked<FPropertyEditorModule>("PropertyEditor");
+		PropertyEditorModule.UnregisterCustomClassLayout("MacGraphicsSwitchingSettings");
+	}
+}
+
+bool FMacGraphicsSwitchingModule::AllowAutomaticGraphicsSwitching()
+{
+	return bAllowAutomaticGraphicsSwitching;
+}
+
+bool FMacGraphicsSwitchingModule::AllowMultipleGPUs()
+{
+	return bAllowMultiGPUs;
+}
+
+#undef LOCTEXT_NAMESPACE
diff --git a/Engine/Plugins/Editor/MacGraphicsSwitching/Source/MacGraphicsSwitching/Private/MacGraphicsSwitchingModule.h b/Engine/Plugins/Editor/MacGraphicsSwitching/Source/MacGraphicsSwitching/Private/MacGraphicsSwitchingModule.h
new file mode 100644
--- /dev/null
+++ b/Engine/Plugins/Editor/MacGraphicsSwitching/Source/MacGraphicsSwitching/Private/MacGraphicsSwitchingModule.h
@@ -0,0 +1,5 @@
+// Copyright 1998-2014 Epic Games, Inc. All Rights Reserved.
+
+#include "Engine.h"
+#include "Slate.h"
+#include "UnrealEd.h"
diff --git a/Engine/Plugins/Editor/MacGraphicsSwitching/Source/MacGraphicsSwitching/Private/MacGraphicsSwitchingSettings.cpp b/Engine/Plugins/Editor/MacGraphicsSwitching/Source/MacGraphicsSwitching/Private/MacGraphicsSwitchingSettings.cpp
new file mode 100644
--- /dev/null
+++ b/Engine/Plugins/Editor/MacGraphicsSwitching/Source/MacGraphicsSwitching/Private/MacGraphicsSwitchingSettings.cpp
@@ -0,0 +1,9 @@
+// Copyright 1998-2014 Epic Games, Inc. All Rights Reserved.
+
+#include "MacGraphicsSwitchingModule.h"
+#include "MacGraphicsSwitchingSettings.h"
+
+UMacGraphicsSwitchingSettings::UMacGraphicsSwitchingSettings( const class FPostConstructInitializeProperties& PCIP )
+	: Super(PCIP)
+{
+}
diff --git a/Engine/Plugins/Editor/MacGraphicsSwitching/Source/MacGraphicsSwitching/Private/MacGraphicsSwitchingSettings.h b/Engine/Plugins/Editor/MacGraphicsSwitching/Source/MacGraphicsSwitching/Private/MacGraphicsSwitchingSettings.h
new file mode 100644
--- /dev/null
+++ b/Engine/Plugins/Editor/MacGraphicsSwitching/Source/MacGraphicsSwitching/Private/MacGraphicsSwitchingSettings.h
@@ -0,0 +1,24 @@
+// Copyright 1998-2014 Epic Games, Inc. All Rights Reserved.
+#pragma once
+
+#include "MacGraphicsSwitchingSettings.generated.h"
+
+UCLASS(Config=EditorGameAgnostic)
+class UMacGraphicsSwitchingSettings : public UObject
+{
+	GENERATED_UCLASS_BODY()
+	
+public:
+	
+	UPROPERTY(Config, EditAnywhere, Category=OpenGL)
+	int32 DisplayMask;
+
+	UPROPERTY(Config, EditAnywhere, Category=OpenGL)
+	int32 RendererID;
+	
+	UPROPERTY(Config, EditAnywhere, Category=OpenGL)
+	bool bAllowAutomaticGraphicsSwitching;
+	
+	UPROPERTY(Config, EditAnywhere, Category=OpenGL)
+	bool bUseMultipleRenderers;
+};
diff --git a/Engine/Plugins/Editor/MacGraphicsSwitching/Source/MacGraphicsSwitching/Private/MacGraphicsSwitchingSettingsDetails.cpp b/Engine/Plugins/Editor/MacGraphicsSwitching/Source/MacGraphicsSwitching/Private/MacGraphicsSwitchingSettingsDetails.cpp
new file mode 100644
--- /dev/null
+++ b/Engine/Plugins/Editor/MacGraphicsSwitching/Source/MacGraphicsSwitching/Private/MacGraphicsSwitchingSettingsDetails.cpp
@@ -0,0 +1,52 @@
+// Copyright 1998-2014 Epic Games, Inc. All Rights Reserved.
+
+#include "MacGraphicsSwitchingModule.h"
+#include "MacGraphicsSwitchingSettingsDetails.h"
+#include "MacGraphicsSwitchingWidget.h"
+
+#define LOCTEXT_NAMESPACE "MacGraphicsSwitchingSettingsDetails"
+
+TSharedRef<IDetailCustomization> FMacGraphicsSwitchingSettingsDetails::MakeInstance()
+{
+	return MakeShareable(new FMacGraphicsSwitchingSettingsDetails());
+}
+
+void FMacGraphicsSwitchingSettingsDetails::CustomizeDetails( IDetailLayoutBuilder& DetailLayout )
+{
+	TSharedRef<IPropertyHandle> PreferredRendererPropertyHandle = DetailLayout.GetProperty("RendererID");
+	DetailLayout.HideProperty("RendererID");
+
+	bool bAllowMultiGPUs = IMacGraphicsSwitchingModule::Get().AllowMultipleGPUs();
+	bool bAllowAutomaticGraphicsSwitching = IMacGraphicsSwitchingModule::Get().AllowAutomaticGraphicsSwitching();
+	
+	TSharedRef<IPropertyHandle> MultiGPUPropertyHandle = DetailLayout.GetProperty("bUseMultipleRenderers");
+	if (!bAllowMultiGPUs)
+	{
+		MultiGPUPropertyHandle->SetValue(false);
+		DetailLayout.HideProperty("bUseMultipleRenderers");
+	}
+	
+	TSharedRef<IPropertyHandle> SwitchingPropertyHandle = DetailLayout.GetProperty("bAllowAutomaticGraphicsSwitching");
+	if (!bAllowAutomaticGraphicsSwitching)
+	{
+		SwitchingPropertyHandle->SetValue(false);
+		DetailLayout.HideProperty("bAllowAutomaticGraphicsSwitching");
+	}
+	
+	IDetailCategoryBuilder& AccessorCategory = DetailLayout.EditCategory( "OpenGL" );
+	AccessorCategory.AddCustomRow( LOCTEXT("PreferredRenderer", "Preferred Renderer").ToString() )
+	.NameContent()
+	[
+		PreferredRendererPropertyHandle->CreatePropertyNameWidget()
+	]
+	.ValueContent()
+	.MinDesiredWidth(113)
+	.MaxDesiredWidth(113)
+	[
+		SNew(SMacGraphicsSwitchingWidget)
+		.bLiveSwitching(false)
+		.PreferredRendererPropertyHandle(PreferredRendererPropertyHandle)
+	];
+}
+
+#undef LOCTEXT_NAMESPACE
\ No newline at end of file
diff --git a/Engine/Plugins/Editor/MacGraphicsSwitching/Source/MacGraphicsSwitching/Private/MacGraphicsSwitchingSettingsDetails.h b/Engine/Plugins/Editor/MacGraphicsSwitching/Source/MacGraphicsSwitching/Private/MacGraphicsSwitchingSettingsDetails.h
new file mode 100644
--- /dev/null
+++ b/Engine/Plugins/Editor/MacGraphicsSwitching/Source/MacGraphicsSwitching/Private/MacGraphicsSwitchingSettingsDetails.h
@@ -0,0 +1,16 @@
+// Copyright 1998-2014 Epic Games, Inc. All Rights Reserved.
+
+#pragma once
+
+#include "PropertyEditing.h"
+#include "PropertyCustomizationHelpers.h"
+
+class FMacGraphicsSwitchingSettingsDetails : public IDetailCustomization
+{
+public:
+	/** Makes a new instance of this detail layout class for a specific detail view requesting it */
+	static TSharedRef<IDetailCustomization> MakeInstance();
+
+	/** IDetailCustomization interface */
+	virtual void CustomizeDetails( IDetailLayoutBuilder& DetailLayout ) override;
+};
\ No newline at end of file
diff --git a/Engine/Plugins/Editor/MacGraphicsSwitching/Source/MacGraphicsSwitching/Private/MacGraphicsSwitchingStyle.cpp b/Engine/Plugins/Editor/MacGraphicsSwitching/Source/MacGraphicsSwitching/Private/MacGraphicsSwitchingStyle.cpp
new file mode 100644
--- /dev/null
+++ b/Engine/Plugins/Editor/MacGraphicsSwitching/Source/MacGraphicsSwitching/Private/MacGraphicsSwitchingStyle.cpp
@@ -0,0 +1,68 @@
+// Copyright 1998-2014 Epic Games, Inc. All Rights Reserved.
+
+#include "MacGraphicsSwitchingModule.h"
+#include "MacGraphicsSwitchingStyle.h"
+#include "SlateStyle.h"
+
+#define IMAGE_PLUGIN_BRUSH( RelativePath, ... ) FSlateImageBrush( FPaperStyle::InContent( RelativePath, ".png" ), __VA_ARGS__ )
+#define IMAGE_BRUSH(RelativePath, ...) FSlateImageBrush(StyleSet->RootToContentDir(RelativePath, TEXT(".png")), __VA_ARGS__)
+#define BOX_BRUSH(RelativePath, ...) FSlateBoxBrush(StyleSet->RootToContentDir(RelativePath, TEXT(".png")), __VA_ARGS__)
+
+TSharedPtr< FSlateStyleSet > FMacGraphicsSwitchingStyle::StyleSet = NULL;
+TSharedPtr< class ISlateStyle > FMacGraphicsSwitchingStyle::Get() { return StyleSet; }
+
+void FMacGraphicsSwitchingStyle::Initialize()
+{
+	const FVector2D Icon8x8(8.0f, 8.0f);
+	const FLinearColor DisabledColor(FColor(0xaaaaaa));
+
+	// Only register once
+	if (StyleSet.IsValid())
+	{
+		return;
+	}
+
+	StyleSet = MakeShareable(new FSlateStyleSet("MacGraphicsSwitchingStyle"));
+	StyleSet->SetContentRoot(FPaths::EngineContentDir() / TEXT("Editor/Slate"));
+	StyleSet->SetCoreContentRoot(FPaths::EngineContentDir() / TEXT("Slate"));
+
+	{
+		const FButtonStyle ComboBoxButton = FButtonStyle()
+		.SetNormal(BOX_BRUSH("Tutorials/FlatColorSquare", FVector2D(1.0f, 1.0f), FMargin(1), FLinearColor(0.05f,0.05f,0.05f)))
+		.SetHovered(BOX_BRUSH("Tutorials/FlatColorSquare", FVector2D(1.0f, 1.0f), FMargin(1), FLinearColor(0.05f,0.05f,0.05f)))
+		.SetPressed(BOX_BRUSH("Tutorials/FlatColorSquare", FVector2D(1.0f, 1.0f), FMargin(1), FLinearColor(0.05f,0.05f,0.05f)))
+		.SetDisabled(BOX_BRUSH("Tutorials/FlatColorSquare", FVector2D(1.0f, 1.0f), FMargin(1), DisabledColor))
+		.SetNormalPadding(FMargin(2, 2, 2, 2))
+		.SetPressedPadding(FMargin(2, 3, 2, 1));
+		
+		const FComboButtonStyle ComboButton = FComboButtonStyle()
+		.SetButtonStyle(FEditorStyle::Get().GetWidgetStyle<FButtonStyle>("Button"))
+		.SetDownArrowImage(IMAGE_BRUSH("Common/ComboArrow", Icon8x8))
+		.SetMenuBorderBrush(BOX_BRUSH("Tutorials/FlatColorSquare", FVector2D(1.0f, 1.0f), FMargin(1), FLinearColor(0.05f,0.05f,0.05f)))
+		.SetMenuBorderPadding(FMargin(0.0f));
+		
+		const FComboButtonStyle ComboBoxComboButton = FComboButtonStyle(ComboButton)
+		.SetButtonStyle(ComboBoxButton)
+		.SetMenuBorderPadding(FMargin(1.0));
+		
+		StyleSet->Set("MacGraphicsSwitcher.ComboBox", FComboBoxStyle()
+			.SetComboButtonStyle(ComboBoxComboButton)
+			);
+	}
+
+	FSlateStyleRegistry::RegisterSlateStyle(*StyleSet.Get());
+};
+
+#undef IMAGE_PLUGIN_BRUSH
+#undef IMAGE_BRUSH
+#undef BOX_BRUSH
+
+void FMacGraphicsSwitchingStyle::Shutdown()
+{
+	if (StyleSet.IsValid())
+	{
+		FSlateStyleRegistry::UnRegisterSlateStyle(*StyleSet.Get());
+		ensure(StyleSet.IsUnique());
+		StyleSet.Reset();
+	}
+}
diff --git a/Engine/Plugins/Editor/MacGraphicsSwitching/Source/MacGraphicsSwitching/Private/MacGraphicsSwitchingStyle.h b/Engine/Plugins/Editor/MacGraphicsSwitching/Source/MacGraphicsSwitching/Private/MacGraphicsSwitchingStyle.h
new file mode 100644
--- /dev/null
+++ b/Engine/Plugins/Editor/MacGraphicsSwitching/Source/MacGraphicsSwitching/Private/MacGraphicsSwitchingStyle.h
@@ -0,0 +1,16 @@
+// Copyright 1998-2014 Epic Games, Inc. All Rights Reserved.
+
+#pragma once
+
+class FMacGraphicsSwitchingStyle
+{
+public:
+	static void Initialize();
+
+	static void Shutdown();
+
+	static TSharedPtr< class ISlateStyle > Get();
+
+private:
+	static TSharedPtr< class FSlateStyleSet > StyleSet;
+};
\ No newline at end of file
diff --git a/Engine/Plugins/Editor/MacGraphicsSwitching/Source/MacGraphicsSwitching/Private/MacGraphicsSwitchingWidget.cpp b/Engine/Plugins/Editor/MacGraphicsSwitching/Source/MacGraphicsSwitching/Private/MacGraphicsSwitchingWidget.cpp
new file mode 100644
--- /dev/null
+++ b/Engine/Plugins/Editor/MacGraphicsSwitching/Source/MacGraphicsSwitching/Private/MacGraphicsSwitchingWidget.cpp
@@ -0,0 +1,169 @@
+// Copyright 1998-2014 Epic Games, Inc. All Rights Reserved.
+
+#include "MacGraphicsSwitchingModule.h"
+#include "MacGraphicsSwitchingStyle.h"
+#include "MacGraphicsSwitchingWidget.h"
+#include "DetailLayoutBuilder.h"
+#include <OpenGL/gl3.h>
+
+#define LOCTEXT_NAMESPACE "MacGraphicsSwitchingWidget"
+
+SMacGraphicsSwitchingWidget::SMacGraphicsSwitchingWidget()
+: bLiveSwitching(false)
+{
+	// regenerate renderer list
+	Renderers.Empty();
+	
+	FString DefaultDesc = FString(TEXT("System Default: ")) + FPlatformMisc::GetPrimaryGPUBrand();
+	Renderers.Add(MakeShareable(new FRendererItem(FText::FromString(DefaultDesc), 0)));
+	
+	NSOpenGLPixelFormatAttribute Attributes[] =
+	{
+		NSOpenGLPFAOpenGLProfile,
+		NSOpenGLProfileVersion3_2Core,
+		NSOpenGLPFAAccelerated,
+		NSOpenGLPFANoRecovery,
+		NSOpenGLPFAAllowOfflineRenderers,
+		kCGLPFASupportsAutomaticGraphicsSwitching,
+		NSOpenGLPFADoubleBuffer,
+		NSOpenGLPFAColorSize,
+		32,
+		0
+	};
+	
+	NSOpenGLPixelFormat* PixelFormat = [[NSOpenGLPixelFormat alloc] initWithAttributes: Attributes];
+	if (PixelFormat)
+	{
+		NSOpenGLContext* Context = [[NSOpenGLContext alloc] initWithFormat: PixelFormat shareContext: nil];
+		
+		if (Context)
+		{
+			NSOpenGLContext* PreviousContext = [NSOpenGLContext currentContext];
+			[Context makeCurrentContext];
+			
+			for(uint32 i = 0; i < [PixelFormat numberOfVirtualScreens]; i++)
+			{
+				[Context setCurrentVirtualScreen: i];
+				
+				GLint RendererID = 0;
+				[Context getValues:&RendererID forParameter:NSOpenGLCPCurrentRendererID];
+				
+				FString Gpu(ANSI_TO_TCHAR((const ANSICHAR*)glGetString(GL_RENDERER)));
+				
+				FString RendererDesc = FString::Printf(TEXT("%d: %s"), i, *Gpu);
+				
+				Renderers.Add(MakeShareable(new FRendererItem(FText::FromString(RendererDesc), RendererID)));
+			}
+			
+			if ( PreviousContext )
+			{
+				[PreviousContext makeCurrentContext];
+			}
+			else
+			{
+				[NSOpenGLContext clearCurrentContext];
+			}
+			
+			[Context release];
+		}
+		[PixelFormat release];
+	}
+}
+
+SMacGraphicsSwitchingWidget::~SMacGraphicsSwitchingWidget()
+{
+	
+}
+
+void SMacGraphicsSwitchingWidget::Construct(SMacGraphicsSwitchingWidget::FArguments const& InArgs)
+{
+	bLiveSwitching = InArgs._bLiveSwitching;
+	
+	if ( bLiveSwitching )
+	{
+		ChildSlot
+		[
+			SNew(SComboBox< TSharedPtr<FRendererItem>>)
+			.ToolTipText(LOCTEXT("PreferredRendererToolTip", "Choose the preferred rendering device."))
+			.ComboBoxStyle(FMacGraphicsSwitchingStyle::Get(), "MacGraphicsSwitcher.ComboBox")
+			.ForegroundColor(FCoreStyle::Get().GetSlateColor("DefaultForeground"))
+			.OptionsSource(&Renderers)
+			.OnSelectionChanged(this, &SMacGraphicsSwitchingWidget::OnSelectionChanged, InArgs._PreferredRendererPropertyHandle)
+			.ContentPadding(2)
+			.OnGenerateWidget(this, &SMacGraphicsSwitchingWidget::OnGenerateWidget)
+			.Content()
+			[
+				SNew(STextBlock)
+				.Text(this, &SMacGraphicsSwitchingWidget::GetRendererText)
+				.Font( IDetailLayoutBuilder::GetDetailFont() )
+			]
+		];
+	}
+	else
+	{
+		ChildSlot
+		[
+			SNew(SComboBox< TSharedPtr<FRendererItem>>)
+			.ToolTipText(LOCTEXT("PreferredRendererToolTip", "Choose the preferred rendering device."))
+			.OptionsSource(&Renderers)
+			.OnSelectionChanged(this, &SMacGraphicsSwitchingWidget::OnSelectionChanged, InArgs._PreferredRendererPropertyHandle)
+			.ContentPadding(2)
+			.OnGenerateWidget(this, &SMacGraphicsSwitchingWidget::OnGenerateWidget)
+			.Content()
+			[
+				SNew(STextBlock)
+				.Text(this, &SMacGraphicsSwitchingWidget::GetRendererText)
+				.Font( IDetailLayoutBuilder::GetDetailFont() )
+			]
+		];
+	}
+}
+
+TSharedRef<SWidget> SMacGraphicsSwitchingWidget::OnGenerateWidget( TSharedPtr<FRendererItem> InItem )
+{
+	return SNew(STextBlock)
+		.Text(InItem->Text);
+}
+
+void SMacGraphicsSwitchingWidget::OnSelectionChanged(TSharedPtr<FRendererItem> InItem, ESelectInfo::Type InSeletionInfo, TSharedPtr<IPropertyHandle> PreferredRendererPropertyHandle)
+{
+	if ( PreferredRendererPropertyHandle.IsValid() )
+	{
+		PreferredRendererPropertyHandle->SetValue(InItem->RendererID);
+	}
+	else if ( bLiveSwitching )
+	{
+		GConfig->SetInt(TEXT("/Script/MacGraphicsSwitching.MacGraphicsSwitchingSettings"), TEXT("RendererID"), InItem->RendererID, GEditorGameAgnosticIni);
+	}
+	
+	// Update all the contexts
+	static const auto CVar = IConsoleManager::Get().FindConsoleVariable(TEXT("r.Mac.ExplicitRendererID"));
+	if ( CVar )
+	{
+		CVar->Set( (int32)InItem->RendererID );
+	}
+}
+
+FText SMacGraphicsSwitchingWidget::GetRendererText() const
+{
+	FText Text = FText::FromString(FString::Printf(TEXT("System Default: %s"), *FPlatformMisc::GetPrimaryGPUBrand()));
+
+	int32 ExplicitRendererId = 0;
+	GConfig->GetInt(TEXT("/Script/MacGraphicsSwitching.MacGraphicsSwitchingSettings"), TEXT("RendererID"), ExplicitRendererId, GEngineIni);
+	
+	if ( ExplicitRendererId )
+	{
+		for (auto Renderer: Renderers)
+		{
+			if ( Renderer->RendererID == ExplicitRendererId )
+			{
+				Text = Renderer->Text;
+				break;
+			}
+		}
+	}
+	
+	return Text;
+}
+
+#undef LOCTEXT_NAMESPACE
\ No newline at end of file
diff --git a/Engine/Plugins/Editor/MacGraphicsSwitching/Source/MacGraphicsSwitching/Private/MacGraphicsSwitchingWidget.h b/Engine/Plugins/Editor/MacGraphicsSwitching/Source/MacGraphicsSwitching/Private/MacGraphicsSwitchingWidget.h
new file mode 100644
--- /dev/null
+++ b/Engine/Plugins/Editor/MacGraphicsSwitching/Source/MacGraphicsSwitching/Private/MacGraphicsSwitchingWidget.h
@@ -0,0 +1,54 @@
+// Copyright 1998-2014 Epic Games, Inc. All Rights Reserved.
+
+#pragma once
+
+#include "IMacGraphicsSwitchingModule.h"
+#include "Slate.h"
+#include "SlateStyle.h"
+
+struct FRendererItem
+{
+	FRendererItem(const FText& InText, int32 const InRendererID)
+		: Text(InText)
+		, RendererID(InRendererID)
+	{
+	}
+
+	/** Text to display */
+	FText Text;
+
+	/** ID of the renderer */
+	int32 RendererID;
+};
+
+class SMacGraphicsSwitchingWidget : public SCompoundWidget
+{
+public:
+	SLATE_BEGIN_ARGS(SMacGraphicsSwitchingWidget)
+		: _bLiveSwitching(false) {}
+	
+		SLATE_ARGUMENT(TSharedPtr<IPropertyHandle>, PreferredRendererPropertyHandle)
+		SLATE_ARGUMENT(bool, bLiveSwitching)
+	SLATE_END_ARGS();
+	
+	SMacGraphicsSwitchingWidget();
+	virtual ~SMacGraphicsSwitchingWidget();
+	
+	void Construct(FArguments const& InArgs);
+	
+private:
+	/** Generate a row widget for display in the list view */
+	TSharedRef<SWidget> OnGenerateWidget( TSharedPtr<FRendererItem> InItem );
+	
+	/** Set the renderer when we change selection */
+	void OnSelectionChanged(TSharedPtr<FRendererItem> InItem, ESelectInfo::Type InSeletionInfo, TSharedPtr<IPropertyHandle> PreferredProviderPropertyHandle);
+	
+	/** Get the text to display on the renderer drop-down */
+	FText GetRendererText() const;
+	
+private:
+	/** Accessor names to display in a drop-down list */
+	TArray<TSharedPtr<FRendererItem>> Renderers;
+	/** Whether we are modifying the current or default */
+	bool bLiveSwitching;
+};
diff --git a/Engine/Plugins/Editor/MacGraphicsSwitching/Source/MacGraphicsSwitching/Public/IMacGraphicsSwitchingModule.h b/Engine/Plugins/Editor/MacGraphicsSwitching/Source/MacGraphicsSwitching/Public/IMacGraphicsSwitchingModule.h
new file mode 100644
--- /dev/null
+++ b/Engine/Plugins/Editor/MacGraphicsSwitching/Source/MacGraphicsSwitching/Public/IMacGraphicsSwitchingModule.h
@@ -0,0 +1,48 @@
+// Copyright 1998-2014 Epic Games, Inc. All Rights Reserved.
+
+#pragma once
+
+#include "ModuleManager.h"
+
+/**
+ * The public interface to this module
+ */
+class IMacGraphicsSwitchingModule : public IModuleInterface
+{
+
+public:
+
+	/**
+	 * Singleton-like access to this module's interface.  This is just for convenience!
+	 * Beware of calling this during the shutdown phase, though.  Your module might have been unloaded already.
+	 *
+	 * @return Returns singleton instance, loading the module on demand if needed
+	 */
+	static inline IMacGraphicsSwitchingModule& Get()
+	{
+		return FModuleManager::LoadModuleChecked< IMacGraphicsSwitchingModule >( "MacGraphicsSwitching" );
+	}
+
+	/**
+	 * Checks to see if this module is loaded and ready.  It is only valid to call Get() if IsAvailable() returns true.
+	 *
+	 * @return True if the module is loaded and ready to use
+	 */
+	static inline bool IsAvailable()
+	{
+		return FModuleManager::Get().IsModuleLoaded( "MacGraphicsSwitching" );
+	}
+	
+	/**
+	 * Tests whether this Mac can support automatic graphics switching.
+	 * @return True if automatic graphics switching is allowed on this Mac.
+	 */
+	virtual bool AllowAutomaticGraphicsSwitching() { return false; }
+	
+	/**
+	 * Tests whether this Mac can support rendering to displays via multiple GPUs.
+	 * @return True if multiple GPU support is allowed on this Mac.
+	 */
+	virtual bool AllowMultipleGPUs() { return false; }
+};
+
diff --git a/Engine/Source/Runtime/OpenGLDrv/Private/Mac/MacOpenGLContext.cpp b/Engine/Source/Runtime/OpenGLDrv/Private/Mac/MacOpenGLContext.cpp
new file mode 100644
--- /dev/null
+++ b/Engine/Source/Runtime/OpenGLDrv/Private/Mac/MacOpenGLContext.cpp
@@ -0,0 +1,538 @@
+// Copyright 1998-2014 Epic Games, Inc. All Rights Reserved.
+
+#include "OpenGLDrvPrivate.h"
+
+#include "MacOpenGLContext.h"
+#include "MacOpenGLQuery.h"
+
+#include <IOKit/ps/IOPowerSources.h>
+#include <IOKit/ps/IOPSKeys.h>
+#include <OpenGL/gl3.h>
+
+/*------------------------------------------------------------------------------
+ OpenGL defines.
+ ------------------------------------------------------------------------------*/
+
+#if WITH_EDITOR
+#define MAC_OPENGL_SETTINGS TEXT("/Script/MacGraphicsSwitching.MacGraphicsSwitchingSettings")
+#define MAC_OPENGL_INI GEditorGameAgnosticIni
+#else
+#define MAC_OPENGL_SETTINGS TEXT("OpenGL")
+#define MAC_OPENGL_INI GEngineIni
+#endif
+
+/*------------------------------------------------------------------------------
+ OpenGL static variables.
+ ------------------------------------------------------------------------------*/
+
+/** Whether we allow use of multiple GPUs. */
+static bool GMacUseMultipleGPUs = false;
+
+/** Whether we use automatic graphics switching or rely on the system to handle it for us. */
+static bool GMacUseAutomaticGraphicsSwitching = false;
+
+/** The current global output renderer ID. */
+static GLint GMacCurrentRendererID = 0;
+
+/** The selected explicit renderer ID. */
+static int32 GMacExplicitRendererID = 0;
+static FAutoConsoleVariableRef CVarMacExplicitRendererID(
+	TEXT("r.Mac.ExplicitRendererID"),
+	GMacExplicitRendererID,
+	TEXT("Explicitly sets all contexts to the provided renderer ID when using automatic graphics switching. (Default: 0, off)"),
+	ECVF_RenderThreadSafe
+	);
+
+/*------------------------------------------------------------------------------
+ OpenGL internal static variables.
+ ------------------------------------------------------------------------------*/
+
+/** Default context attributes */
+const NSOpenGLPixelFormatAttribute ContextCreationAttributes[] =
+{
+	kCGLPFAOpenGLProfile,
+	kCGLOGLPVersion_3_2_Core,
+	kCGLPFAAccelerated,
+	kCGLPFANoRecovery,
+	kCGLPFADoubleBuffer,
+	kCGLPFAColorSize,
+	32,
+	0
+};
+
+/** The context pixel format is cached to reduce the expense of creating new contexts. */
+static NSOpenGLPixelFormat* PixelFormat = nil;
+
+/*------------------------------------------------------------------------------
+	OpenGL context helper functions.
+------------------------------------------------------------------------------*/
+
+static NSOpenGLContext* CreateContext( NSOpenGLContext* SharedContext )
+{
+	SCOPED_AUTORELEASE_POOL;
+
+	if (!PixelFormat)
+	{
+		TArray<NSOpenGLPixelFormatAttribute> Attributes;
+		Attributes.Append(ContextCreationAttributes, ARRAY_COUNT(ContextCreationAttributes) - 1);
+		
+		int32 DisplayMask = 0;
+		if (GConfig->GetInt(MAC_OPENGL_SETTINGS, TEXT("DisplayMask"), DisplayMask, MAC_OPENGL_INI) && DisplayMask)
+		{
+			Attributes.Add(NSOpenGLPFAScreenMask);
+			Attributes.Add(DisplayMask);
+		}
+		
+		// Only initialise explicit renderer ID values once, require a restart to change, as otherwise the contexts cannot be in the same share group
+		static bool bExplicitRendererSetup = false;
+		if(!bExplicitRendererSetup)
+		{
+			int32 ExplicitRendererId = 0;
+			if(GConfig->GetInt(MAC_OPENGL_SETTINGS, TEXT("RendererID"), ExplicitRendererId, MAC_OPENGL_INI) && ExplicitRendererId)
+			{
+				// Be sure to test that the renderer actually exists, otherwise device creation will fail and that'd be bad
+				bool bRendererFound = false;
+				GLint NumberOfRenderers = 0;
+				CGLRendererInfoObj RendererInfo;
+				// If the display mask is overridden then respect that and disallow renderers not bound to that display
+				int32 OpenGLDisplayMask = DisplayMask ? DisplayMask : 0xffffffff;
+				CGLQueryRendererInfo(OpenGLDisplayMask, &RendererInfo, &NumberOfRenderers);
+				if(RendererInfo)
+				{
+					for(uint32 i = 0; i < NumberOfRenderers; i++)
+					{
+						GLint DeviceRendererID = 0;
+						CGLDescribeRenderer(RendererInfo, i, kCGLRPRendererID, &DeviceRendererID);
+						if(ExplicitRendererId == DeviceRendererID)
+						{
+							bRendererFound = true;
+							break;
+						}
+					}
+					
+					CGLDestroyRendererInfo(RendererInfo);
+				}
+				
+				GMacExplicitRendererID = bRendererFound ? ExplicitRendererId : 0;
+				GMacCurrentRendererID = GMacExplicitRendererID;
+			}
+			bExplicitRendererSetup = true;
+		}
+		
+		// Specify a single explicit renderer ID
+		if (GMacExplicitRendererID && (GNumActiveGPUsForRendering == 1) && !GMacUseAutomaticGraphicsSwitching)
+		{
+			Attributes.Add(NSOpenGLPFARendererID);
+			Attributes.Add(GMacExplicitRendererID);
+		}
+		
+		if(GMacUseAutomaticGraphicsSwitching)
+		{
+			Attributes.Add(kCGLPFASupportsAutomaticGraphicsSwitching);
+		}
+		
+		if(DisplayMask || GMacExplicitRendererID || GNumActiveGPUsForRendering > 1)
+		{
+			Attributes.Add(NSOpenGLPFAAllowOfflineRenderers);
+		}
+		
+		Attributes.Add(0);
+		
+		PixelFormat = [[NSOpenGLPixelFormat alloc] initWithAttributes: Attributes.GetData()];
+		check(PixelFormat);
+	}
+
+	NSOpenGLContext* Context = [[NSOpenGLContext alloc] initWithFormat: PixelFormat shareContext: SharedContext];
+	check(Context);
+	
+	if((GNumActiveGPUsForRendering > 1 || GMacUseAutomaticGraphicsSwitching) && GMacExplicitRendererID)
+	{
+		for(uint32 i = 0; i < [PixelFormat numberOfVirtualScreens]; i++)
+		{
+			GLint ScreenRendererID = 0;
+			[PixelFormat getValues:&ScreenRendererID forAttribute:NSOpenGLPFARendererID forVirtualScreen:i];
+			
+			if(ScreenRendererID == GMacExplicitRendererID)
+			{
+				[Context setCurrentVirtualScreen:i];
+				break;
+			}
+		}
+	}
+	
+	int32 SyncInterval = 0;
+	[Context setValues: &SyncInterval forParameter: NSOpenGLCPSwapInterval];
+	
+	// Setup Opacity - can't change it dynamically later!
+	int32 SurfaceOpacity = 0;
+	[Context setValues: &SurfaceOpacity forParameter: NSOpenGLCPSurfaceOpacity];
+
+	extern int32 GMacUseMTGL;
+	if (GMacUseMTGL)
+	{
+		CGLEnable((CGLContextObj)[Context CGLContextObj], kCGLCEMPEngine);
+		
+		// @todo: We should disable OpenGL.UseMapBuffer for improved performance under MTGL, but radr://17662549 prevents it when using GL_TEXTURE_BUFFER's for skinning.
+	}
+
+	return Context;
+}
+
+void MacOpenGLContextReconfigurationCallBack(CGDirectDisplayID Display, CGDisplayChangeSummaryFlags Flags, void* UserInfo)
+{
+    if(GMacUseAutomaticGraphicsSwitching && (Flags & kCGDisplaySetModeFlag))
+    {
+        // Display has been reconfigured so find the online GPU so we may adapt to the new capabilities.
+		GLint NumberOfRenderers = 0;
+		CGLRendererInfoObj RendererInfo;
+		CGLQueryRendererInfo(0xffffffff, &RendererInfo, &NumberOfRenderers);
+		if(RendererInfo)
+		{
+			for(uint32 i = 0; i < NumberOfRenderers; i++)
+			{
+				GLint DeviceOnline = 0;
+				CGLDescribeRenderer(RendererInfo, i, kCGLRPOnline, &DeviceOnline);
+				if(DeviceOnline)
+				{
+					CGLDescribeRenderer(RendererInfo, i, kCGLRPRendererID, &GMacCurrentRendererID);
+					break;
+				}
+			}
+			
+			CGLDestroyRendererInfo(RendererInfo);
+		}
+    }
+}
+	
+/*------------------------------------------------------------------------------
+ OpenGL context.
+ ------------------------------------------------------------------------------*/
+
+FPlatformOpenGLContext::FPlatformOpenGLContext()
+: OpenGLContext(nil)
+, OpenGLView(nil)
+, WindowHandle(nil)
+, EmulatedQueries(nullptr)
+, VertexArrayObject(0)
+, ViewportFramebuffer(0)
+, SyncInterval(0)
+, VirtualScreen(0)
+, VendorID(0)
+, RendererID(0)
+, SupportsDepthFetchDuringDepthTest(false)
+{
+}
+
+FPlatformOpenGLContext::~FPlatformOpenGLContext()
+{
+	SCOPED_AUTORELEASE_POOL;
+	
+	NSOpenGLContext* CurrentContext = [NSOpenGLContext currentContext];
+	bool bMadeCurrent = false;
+	
+	if ((ViewportFramebuffer && !CurrentContext) || VertexArrayObject || GIsEmulatingTimestamp)
+	{
+		MakeCurrent();
+		bMadeCurrent = true;
+	}
+	
+	if (ViewportFramebuffer)
+	{
+		// this can be done from any context, as long as it's not nil.
+		glDeleteFramebuffers(1, &ViewportFramebuffer);
+	}
+	
+	if(VertexArrayObject)
+	{
+		glBindVertexArray(0);
+		glDeleteVertexArrays(1, &VertexArrayObject);
+	}
+	
+	delete EmulatedQueries;
+	
+	if ( OpenGLContext )
+	{
+		// Unbind the platform context from the CGL context
+		intptr_t Val = 0;
+		CGLSetParameter((CGLContextObj)[OpenGLContext CGLContextObj], kCGLCPClientStorage, (GLint*)&Val);
+		
+		if (CurrentContext == OpenGLContext)
+		{
+			CurrentContext = nil;
+		}
+		
+		[OpenGLContext release];
+	}
+	
+	if ( bMadeCurrent || CurrentContext == nil )
+	{
+		if ( CurrentContext && CurrentContext != OpenGLContext )
+		{
+			[CurrentContext makeCurrentContext];
+		}
+		else
+		{
+			[NSOpenGLContext clearCurrentContext];
+		}
+	}
+}
+	
+void FPlatformOpenGLContext::Initialise( NSOpenGLContext* const SharedContext )
+{
+	SCOPED_AUTORELEASE_POOL;
+	
+	if ( !OpenGLContext )
+	{
+		OpenGLContext = CreateContext(SharedContext);
+	}
+	check(OpenGLContext);
+	
+	MakeCurrent();
+	
+	// Initialise default stuff.
+	extern void InitDebugContext();
+	InitDebugContext();
+	
+	InitDefaultGLContextState();
+	
+	if ( !VertexArrayObject )
+	{
+		glGenVertexArrays(1,&VertexArrayObject);
+		glBindVertexArray(VertexArrayObject);
+	}
+	
+	// Bind our context pointer into the CGL user pointer field
+	intptr_t ContextPtr = (intptr_t)this;
+	CGLSetParameter((CGLContextObj)[OpenGLContext CGLContextObj], kCGLCPClientStorage, (GLint*)&ContextPtr);
+	
+	CGLGetVirtualScreen((CGLContextObj)[OpenGLContext CGLContextObj], &VirtualScreen);
+	
+	EmulatedQueries = new FMacOpenGLQueryEmu(this);
+}
+
+void FPlatformOpenGLContext::Reset(void)
+{
+	SCOPED_AUTORELEASE_POOL;
+	
+	NSOpenGLContext* CurrentContext = [NSOpenGLContext currentContext];
+	bool bMadeCurrent = false;
+	
+	if ((ViewportFramebuffer && !CurrentContext) || GIsEmulatingTimestamp)
+	{
+		MakeCurrent();
+		bMadeCurrent = true;
+	}
+	
+	if (ViewportFramebuffer)
+	{
+		// this can be done from any context, as long as it's not nil.
+		glDeleteFramebuffers(1, &ViewportFramebuffer);
+		ViewportFramebuffer = 0;
+	}
+	
+	[OpenGLContext clearDrawable];
+	
+	if ( OpenGLView )
+	{
+		[OpenGLView release];
+		OpenGLView = nil;
+	}
+	
+	WindowHandle = nil;
+	
+	delete EmulatedQueries;
+	EmulatedQueries = nullptr;
+	
+	if ( bMadeCurrent )
+	{
+		if ( CurrentContext )
+		{
+			[CurrentContext makeCurrentContext];
+		}
+		else
+		{
+			[NSOpenGLContext clearCurrentContext];
+		}
+	}
+}
+
+void FPlatformOpenGLContext::MakeCurrent(void)
+{
+	SCOPED_AUTORELEASE_POOL;
+	
+	check(OpenGLContext);
+	if ( [NSOpenGLContext currentContext] )
+	{
+		glFlushRenderAPPLE();
+	}
+	
+	// Make the context current
+	[OpenGLContext makeCurrentContext];
+	
+	// Ensure that the OpenGL device state hasn't changed.
+	VerifyCurrentContext();
+}
+
+void FPlatformOpenGLContext::VerifyCurrentContext()
+{
+	CGLContextObj Current = CGLGetCurrentContext();
+	FPlatformOpenGLContext* PlatformContext = nullptr;
+	if(Current && (CGLGetParameter(Current, kCGLCPClientStorage, (GLint *)&PlatformContext) == kCGLNoError) && PlatformContext)
+	{
+		// Get the current renderer ID
+		GLint RendererID = 0;
+		CGLError Error = CGLGetParameter(Current, kCGLCPCurrentRendererID, &RendererID);
+		check(Error == kCGLNoError && RendererID != 0);
+		
+		// When using automatic graphics switching each context much switch to use the same GPU, so that we idle the other GPU.
+		// Or when using an explicit renderer for automatic or multiple GPU setups we want to prevent the OS switching it for us.
+		// We have to do this ourselves because we have offscreen contexts which OS X is quite happy to leave running on an offline GPU, but that's bad for laptops.
+		int32 const RendererToMatch = GMacExplicitRendererID ? GMacExplicitRendererID : GMacCurrentRendererID;
+		bool const bGPUChange = (GMacUseAutomaticGraphicsSwitching || GMacUseMultipleGPUs && GMacExplicitRendererID);
+		if(bGPUChange && RendererID != RendererToMatch)
+		{
+			CGLPixelFormatObj PixelFormat = CGLGetPixelFormat(Current);
+			GLint NumVirtualScreens = 0;
+			CGLDescribePixelFormat(PixelFormat, 0, kCGLPFAVirtualScreenCount, &NumVirtualScreens);
+			if(NumVirtualScreens > 1)
+			{
+				for(GLint i = 0; i < NumVirtualScreens; i++)
+				{
+					CGLDescribePixelFormat(PixelFormat, i, kCGLPFARendererID, &RendererID);
+					if(RendererID == RendererToMatch)
+					{
+						CGLSetVirtualScreen(Current, i);
+						break;
+					}
+				}
+			}
+		}
+		
+		// Get the current virtual screen
+		GLint VirtualScreen = 0;
+		CGLGetVirtualScreen(Current, &VirtualScreen);
+		
+		// Verify that the renderer ID hasn't changed, as if it has our capabilities may also be different
+		if(RendererID != PlatformContext->RendererID || VirtualScreen != PlatformContext->VirtualScreen)
+		{
+			PlatformContext->RendererID = RendererID;
+			PlatformContext->VirtualScreen = VirtualScreen;
+			
+			// Get the Vendor ID by parsing the renderer string
+			FString VendorName( ANSI_TO_TCHAR((const ANSICHAR*)glGetString(GL_VENDOR)));
+			if (VendorName.Contains(TEXT("Intel ")))
+			{
+				PlatformContext->VendorID = 0x8086;
+			}
+			else if (VendorName.Contains(TEXT("NVIDIA ")))
+			{
+				PlatformContext->VendorID = 0x10DE;
+			}
+			else if (VendorName.Contains(TEXT("ATi ")) || VendorName.Contains(TEXT("AMD ")))
+			{
+				PlatformContext->VendorID = 0x1002;
+			}
+			
+			// Renderer IDs matchup to driver kexts, so switching based on them will allow us to target workarouds to many GPUs
+			// which exhibit the same unfortunate driver bugs without having to parse their individual ID strings.
+			switch((PlatformContext->RendererID & kCGLRendererIDMatchingMask))
+			{
+				case kCGLRendererATIRadeonX4000ID:
+				{
+					// @todo: remove once AMD fix the AMDX4000 driver for GCN cards so that it is possible to sample the depth while also stencil testing to the same DEPTH_STENCIL texture - it works on all other cards we have.
+					PlatformContext->SupportsDepthFetchDuringDepthTest = false;
+					break;
+				}
+				default:
+				{
+					PlatformContext->SupportsDepthFetchDuringDepthTest = true;
+				}
+			}
+		}
+		
+		GRHIVendorId = PlatformContext->VendorID;
+		GSupportsDepthFetchDuringDepthTest = PlatformContext->SupportsDepthFetchDuringDepthTest;
+	}
+}
+
+void FPlatformOpenGLContext::RegisterGraphicsSwitchingCallback(void)
+{
+#if WITH_SLI
+	GNumActiveGPUsForRendering = 1;
+#endif
+	
+	GConfig->GetBool(MAC_OPENGL_SETTINGS, TEXT("bAllowAutomaticGraphicsSwitching"), GMacUseAutomaticGraphicsSwitching, MAC_OPENGL_INI);
+	GConfig->GetBool(MAC_OPENGL_SETTINGS, TEXT("bUseMultipleRenderers"), GMacUseMultipleGPUs, MAC_OPENGL_INI);
+	
+	// Graphics switching requires a laptop, while multi-GPU requires a desktop.
+	// Check the power sources for info about this Mac to determine which it is.
+	CFTypeRef PowerSourcesInfo = IOPSCopyPowerSourcesInfo();
+	if (PowerSourcesInfo)
+	{
+		CFArrayRef PowerSourcesArray = IOPSCopyPowerSourcesList(PowerSourcesInfo);
+		for (CFIndex Index = 0; Index < CFArrayGetCount(PowerSourcesArray); Index++)
+		{
+			CFTypeRef PowerSource = CFArrayGetValueAtIndex(PowerSourcesArray, Index);
+			NSDictionary* Description = (NSDictionary*)IOPSGetPowerSourceDescription(PowerSourcesInfo, PowerSource);
+			if ([(NSString*)[Description objectForKey: @kIOPSTypeKey] isEqualToString: @kIOPSInternalBatteryType])
+			{
+				// This is a laptop
+				GMacUseAutomaticGraphicsSwitching &= true;
+				GMacUseMultipleGPUs = false;
+				break;
+			}
+		}
+		CFRelease(PowerSourcesArray);
+		CFRelease(PowerSourcesInfo);
+	}
+	
+	if(GMacUseMultipleGPUs || GMacUseAutomaticGraphicsSwitching)
+	{
+		GLint NumberOfRenderers = 0;
+		CGLRendererInfoObj RendererInfo;
+		CGLQueryRendererInfo(0xffffffff, &RendererInfo, &NumberOfRenderers);
+		if(RendererInfo)
+		{
+			TArray<uint32> HardwareRenderers;
+			for(uint32 i = 0; i < NumberOfRenderers; i++)
+			{
+				GLint bAccelerated = 0;
+				CGLDescribeRenderer(RendererInfo, i, kCGLRPAccelerated, &bAccelerated);
+				if(bAccelerated)
+				{
+					HardwareRenderers.Push(i);
+				}
+			}
+			
+			bool bRenderersMatch = true;
+			GLint MajorGLVersion = 0;
+			CGLDescribeRenderer(RendererInfo, 0, kCGLRPMajorGLVersion, &MajorGLVersion);
+			for(uint32 i = 1; i < HardwareRenderers.Num(); i++)
+			{
+				GLint OtherMajorGLVersion = 0;
+				CGLDescribeRenderer(RendererInfo, i, kCGLRPMajorGLVersion, &OtherMajorGLVersion);
+				bRenderersMatch &= MajorGLVersion == OtherMajorGLVersion;
+			}
+			
+			CGLDestroyRendererInfo(RendererInfo);
+			
+			// Can only automatically switch when the renderers use the same major GL version, or errors will occur.
+			GMacUseAutomaticGraphicsSwitching = bRenderersMatch;
+			
+#if WITH_SLI
+			if(GMacUseMultipleGPUs)
+			{
+				GNumActiveGPUsForRendering = !bRenderersMatch ? 1 : HardwareRenderers.Num();
+			}
+#endif
+		}
+	}
+	
+	CGDisplayRegisterReconfigurationCallback(&MacOpenGLContextReconfigurationCallBack, nullptr);
+}
+
+void FPlatformOpenGLContext::UnregisterGraphicsSwitchingCallback(void)
+{
+	CGDisplayRemoveReconfigurationCallback(&MacOpenGLContextReconfigurationCallBack, nullptr);
+}
+
diff --git a/Engine/Source/Runtime/OpenGLDrv/Private/Mac/MacOpenGLContext.h b/Engine/Source/Runtime/OpenGLDrv/Private/Mac/MacOpenGLContext.h
new file mode 100644
--- /dev/null
+++ b/Engine/Source/Runtime/OpenGLDrv/Private/Mac/MacOpenGLContext.h
@@ -0,0 +1,55 @@
+// Copyright 1998-2014 Epic Games, Inc. All Rights Reserved.
+
+#pragma once
+
+@class FCocoaOpenGLView;
+class FMacOpenGLQueryEmu;
+@class NSOpenGLContext;
+@class NSWindow;
+
+struct FPlatformOpenGLContext
+{
+	FPlatformOpenGLContext();
+	~FPlatformOpenGLContext();
+	
+	void Initialise( NSOpenGLContext* const SharedContext );
+	void Reset(void);
+	void MakeCurrent(void);
+	
+	static void VerifyCurrentContext(void);
+	static void RegisterGraphicsSwitchingCallback(void);
+	static void UnregisterGraphicsSwitchingCallback(void);
+	
+	/** Underlying OpenGL context. */
+	NSOpenGLContext*	OpenGLContext;
+
+	// Window Management
+	/** Currently bound Cocoa view - nil for offscreen only contexts */
+	FCocoaOpenGLView*	OpenGLView;
+	/** Owning window - nil for offscreen only contexts */
+	NSWindow*			WindowHandle;
+	
+	// Mac OpenGL Timestamp Query emulation
+	/** Emulated queries */
+	FMacOpenGLQueryEmu*	EmulatedQueries;
+	
+	// OpenGL Resources
+	/** Required vertex array object. */
+	GLuint				VertexArrayObject;
+	/** Output Framebuffer */
+	GLuint				ViewportFramebuffer;
+	
+	// Context Settings
+	/** Flush Synchronisation Interval (V-Sync) */
+	int32				SyncInterval;
+	/** Virtual Screen */
+	GLint				VirtualScreen;
+	
+	// Context Renderer Details
+	/** Device Vendor ID */
+	uint32				VendorID;
+	/** Renderer ID */
+	GLint				RendererID;
+	/** Whether the renderer supports depth-fetch during depth-test */
+	bool				SupportsDepthFetchDuringDepthTest;
+};
diff --git a/Engine/Source/Runtime/OpenGLDrv/Private/Mac/MacOpenGLQuery.cpp b/Engine/Source/Runtime/OpenGLDrv/Private/Mac/MacOpenGLQuery.cpp
new file mode 100644
--- /dev/null
+++ b/Engine/Source/Runtime/OpenGLDrv/Private/Mac/MacOpenGLQuery.cpp
@@ -0,0 +1,355 @@
+// Copyright 1998-2014 Epic Games, Inc. All Rights Reserved.
+
+#include "OpenGLDrvPrivate.h"
+
+#include "MacOpenGLQuery.h"
+
+/*------------------------------------------------------------------------------
+ OpenGL static variables.
+ ------------------------------------------------------------------------------*/
+
+bool GIsEmulatingTimestamp = false; // @todo: Now crashing on Nvidia cards, but not on AMD...
+
+/*------------------------------------------------------------------------------
+ OpenGL query emulation.
+ ------------------------------------------------------------------------------*/
+
+FMacOpenGLTimer::FMacOpenGLTimer(FPlatformOpenGLContext* InContext)
+: Context(InContext)
+, Name(0)
+, Result(0)
+, Accumulated(0)
+, Available(0)
+, Cached(false)
+, Running(false)
+{
+	check(Context);
+	glGenQueries(1, &Name);
+}
+
+FMacOpenGLTimer::~FMacOpenGLTimer()
+{
+	End();
+	Next.Reset();
+	glDeleteQueries(1, &Name);
+}
+
+void FMacOpenGLTimer::Begin(void)
+{
+	if(!Running)
+	{
+		glBeginQuery(GL_TIME_ELAPSED, Name);
+		Running = true;
+	}
+}
+
+void FMacOpenGLTimer::End(void)
+{
+	if(Running)
+	{
+		Running = false;
+		glEndQuery(GL_TIME_ELAPSED);
+	}
+}
+
+uint64 FMacOpenGLTimer::GetResult(void)
+{
+	check(!Running);
+	CacheResult();
+	return Result;
+}
+
+uint64 FMacOpenGLTimer::GetAccumulatedResult(void)
+{
+	check(!Running);
+	CacheResult();
+	return Accumulated;
+}
+
+int32 FMacOpenGLTimer::GetResultAvailable(void)
+{
+	check(!Running);
+	if(!Available)
+	{
+		glGetQueryObjectuiv(Name, GL_QUERY_RESULT_AVAILABLE, &Available);
+	}
+	return Available;
+}
+
+void FMacOpenGLTimer::CacheResult(void)
+{
+	check(!Running);
+	if(!Cached)
+	{
+		glGetQueryObjectui64v(Name, GL_QUERY_RESULT, &Result);
+		
+		if(Previous.IsValid())
+		{
+			Accumulated += Result;
+			Accumulated += Previous->GetAccumulatedResult();
+			
+			Previous->Next.Reset();
+			Previous.Reset();
+		}
+		
+		Available = true;
+		Cached = true;
+	}
+}
+
+FMacOpenGLQuery::FMacOpenGLQuery(FPlatformOpenGLContext* InContext)
+	: Name(0)
+	, Target(0)
+	, Context(InContext)
+	, Result(0)
+	, Available(0)
+	, Running(0)
+	, Cached(false)
+{
+	check(Context);
+}
+
+FMacOpenGLQuery::~FMacOpenGLQuery()
+{
+	if(Running)
+	{
+		End();
+	}
+	if(Name)
+	{
+		glDeleteQueries(1, &Name);
+	}
+}
+
+void FMacOpenGLQuery::Begin(GLenum InTarget)
+{
+	if(Target == 0 || Target == InTarget)
+	{
+		Target = InTarget;
+		Result = 0;
+		Running = true;
+		Available = false;
+		Cached = false;
+		Start.Reset();
+		Finish.Reset();
+		switch(Target)
+		{
+			case GL_TIMESTAMP:
+			{
+				// There can be a lot of timestamps, clear emulated queries out to avoid problems.
+				for( auto It : Context->EmulatedQueries->Queries )
+				{
+					if(It.IsValid() && It->Running == false && (It->Target == GL_TIMESTAMP || It->Target == GL_TIME_ELAPSED))
+					{
+						if(It->GetResultAvailable())
+						{
+							It->GetResult();
+						}
+						else
+						{
+							break;
+						}
+					}
+				}
+				
+				check(Context->EmulatedQueries->LastTimer.IsValid());
+				Context->EmulatedQueries->LastTimer->End();
+				Start = Context->EmulatedQueries->LastTimer;
+				
+				TSharedPtr<FMacOpenGLTimer, ESPMode::Fast> Current(new FMacOpenGLTimer(Context));
+				Current->Previous = Start;
+				Start->Next = Current;
+				Current->Begin();
+				Context->EmulatedQueries->LastTimer = Current;
+				
+				break;
+			}
+			case GL_TIME_ELAPSED:
+			{
+				check(Context->EmulatedQueries->LastTimer.IsValid());
+				Context->EmulatedQueries->LastTimer->End();
+				
+				Start = TSharedPtr<FMacOpenGLTimer, ESPMode::Fast>(new FMacOpenGLTimer(Context));
+				Start->Previous = Context->EmulatedQueries->LastTimer;
+				Context->EmulatedQueries->LastTimer->Next = Start;
+				Start->Begin();
+				
+				Context->EmulatedQueries->LastTimer = Start;
+				break;
+			}
+			default:
+			{
+				if(!Name)
+				{
+					glGenQueries(1, &Name);
+				}
+				glBeginQuery(Target, Name);
+				break;
+			}
+		}
+	}
+}
+
+void FMacOpenGLQuery::End()
+{
+	Running = false;
+	if(Target == GL_TIME_ELAPSED)
+	{
+		Finish = Context->EmulatedQueries->LastTimer;
+		TSharedPtr<FMacOpenGLTimer, ESPMode::Fast> Current(new FMacOpenGLTimer(Context));
+		Current->Previous = Finish;
+		Finish->Next = Current;
+		Finish->End();
+		Current->Begin();
+		Context->EmulatedQueries->LastTimer = Current;
+		
+		if (Start == Finish)
+		{
+			Start = Start->Previous;
+		}
+	}
+	else if(Target != GL_TIMESTAMP)
+	{
+		glEndQuery(Target);
+	}
+}
+
+uint64 FMacOpenGLQuery::GetResult(void)
+{
+	if(!Cached)
+	{
+		switch(Target)
+		{
+			case GL_TIMESTAMP:
+				Result = Start->GetAccumulatedResult();
+				Start.Reset();
+				Finish.Reset();
+				break;
+			case GL_TIME_ELAPSED:
+				Result = Finish->GetAccumulatedResult() - Start->GetAccumulatedResult();
+				Start.Reset();
+				Finish.Reset();
+				break;
+			default:
+				glGetQueryObjectui64v(Name, GL_QUERY_RESULT, &Result);
+				break;
+		}
+		Available = true;
+		Cached = true;
+	}
+	return Result;
+}
+
+int32 FMacOpenGLQuery::GetResultAvailable(void)
+{
+	if(!Available)
+	{
+		switch(Target)
+		{
+			case GL_TIMESTAMP:
+				Available = Start->GetResultAvailable();
+				break;
+			case GL_TIME_ELAPSED:
+				Available = Finish->GetResultAvailable();
+				break;
+			default:
+				glGetQueryObjectuiv(Name, GL_QUERY_RESULT_AVAILABLE, &Available);
+				break;
+		}
+	}
+	return Available;
+}
+
+FMacOpenGLQueryEmu::FMacOpenGLQueryEmu(FPlatformOpenGLContext* InContext)
+: PlatformContext(InContext)
+{
+	check(PlatformContext);
+	
+	// Only use the timestamp emulation in a non-shipping build - end-users shouldn't care about this profiling feature.
+#if (!UE_BUILD_SHIPPING)
+	// Opt-in since it crashes Nvidia cards at the moment
+	GIsEmulatingTimestamp = FParse::Param(FCommandLine::Get(), TEXT("EnableMacGPUTimestamp"));
+#endif
+	
+	if ( GIsEmulatingTimestamp )
+	{
+		LastTimer = TSharedPtr<FMacOpenGLTimer, ESPMode::Fast>(new FMacOpenGLTimer(InContext));
+		LastTimer->Begin();
+	}
+}
+
+FMacOpenGLQueryEmu::~FMacOpenGLQueryEmu()
+{
+	if ( GIsEmulatingTimestamp && LastTimer.IsValid() )
+	{
+		LastTimer->End();
+	}
+}
+	
+bool FMacOpenGLQueryEmu::IsQuery(GLuint const QueryID) const
+{
+	return QueryID && Queries.Num() >= QueryID && Queries[QueryID-1].IsValid();
+}
+
+TSharedPtr<FMacOpenGLQuery, ESPMode::Fast> FMacOpenGLQueryEmu::GetQuery(GLuint const QueryID) const
+{
+	TSharedPtr<FMacOpenGLQuery, ESPMode::Fast> Query;
+	if ( IsQuery(QueryID) )
+	{
+		Query = Queries[QueryID-1];
+	}
+	return Query;
+}
+
+GLuint FMacOpenGLQueryEmu::CreateQuery(void)
+{
+	GLuint Result = 0;
+	if ( Queries.Num() < UINT_MAX )
+	{
+		TSharedPtr<FMacOpenGLQuery, ESPMode::Fast> Query(new FMacOpenGLQuery(PlatformContext));
+		for(int32 Index = 0; Index < Queries.Num(); ++Index)
+		{
+			if(!Queries[Index].IsValid())
+			{
+				Queries[Index] = Query;
+				Result = Index + 1;
+				break;
+			}
+		}
+		
+		if(!Result)
+		{
+			Queries.Add(Query);
+			Result = Queries.Num();
+		}
+	}
+	return Result;
+}
+
+void FMacOpenGLQueryEmu::DeleteQuery(GLuint const QueryID)
+{
+	if ( IsQuery(QueryID) )
+	{
+		Queries[QueryID - 1].Reset();
+	}
+}
+
+void FMacOpenGLQueryEmu::BeginQuery(GLenum const QueryType, GLuint const QueryID)
+{
+	TSharedPtr<FMacOpenGLQuery, ESPMode::Fast> Query = GetQuery( QueryID );
+	if ( Query.IsValid() )
+	{
+		Query->Begin(QueryType);
+		RunningQueries.Add(QueryType, Query);
+	}
+}
+
+void FMacOpenGLQueryEmu::EndQuery(GLenum const QueryType)
+{
+	TSharedPtr<FMacOpenGLQuery> Query = RunningQueries.FindOrAdd(QueryType);
+	if ( Query.IsValid() )
+	{
+		Query->End();
+		RunningQueries[QueryType].Reset();
+	}
+}
diff --git a/Engine/Source/Runtime/OpenGLDrv/Private/Mac/MacOpenGLQuery.h b/Engine/Source/Runtime/OpenGLDrv/Private/Mac/MacOpenGLQuery.h
new file mode 100644
--- /dev/null
+++ b/Engine/Source/Runtime/OpenGLDrv/Private/Mac/MacOpenGLQuery.h
@@ -0,0 +1,82 @@
+// Copyright 1998-2014 Epic Games, Inc. All Rights Reserved.
+
+#pragma once
+
+#include "SharedPointer.h"
+#include "MacOpenGLContext.h"
+
+class FMacOpenGLTimer : public TSharedFromThis<FMacOpenGLTimer, ESPMode::Fast>
+{
+public:
+	FMacOpenGLTimer(FPlatformOpenGLContext* InContext);
+	virtual ~FMacOpenGLTimer();
+	
+	void Begin(void);
+	void End(void);
+	uint64 GetResult(void);
+	uint64 GetAccumulatedResult(void);
+	int32 GetResultAvailable(void);
+	
+public:
+	TSharedPtr<FMacOpenGLTimer, ESPMode::Fast> Next;
+	TSharedPtr<FMacOpenGLTimer, ESPMode::Fast> Previous;
+	
+private:
+	void CacheResult(void);
+	
+private:
+	FPlatformOpenGLContext* Context;
+	uint32 Name;
+	uint64 Result;
+	uint64 Accumulated;
+	uint32 Available;
+	bool Cached;
+	bool Running;
+};
+
+class FMacOpenGLQuery : public TSharedFromThis<FMacOpenGLQuery, ESPMode::Fast>
+{
+public:
+	FMacOpenGLQuery(FPlatformOpenGLContext* InContext);
+	virtual ~FMacOpenGLQuery();
+	
+	void Begin(GLenum InTarget);
+	void End();
+	uint64 GetResult(void);
+	int32 GetResultAvailable(void);
+	
+	uint32 Name;
+	GLenum Target;
+	
+private:
+	FPlatformOpenGLContext* Context;
+	uint64 Result;
+	uint32 Available;
+	bool Running;
+	bool Cached;
+	TSharedPtr<FMacOpenGLTimer, ESPMode::Fast> Start;
+	TSharedPtr<FMacOpenGLTimer, ESPMode::Fast> Finish;
+};
+
+class FMacOpenGLQueryEmu
+{
+	friend class FMacOpenGLQuery;
+public:
+	FMacOpenGLQueryEmu(FPlatformOpenGLContext* InContext);
+	~FMacOpenGLQueryEmu();
+	
+	bool IsQuery(GLuint const QueryID) const;
+	TSharedPtr<FMacOpenGLQuery, ESPMode::Fast> GetQuery(GLuint const QueryID) const;
+	GLuint CreateQuery(void);
+	void DeleteQuery(GLuint const QueryID);
+	void BeginQuery(GLenum const QueryType, GLuint const QueryID);
+	void EndQuery(GLenum const QueryType);
+	
+private:
+	FPlatformOpenGLContext* PlatformContext;
+	TArray<TSharedPtr<FMacOpenGLQuery, ESPMode::Fast>> Queries;
+	TMap<GLenum, TSharedPtr<FMacOpenGLQuery, ESPMode::Fast>> RunningQueries;
+	TSharedPtr<FMacOpenGLTimer, ESPMode::Fast> LastTimer;
+};
+
+extern bool GIsEmulatingTimestamp;
diff --git a/Engine/Source/Runtime/OpenGLDrv/Private/Mac/OpenGLMac.cpp b/Engine/Source/Runtime/OpenGLDrv/Private/Mac/OpenGLMac.cpp
--- a/Engine/Source/Runtime/OpenGLDrv/Private/Mac/OpenGLMac.cpp
+++ b/Engine/Source/Runtime/OpenGLDrv/Private/Mac/OpenGLMac.cpp
@@ -2,11 +2,15 @@
 
 #include "OpenGLDrvPrivate.h"
 
+#include "MacOpenGLContext.h"
+#include "MacOpenGLQuery.h"
+
 #include "MacWindow.h"
 #include "MacTextInputMethodSystem.h"
 #include "CocoaTextView.h"
 #include "CocoaThread.h"
 #include "CocoaWindow.h"
+#include <IOKit/IOKitLib.h>
 
 /*------------------------------------------------------------------------------
  OpenGL static variables.
@@ -21,15 +25,15 @@ static TAutoConsoleVariable<int32> CVarMacUseFrameBufferSRGB(
 	);
 
 // @todo: remove once Apple fixes radr://15553950, TTP# 315197
-static int32 GMacMustFlushTexStorage = 0;
-static FAutoConsoleVariableRef CVarMacMustFlushTexStorage(
-	TEXT("r.Mac.MustFlushTexStorage"),
-	GMacMustFlushTexStorage,
-	TEXT("If true, flush the OpenGL command stream after calls to glTexStorage* to avoid driver errors, do nothing if false (faster, the default)."),
+static int32 GMacFlushTexStorage = true;
+static FAutoConsoleVariableRef CVarMacFlushTexStorage(
+	TEXT("r.Mac.FlushTexStorage"),
+	GMacFlushTexStorage,
+	TEXT("When true, for GPUs that require it the OpenGL command stream will be flushed after glTexStorage* to avoid driver errors, when false no flush will be inserted which is faster but may cause errors on some drivers. (Default: True)"),
 	ECVF_RenderThreadSafe
 	);
 
-static int32 GMacUseMTGL = 1;
+int32 GMacUseMTGL = 1;
 static FAutoConsoleVariableRef CVarMacUseMTGL(
 	TEXT("r.Mac.UseMTGL"),
 	GMacUseMTGL,
@@ -66,11 +70,11 @@ static TArray<GLuint> GMacTexturesToDelete[GMacTexturePoolNum];
 static TArray<GLuint> GMacBuffers;
 static TArray<GLuint> GMacBuffersToDelete;
 
-bool GIsRunningOnIntelCard = false; // @todo: remove once Apple fixes radr://16223045 Changes to the GL separate blend state aren't always respected on Intel cards
-static bool GIsEmulatingTimestamp = false; // @todo: Now crashing on Nvidia cards, but not on AMD...
+/** Used to temporarily disable Cocoa screen updates to make window updates happen only on the render thread. */
+bool GMacEnableCocoaScreenUpdates = true;
 
 /*------------------------------------------------------------------------------
-	OpenGL context management.
+	OpenGL context blit shaders.
 ------------------------------------------------------------------------------*/
 
 char const* const CompositedBlitVertexShader = "#version 150\n"
@@ -108,63 +112,9 @@ char const* const CompositedBlitFragmentShader = "#version 150\n"
 "	Color = vec4(WindowColor.x, WindowColor.y, WindowColor.z, WindowColor.x);\n"
 "}\n";
 
-const NSOpenGLPixelFormatAttribute ContextCreationAttributes[] =
-{
-	kCGLPFAOpenGLProfile,
-	kCGLOGLPVersion_3_2_Core,
-	kCGLPFAAccelerated,
-	kCGLPFANoRecovery,
-	kCGLPFADoubleBuffer,
-	kCGLPFAColorSize,
-	32,
-	0
-};
-
-static NSOpenGLPixelFormat* PixelFormat = nil;
-
-static NSOpenGLContext* CreateContext( NSOpenGLContext* SharedContext )
-{
-	SCOPED_AUTORELEASE_POOL;
-
-	if (!PixelFormat)
-	{
-		int32 DisplayMask = 0;
-		if (GConfig->GetInt(TEXT("OpenGL"), TEXT("OverrideDisplayMask"), DisplayMask, GEngineIni))
-		{
-			TArray<NSOpenGLPixelFormatAttribute> Attributes;
-			
-			Attributes.Append(ContextCreationAttributes, ARRAY_COUNT(ContextCreationAttributes) - 1);
-			Attributes.Add(kCGLPFAAllowOfflineRenderers);
-			Attributes.Add(kCGLPFADisplayMask);
-			Attributes.Add(DisplayMask);
-			Attributes.Add(0);
-			
-			PixelFormat = [[NSOpenGLPixelFormat alloc] initWithAttributes: Attributes.GetData()];
-		}
-		else
-		{
-			PixelFormat = [[NSOpenGLPixelFormat alloc] initWithAttributes: ContextCreationAttributes];
-		}
-		check(PixelFormat);
-	}
-
-	NSOpenGLContext* Context = [[NSOpenGLContext alloc] initWithFormat: PixelFormat shareContext: SharedContext];
-	check(Context);
-	int32 SyncInterval = 0;
-	[Context setValues: &SyncInterval forParameter: NSOpenGLCPSwapInterval];
-	// Setup Opacity - can't change it dynamically later!
-	int32 SurfaceOpacity = 0;
-	[Context setValues: &SurfaceOpacity forParameter: NSOpenGLCPSurfaceOpacity];
-
-	if (GMacUseMTGL)
-	{
-		CGLEnable((CGLContextObj)[Context CGLContextObj], kCGLCEMPEngine);
-		
-		// @todo: We should disable OpenGL.UseMapBuffer for improved performance under MTGL, but radr://17662549 prevents it when using GL_TEXTURE_BUFFER's for skinning.
-	}
-
-	return Context;
-}
+/*------------------------------------------------------------------------------
+ OpenGL context management.
+ ------------------------------------------------------------------------------*/
 
 class FScopeContext
 {
@@ -183,6 +133,7 @@ public:
 				glFlushRenderAPPLE();
 			}
 			[Context makeCurrentContext];
+			FPlatformOpenGLContext::VerifyCurrentContext();
 		}
 	}
 
@@ -196,6 +147,7 @@ public:
 			if (PreviousContext)
 			{
 				[PreviousContext makeCurrentContext];
+				FPlatformOpenGLContext::VerifyCurrentContext();
 			}
 			else
 			{
@@ -211,8 +163,9 @@ private:
 
 void DeleteQueriesForCurrentContext( NSOpenGLContext* Context );
 
-/** Used to temporarily disable Cocoa screen updates to make window updates happen only on the render thread. */
-bool GMacEnableCocoaScreenUpdates = true;
+/*------------------------------------------------------------------------------
+ OpenGL view.
+ ------------------------------------------------------------------------------*/
 
 @interface NSView(NSThemeFramePrivate)
 - (float)roundedCornerRadius;
@@ -277,79 +230,6 @@ bool GMacEnableCocoaScreenUpdates = true;
 
 @end
 
-class FMacOpenGLTimer
-{
-public:
-	FMacOpenGLTimer(FPlatformOpenGLContext* InContext);
-	~FMacOpenGLTimer();
-	
-	void Begin(void);
-	void End(void);
-	uint64 GetResult(void);
-	uint64 GetAccumulatedResult(void);
-	int32 GetResultAvailable(void);
-	
-public:
-	TSharedPtr<FMacOpenGLTimer> Next;
-	TSharedPtr<FMacOpenGLTimer> Previous;
-	
-private:
-	void CacheResult(void);
-	
-private:
-	FPlatformOpenGLContext* Context;
-	uint32 Name;
-	uint64 Result;
-	uint64 Accumulated;
-	uint32 Available;
-	bool Cached;
-	bool Running;
-};
-
-class FMacOpenGLQuery : public TSharedFromThis<FMacOpenGLQuery, ESPMode::Fast>
-{
-public:
-	FMacOpenGLQuery(FPlatformOpenGLContext* InContext);
-	~FMacOpenGLQuery();
-	
-	void Begin(GLenum InTarget);
-	void End();
-	uint64 GetResult(void);
-	int32 GetResultAvailable(void);
-	
-	uint32 Name;
-	GLenum Target;
-	
-private:
-	FPlatformOpenGLContext* Context;
-	uint64 Result;
-	uint32 Available;
-	bool Running;
-	bool Cached;
-	TSharedPtr<FMacOpenGLTimer> Start;
-	TSharedPtr<FMacOpenGLTimer> Finish;
-};
-
-/** Platform specific OpenGL context. */
-struct FPlatformOpenGLContext
-{
-	NSWindow*			WindowHandle;
-	FCocoaOpenGLView*	OpenGLView;
-	NSOpenGLContext*	OpenGLContext;
-	GLuint				ViewportFramebuffer;
-	int32				SyncInterval;
-	GLuint				VertexArrayObject;	// one has to be generated and set for each context (OpenGL 3.2 Core requirements)
-	TArray<TSharedPtr<FMacOpenGLQuery>> Queries;
-	TMap<GLenum, TSharedPtr<FMacOpenGLQuery>> RunningQueries;
-	TSharedPtr<FMacOpenGLTimer> LastTimer;
-};
-
-struct OpenGLContextInfo
-{
-	NSOpenGLContext*	Context;
-	GLuint				VertexArrayObject;
-};
-
 extern void OnQueryInvalidation( void );
 
 struct FPlatformOpenGLDevice
@@ -367,7 +247,7 @@ struct FPlatformOpenGLDevice
 	FPlatformOpenGLContext	RenderingContext;
 
 
-	TArray<OpenGLContextInfo>	FreeOpenGLContexts;
+	TArray<FPlatformOpenGLContext*>	FreeOpenGLContexts;
 	int32						NumUsedContexts;
 
 	/** Guards against operating on viewport contexts from more than one thread at the same time. */
@@ -378,45 +258,14 @@ struct FPlatformOpenGLDevice
 	{
 		SCOPED_AUTORELEASE_POOL;
 		
-		extern void InitDebugContext();
-		ContextUsageGuard = new FCriticalSection;
-		SharedContext.OpenGLContext = CreateContext(NULL);
-		RenderingContext.OpenGLContext = CreateContext(SharedContext.OpenGLContext);
-		
-		[RenderingContext.OpenGLContext makeCurrentContext];
-		InitDebugContext();
-		glGenVertexArrays(1,&RenderingContext.VertexArrayObject);
-		glBindVertexArray(RenderingContext.VertexArrayObject);
-		
-		// Only use the timestamp emulation in a non-shipping build - end-users shouldn't care about this profiling feature.
-#if (!UE_BUILD_SHIPPING)
-		// Opt-in since it crashes Nvidia cards at the moment
-		GIsEmulatingTimestamp = FParse::Param(FCommandLine::Get(), TEXT("EnableMacGPUTimestamp"));
-#endif
-		
-		if(GIsEmulatingTimestamp)
-		{
-			intptr_t RenderPtr = (intptr_t)&RenderingContext;
-			CGLSetParameter((CGLContextObj)[RenderingContext.OpenGLContext CGLContextObj], kCGLCPClientStorage, (GLint*)&RenderPtr);
-			RenderingContext.LastTimer = TSharedPtr<FMacOpenGLTimer>(new FMacOpenGLTimer(&RenderingContext));
-			RenderingContext.LastTimer->Begin();
-		}
-		InitDefaultGLContextState();
+		FPlatformOpenGLContext::RegisterGraphicsSwitchingCallback();
 		
-		glFlushRenderAPPLE();
+		ContextUsageGuard = new FCriticalSection;
+		SharedContext.Initialise(NULL);
 		
-		[SharedContext.OpenGLContext makeCurrentContext];
-		InitDebugContext();
-		glGenVertexArrays(1,&SharedContext.VertexArrayObject);
-		glBindVertexArray(SharedContext.VertexArrayObject);
+		RenderingContext.Initialise(SharedContext.OpenGLContext);
 		
-		if(GIsEmulatingTimestamp)
-		{
-			intptr_t SharedPtr = (intptr_t)&SharedContext;
-			CGLSetParameter((CGLContextObj)[SharedContext.OpenGLContext CGLContextObj], kCGLCPClientStorage, (GLint*)&SharedPtr);
-			SharedContext.LastTimer = TSharedPtr<FMacOpenGLTimer>(new FMacOpenGLTimer(&SharedContext));
-			SharedContext.LastTimer->Begin();
-		}
+		SharedContext.MakeCurrent();
 		
 		SharedContextCompositeVertexShader = glCreateShader(GL_VERTEX_SHADER);
 		glShaderSource(SharedContextCompositeVertexShader, 1, &CompositedBlitVertexShader, NULL);
@@ -438,8 +287,6 @@ struct FPlatformOpenGLDevice
 		glGenTextures(1, &SharedContextCompositeTexture);
 		SharedContextCompositeTextureSizeX = 0;
 		SharedContextCompositeTextureSizeY = 0;
-		
-		InitDefaultGLContextState();
 	}
 
 	~FPlatformOpenGLDevice()
@@ -453,42 +300,15 @@ struct FPlatformOpenGLDevice
 			FScopeLock ScopeLock(ContextUsageGuard);
 			for (int32 ContextIndex = 0; ContextIndex < FreeOpenGLContexts.Num(); ++ContextIndex)
 			{
-				{
-					FScopeContext ScopeContext(FreeOpenGLContexts[ContextIndex].Context);
-					DeleteQueriesForCurrentContext(FreeOpenGLContexts[ContextIndex].Context);
-					glBindVertexArray(0);
-					glDeleteVertexArrays(1, &FreeOpenGLContexts[ContextIndex].VertexArrayObject);
-				}
-				[FreeOpenGLContexts[ContextIndex].Context release];
+				delete FreeOpenGLContexts[ContextIndex];
 			}
 			FreeOpenGLContexts.Empty();
 		}
 
 		OnQueryInvalidation();
 		
-		{
-			FScopeContext Context(RenderingContext.OpenGLContext);
-			DeleteQueriesForCurrentContext(RenderingContext.OpenGLContext);
-			glBindVertexArray(0);
-			glDeleteVertexArrays(1,&RenderingContext.VertexArrayObject);
-			
-			if(GIsEmulatingTimestamp)
-			{
-				RenderingContext.RunningQueries.Empty();
-				RenderingContext.Queries.Empty();
-				RenderingContext.LastTimer.Reset();
-			}
-			
-			// Unbind the platform context from the CGL context
-			intptr_t Val = 0;
-			CGLSetParameter((CGLContextObj)[RenderingContext.OpenGLContext CGLContextObj], kCGLCPClientStorage, (GLint*)&Val);
-		}
-		
 		{
 			FScopeContext Context(SharedContext.OpenGLContext);
-			DeleteQueriesForCurrentContext(SharedContext.OpenGLContext);
-			glBindVertexArray(0);
-			glDeleteVertexArrays(1,&SharedContext.VertexArrayObject);
 			
 			glDeleteProgram(SharedContextCompositeProgram);
 			glDeleteShader(SharedContextCompositeVertexShader);
@@ -510,21 +330,9 @@ struct FPlatformOpenGLDevice
 			
 			glDeleteBuffers(GMacBuffersToDelete.Num(), GMacBuffersToDelete.GetData());
 			GMacBuffersToDelete.Reset();
-			
-			if(GIsEmulatingTimestamp)
-			{
-				SharedContext.RunningQueries.Empty();
-				SharedContext.Queries.Empty();
-				SharedContext.LastTimer.Reset();
-			}
-			
-			// Unbind the platform context from the CGL context
-			intptr_t Val = 0;
-			CGLSetParameter((CGLContextObj)[SharedContext.OpenGLContext CGLContextObj], kCGLCPClientStorage, (GLint*)&Val);
 		}
 		
-		[RenderingContext.OpenGLContext release];
-		[SharedContext.OpenGLContext release];
+		FPlatformOpenGLContext::UnregisterGraphicsSwitchingCallback();
 		
 		delete ContextUsageGuard;
 	}
@@ -546,37 +354,42 @@ FPlatformOpenGLContext* PlatformCreateOpenGLContext(FPlatformOpenGLDevice* Devic
 
 	check(InWindowHandle);
 
-	FPlatformOpenGLContext* Context = new FPlatformOpenGLContext;
-
-	Context->ViewportFramebuffer = 0;	// will be created on demand
-	Context->SyncInterval = 0;
-	Context->WindowHandle = (NSWindow*)InWindowHandle;
-
-	NSRect ContentRect = [[Context->WindowHandle contentView] frame];
-	Context->OpenGLView = [[FCocoaOpenGLView alloc] initWithFrame: ContentRect];
-
+	FPlatformOpenGLContext* Context = nullptr;
 	{
 		FScopeLock ScopeLock(Device->ContextUsageGuard);
 		if( Device->FreeOpenGLContexts.Num() )
 		{
-			OpenGLContextInfo Info = Device->FreeOpenGLContexts.Pop();
-			Context->OpenGLContext = Info.Context;
-			Context->VertexArrayObject = Info.VertexArrayObject;
+			Context = Device->FreeOpenGLContexts.Pop();
 		}
 		else
 		{
-			Context->OpenGLContext = CreateContext(Device->SharedContext.OpenGLContext);
-
-			{
-				FScopeContext ScopeContext(Context->OpenGLContext);
-
-				glGenVertexArrays(1, &Context->VertexArrayObject);
-				glBindVertexArray(Context->VertexArrayObject);
-				
-				InitDefaultGLContextState();
-			}
+			Context = new FPlatformOpenGLContext;
+		}
+	}
+	check(Context);
+	{
+		NSOpenGLContext* PreviousContext = [NSOpenGLContext currentContext];
+		
+		Context->Initialise(Device->SharedContext.OpenGLContext);
+		Context->ViewportFramebuffer = 0;	// will be created on demand
+		Context->SyncInterval = 0;
+		Context->WindowHandle = (NSWindow*)InWindowHandle;
+		
+		NSRect ContentRect = [[Context->WindowHandle contentView] frame];
+		Context->OpenGLView = [[FCocoaOpenGLView alloc] initWithFrame: ContentRect];
+		
+		glFlushRenderAPPLE();
+		if (PreviousContext)
+		{
+			[PreviousContext makeCurrentContext];
+			FPlatformOpenGLContext::VerifyCurrentContext();
+		}
+		else
+		{
+			[NSOpenGLContext clearCurrentContext];
 		}
 	}
+	
 	++Device->NumUsedContexts;
 
 	// Attach the view to the window
@@ -601,42 +414,6 @@ FPlatformOpenGLContext* PlatformCreateOpenGLContext(FPlatformOpenGLDevice* Devic
 	// Attach the context to the view if needed
 	[Context->OpenGLContext setView: Context->OpenGLView];
 
-	FString VendorName( ANSI_TO_TCHAR((const ANSICHAR*)glGetString(GL_VENDOR)));
-	if (VendorName.Contains(TEXT("Intel ")))
-	{
-		GIsRunningOnIntelCard = true;
-	}
-	else if (VendorName.Contains(TEXT("NVIDIA ")))
-	{
-		GMacMustFlushTexStorage = true;
-	}
-	
-	// Renderer IDs matchup to driver kexts, so switching based on them will allow us to target workarouds to many GPUs
-	// which exhibit the same unfortunate driver bugs without having to parse their individual ID strings.
-	GLint RendererID = 0;
-	if(CGLGetParameter((CGLContextObj)[Context->OpenGLContext CGLContextObj], kCGLCPCurrentRendererID, &RendererID) == kCGLNoError)
-	{
-		switch((RendererID & kCGLRendererIDMatchingMask))
-		{
-			case kCGLRendererATIRadeonX4000ID:
-			{
-				// @todo: remove once AMD fix the AMDX4000 driver for GCN cards so that it is possible to sample the depth while also stencil testing to the same DEPTH_STENCIL texture - it works on all other cards we have.
-				GSupportsDepthFetchDuringDepthTest = false;
-				break;
-			}
-		}
-    }
-	
-	if(GIsEmulatingTimestamp)
-	{
-		// Bind the platform context into the CGL context
-		CGLSetParameter((CGLContextObj)[Context->OpenGLContext CGLContextObj], kCGLCPClientStorage, (GLint*)&Context);
-		
-		FScopeContext ScopeContext(Context->OpenGLContext);
-		Context->LastTimer = TSharedPtr<FMacOpenGLTimer>(new FMacOpenGLTimer(Context));
-		Context->LastTimer->Begin();
-	}
-	
 	return Context;
 }
 
@@ -645,48 +422,19 @@ void PlatformReleaseOpenGLContext(FPlatformOpenGLDevice* Device, FPlatformOpenGL
 	SCOPED_AUTORELEASE_POOL;
 
 	check(Context && Context->OpenGLView && Context->OpenGLContext);
-	if(GIsEmulatingTimestamp)
-	{
-		FScopeContext ScopeContext(Context->OpenGLContext);
-		Context->RunningQueries.Empty();
-		Context->Queries.Empty();
-		Context->LastTimer.Reset();
-		
-		// Unbind the platform context from the CGL context
-		intptr_t Val = 0;
-		CGLSetParameter((CGLContextObj)[Context->OpenGLContext CGLContextObj], kCGLCPClientStorage, (GLint*)&Val);
-	}
 	{
 		FScopeLock ScopeLock(Device->ContextUsageGuard);
-
-		[Context->OpenGLContext clearDrawable];
-
-		{
-			OpenGLContextInfo Info;
-			Info.Context = Context->OpenGLContext;
-			Info.VertexArrayObject = Context->VertexArrayObject;
-			Device->FreeOpenGLContexts.Add(Info);
-		}
-
-		Context->OpenGLContext = NULL;
-
-		if (Context->ViewportFramebuffer)
-		{
-			check([NSOpenGLContext currentContext]);
-			glDeleteFramebuffers(1,&Context->ViewportFramebuffer);	// this can be done from any context, as long as it's not nil.
-			Context->ViewportFramebuffer = 0;
-		}
+		
+		Context->Reset();
+		Device->FreeOpenGLContexts.Add(Context);
 
 		--Device->NumUsedContexts;
 	}
-	[Context->OpenGLView release];
-	Context->OpenGLView = NULL;
 }
 
 void PlatformDestroyOpenGLContext(FPlatformOpenGLDevice* Device, FPlatformOpenGLContext* Context)
 {
 	PlatformReleaseOpenGLContext(Device, Context);
-	delete Context;
 }
 
 void* PlatformGetWindow(FPlatformOpenGLContext* Context, void** AddParam)
@@ -902,27 +650,13 @@ bool PlatformBlitToViewport( FPlatformOpenGLDevice* Device, const FOpenGLViewpor
 void PlatformRenderingContextSetup(FPlatformOpenGLDevice* Device)
 {
 	check(Device);
-
-	SCOPED_AUTORELEASE_POOL;
-
-	if ([NSOpenGLContext currentContext])
-	{
-		glFlushRenderAPPLE();
-	}
-	[Device->RenderingContext.OpenGLContext makeCurrentContext];
+	Device->RenderingContext.MakeCurrent();
 }
 
 void PlatformSharedContextSetup(FPlatformOpenGLDevice* Device)
 {
 	check(Device);
-
-	SCOPED_AUTORELEASE_POOL;
-
-	if ([NSOpenGLContext currentContext])
-	{
-		glFlushRenderAPPLE();
-	}
-	[Device->SharedContext.OpenGLContext makeCurrentContext];
+	Device->SharedContext.MakeCurrent();
 }
 
 void PlatformNULLContextSetup()
@@ -1234,34 +968,13 @@ void PlatformGetNewRenderQuery( GLuint* OutQuery, uint64* OutQueryContext )
 				{
 					if( ReleasedQueries[Index].Context == Context )
 					{
-						if(PlatformContext->Queries.Num() >= ReleasedQueries[Index].Query)
-						{
-							PlatformContext->Queries[ReleasedQueries[Index].Query - 1].Reset();
-						}
+						PlatformContext->EmulatedQueries->DeleteQuery(ReleasedQueries[Index].Query);
 						ReleasedQueries.RemoveAt(Index);
 						Index--;
 					}
 				}
 				
-				if(PlatformContext->Queries.Num() < UINT_MAX)
-				{
-					TSharedPtr<FMacOpenGLQuery> Query(new FMacOpenGLQuery(PlatformContext));
-					for(int32 Index = 0; Index < PlatformContext->Queries.Num(); ++Index)
-					{
-						if(!PlatformContext->Queries[Index].IsValid())
-						{
-							PlatformContext->Queries[Index] = Query;
-							NewQuery = Index + 1;
-							break;
-						}
-					}
-					
-					if(!NewQuery)
-					{
-						PlatformContext->Queries.Add(Query);
-						NewQuery = PlatformContext->Queries.Num();
-					}
-				}
+				NewQuery = PlatformContext->EmulatedQueries->CreateQuery();
 			}
 		}
 		else
@@ -1302,10 +1015,7 @@ void PlatformReleaseRenderQuery( GLuint Query, uint64 QueryContext )
 			FPlatformOpenGLContext* PlatformContext = nullptr;
 			if(Current && (CGLGetParameter(Current, kCGLCPClientStorage, (GLint *)&PlatformContext) == kCGLNoError) && PlatformContext)
 			{
-				if(PlatformContext->Queries.Num() >= Query)
-				{
-					PlatformContext->Queries[Query - 1].Reset();
-				}
+				PlatformContext->EmulatedQueries->DeleteQuery(Query);
 			}
 		}
 		else
@@ -1344,10 +1054,7 @@ void DeleteQueriesForCurrentContext( NSOpenGLContext* Context )
 			{
 				if( ReleasedQueries[Index].Context == Context )
 				{
-					if(PlatformContext->Queries.Num() >= ReleasedQueries[Index].Query)
-					{
-						PlatformContext->Queries[ReleasedQueries[Index].Query - 1].Reset();
-					}
+					PlatformContext->EmulatedQueries->DeleteQuery(ReleasedQueries[Index].Query);
 					ReleasedQueries.RemoveAtSwap(Index);
 					--Index;
 				}
@@ -1391,7 +1098,6 @@ PFNGLBLENDEQUATIONIARBPROC FMacOpenGL::glBlendEquationi = NULL;
 PFNGLLABELOBJECTEXTPROC FMacOpenGL::glLabelObjectEXT = NULL;
 PFNGLPUSHGROUPMARKEREXTPROC FMacOpenGL::glPushGroupMarkerEXT = NULL;
 PFNGLPOPGROUPMARKEREXTPROC FMacOpenGL::glPopGroupMarkerEXT = NULL;
-bool FMacOpenGL::bSupportsTessellationShader = false;
 PFNGLPATCHPARAMETERIPROC FMacOpenGL::glPatchParameteri = NULL;
 
 uint64 FMacOpenGL::GetVideoMemorySize()
@@ -1491,277 +1197,10 @@ void FMacOpenGL::ProcessExtensions(const FString& ExtensionsString)
 	
 	if(ExtensionsString.Contains(TEXT("GL_ARB_tessellation_shader")))
 	{
-		bSupportsTessellationShader = true;
 		glPatchParameteri = (PFNGLPATCHPARAMETERIPROC)dlsym(RTLD_SELF, "glPatchParameteri");
 	}
 }
 
-FMacOpenGLTimer::FMacOpenGLTimer(FPlatformOpenGLContext* InContext)
-: Context(InContext)
-, Name(0)
-, Result(0)
-, Accumulated(0)
-, Available(0)
-, Cached(false)
-, Running(false)
-{
-	check(Context);
-	glGenQueries(1, &Name);
-}
-
-FMacOpenGLTimer::~FMacOpenGLTimer()
-{
-	End();
-	Next.Reset();
-	glDeleteQueries(1, &Name);
-}
-
-void FMacOpenGLTimer::Begin(void)
-{
-	if(!Running)
-	{
-		glBeginQuery(GL_TIME_ELAPSED, Name);
-		Running = true;
-	}
-}
-
-void FMacOpenGLTimer::End(void)
-{
-	if(Running)
-	{
-		Running = false;
-		glEndQuery(GL_TIME_ELAPSED);
-	}
-}
-
-uint64 FMacOpenGLTimer::GetResult(void)
-{
-	check(!Running);
-	CacheResult();
-	return Result;
-}
-
-uint64 FMacOpenGLTimer::GetAccumulatedResult(void)
-{
-	check(!Running);
-	CacheResult();
-	return Accumulated;
-}
-
-int32 FMacOpenGLTimer::GetResultAvailable(void)
-{
-	check(!Running);
-	if(!Available)
-	{
-		glGetQueryObjectuiv(Name, GL_QUERY_RESULT_AVAILABLE, &Available);
-	}
-	return Available;
-}
-
-void FMacOpenGLTimer::CacheResult(void)
-{
-	check(!Running);
-	if(!Cached)
-	{
-		glGetQueryObjectui64v(Name, GL_QUERY_RESULT, &Result);
-		
-		if(Previous.IsValid())
-		{
-			Accumulated += Result;
-			Accumulated += Previous->GetAccumulatedResult();
-			
-			Previous->Next.Reset();
-			Previous.Reset();
-		}
-		
-		Available = true;
-		Cached = true;
-	}
-}
-
-FMacOpenGLQuery::FMacOpenGLQuery(FPlatformOpenGLContext* InContext)
-	: Name(0)
-	, Target(0)
-	, Context(InContext)
-	, Result(0)
-	, Available(0)
-	, Running(0)
-	, Cached(false)
-{
-	check(Context);
-}
-
-FMacOpenGLQuery::~FMacOpenGLQuery()
-{
-	if(Running)
-	{
-		End();
-	}
-	if(Name)
-	{
-		glDeleteQueries(1, &Name);
-	}
-}
-
-void FMacOpenGLQuery::Begin(GLenum InTarget)
-{
-	if(Target == 0 || Target == InTarget)
-	{
-		Target = InTarget;
-		Result = 0;
-		Running = true;
-		Available = false;
-		Cached = false;
-		Start.Reset();
-		Finish.Reset();
-		switch(Target)
-		{
-			case GL_TIMESTAMP:
-			{
-				// There can be a lot of timestamps, clear emulated queries out to avoid problems.
-				for( auto It : Context->Queries )
-				{
-					if(It.IsValid() && It->Running == false && (It->Target == GL_TIMESTAMP || It->Target == GL_TIME_ELAPSED))
-					{
-						if(It->GetResultAvailable())
-						{
-							It->GetResult();
-						}
-						else
-						{
-							break;
-						}
-					}
-				}
-				
-				check(Context->LastTimer.IsValid());
-				Context->LastTimer->End();
-				Start = Context->LastTimer;
-				
-				TSharedPtr<FMacOpenGLTimer> Current(new FMacOpenGLTimer(Context));
-				Current->Previous = Start;
-				Start->Next = Current;
-				Current->Begin();
-				Context->LastTimer = Current;
-				
-				break;
-			}
-			case GL_TIME_ELAPSED:
-			{
-				check(Context->LastTimer.IsValid());
-				Context->LastTimer->End();
-				
-				Start = TSharedPtr<FMacOpenGLTimer>(new FMacOpenGLTimer(Context));
-				Start->Previous = Context->LastTimer;
-				Context->LastTimer->Next = Start;
-				Start->Begin();
-				
-				Context->LastTimer = Start;
-				break;
-			}
-			default:
-			{
-				if(!Name)
-				{
-					glGenQueries(1, &Name);
-				}
-				glBeginQuery(Target, Name);
-				break;
-			}
-		}
-	}
-}
-
-void FMacOpenGLQuery::End()
-{
-	Running = false;
-	if(Target == GL_TIME_ELAPSED)
-	{
-		Finish = Context->LastTimer;
-		TSharedPtr<FMacOpenGLTimer> Current(new FMacOpenGLTimer(Context));
-		Current->Previous = Finish;
-		Finish->Next = Current;
-		Finish->End();
-		Current->Begin();
-		Context->LastTimer = Current;
-		
-		if(Start == Finish)
-		{
-			Start = Start->Previous;
-		}
-	}
-	else if(Target != GL_TIMESTAMP)
-	{
-		glEndQuery(Target);
-	}
-}
-
-uint64 FMacOpenGLQuery::GetResult(void)
-{
-	if(!Cached)
-	{
-		switch(Target)
-		{
-			case GL_TIMESTAMP:
-				Result = Start->GetAccumulatedResult();
-				Start.Reset();
-				Finish.Reset();
-				break;
-			case GL_TIME_ELAPSED:
-				Result = Finish->GetAccumulatedResult() - Start->GetAccumulatedResult();
-				Start.Reset();
-				Finish.Reset();
-				break;
-			default:
-				glGetQueryObjectui64v(Name, GL_QUERY_RESULT, &Result);
-				break;
-		}
-		Available = true;
-		Cached = true;
-	}
-	return Result;
-}
-
-int32 FMacOpenGLQuery::GetResultAvailable(void)
-{
-	if(!Available)
-	{
-		switch(Target)
-		{
-			case GL_TIMESTAMP:
-				Available = Start->GetResultAvailable();
-				break;
-			case GL_TIME_ELAPSED:
-				Available = Finish->GetResultAvailable();
-				break;
-			default:
-				glGetQueryObjectuiv(Name, GL_QUERY_RESULT_AVAILABLE, &Available);
-				break;
-		}
-	}
-	return Available;
-}
-
-bool IsQuery(FPlatformOpenGLContext* PlatformContext, GLuint QueryID)
-{
-	bool OK = PlatformContext && QueryID;
-	if(OK)
-	{
-		OK = PlatformContext->Queries.Num() >= QueryID && PlatformContext->Queries[QueryID-1].IsValid();
-	}
-	return OK;
-}
-
-TSharedPtr<FMacOpenGLQuery> GetQuery(FPlatformOpenGLContext* PlatformContext, GLuint QueryID)
-{
-	TSharedPtr<FMacOpenGLQuery> Query;
-	if(IsQuery(PlatformContext, QueryID))
-	{
-		Query = PlatformContext->Queries[QueryID-1];
-	}
-	return Query;
-}
-
 void FMacOpenGL::MacQueryTimestampCounter(GLuint QueryID)
 {
 	if(GIsEmulatingTimestamp)
@@ -1770,7 +1209,7 @@ void FMacOpenGL::MacQueryTimestampCounter(GLuint QueryID)
 		FPlatformOpenGLContext* PlatformContext = nullptr;
 		if(Current && (CGLGetParameter(Current, kCGLCPClientStorage, (GLint *)&PlatformContext) == kCGLNoError) && PlatformContext)
 		{
-			TSharedPtr<FMacOpenGLQuery> Query = GetQuery(PlatformContext, QueryID);
+			TSharedPtr<FMacOpenGLQuery> Query = PlatformContext->EmulatedQueries->GetQuery(QueryID);
 			if(Query.IsValid())
 			{
 				Query->Begin(GL_TIMESTAMP);
@@ -1792,12 +1231,7 @@ void FMacOpenGL::MacBeginQuery(GLenum QueryType, GLuint QueryId)
 		FPlatformOpenGLContext* PlatformContext = nullptr;
 		if(Current && (CGLGetParameter(Current, kCGLCPClientStorage, (GLint *)&PlatformContext) == kCGLNoError) && PlatformContext)
 		{
-			TSharedPtr<FMacOpenGLQuery> Query = GetQuery(PlatformContext, QueryId);
-			if(Query.IsValid())
-			{
-				Query->Begin(QueryType);
-				PlatformContext->RunningQueries.Add(QueryType, Query);
-			}
+			PlatformContext->EmulatedQueries->BeginQuery(QueryType, QueryId);
 		}
 	}
 	else
@@ -1814,12 +1248,7 @@ void FMacOpenGL::MacEndQuery(GLenum QueryType)
 		FPlatformOpenGLContext* PlatformContext = nullptr;
 		if(Current && (CGLGetParameter(Current, kCGLCPClientStorage, (GLint *)&PlatformContext) == kCGLNoError) && PlatformContext)
 		{
-			TSharedPtr<FMacOpenGLQuery> Query = PlatformContext->RunningQueries.FindOrAdd(QueryType);
-			if(Query.IsValid())
-			{
-				Query->End();
-				PlatformContext->RunningQueries[QueryType].Reset();
-			}
+			PlatformContext->EmulatedQueries->EndQuery(QueryType);
 		}
 	}
 	else
@@ -1836,7 +1265,7 @@ void FMacOpenGL::MacGetQueryObject(GLuint QueryId, EQueryMode QueryMode, uint64
 		FPlatformOpenGLContext* PlatformContext = nullptr;
 		if(Current && (CGLGetParameter(Current, kCGLCPClientStorage, (GLint *)&PlatformContext) == kCGLNoError) && PlatformContext)
 		{
-			TSharedPtr<FMacOpenGLQuery> Query = GetQuery(PlatformContext, QueryId);
+			TSharedPtr<FMacOpenGLQuery> Query = PlatformContext->EmulatedQueries->GetQuery(QueryId);
 			if(Query.IsValid())
 			{
 				if(QueryMode == QM_Result)
@@ -1864,7 +1293,7 @@ void FMacOpenGL::MacGetQueryObject(GLuint QueryId, EQueryMode QueryMode, GLuint
 		FPlatformOpenGLContext* PlatformContext = nullptr;
 		if(Current && (CGLGetParameter(Current, kCGLCPClientStorage, (GLint *)&PlatformContext) == kCGLNoError) && PlatformContext)
 		{
-			TSharedPtr<FMacOpenGLQuery> Query = GetQuery(PlatformContext, QueryId);
+			TSharedPtr<FMacOpenGLQuery> Query = PlatformContext->EmulatedQueries->GetQuery(QueryId);
 			if(Query.IsValid())
 			{
 				if(QueryMode == QM_Result)
@@ -1887,7 +1316,8 @@ void FMacOpenGL::MacGetQueryObject(GLuint QueryId, EQueryMode QueryMode, GLuint
 bool FMacOpenGL::MustFlushTexStorage(void)
 {
 	// @todo There is a bug in Apple's GL with TexStorage calls, on Nvidia and when using MTGL that can see the texture never be created, which then subsequently causes crashes
-	return GMacMustFlushTexStorage || GMacUseMTGL;
+	FPlatformOpenGLContext::VerifyCurrentContext();
+	return GMacFlushTexStorage && (IsRHIDeviceNVIDIA() || GMacUseMTGL);
 }
 
 void FMacOpenGL::GenBuffers( GLsizei n, GLuint *buffers)
@@ -1951,27 +1381,29 @@ void FMacOpenGL::DeleteBuffers(GLsizei Number, const GLuint* Buffers)
 	}
 #endif
 	
-	GMacBuffersToDelete.Append(&Buffers[Index], Number-Index);
+	if( Number > Index )
+	{
+		GMacBuffersToDelete.Append(&Buffers[Index], Number - Index);
+	}
+	
 	if(GMacBuffersToDelete.Num() >= GMacMinBuffersToDelete)
 	{
 		glDeleteBuffers(GMacBuffersToDelete.Num(), GMacBuffersToDelete.GetData());
 		GMacBuffersToDelete.Empty();
 	}
+#if UE_BUILD_DEBUG
 	else
 	{
 		for(uint32 i = Index; i < Number; i++)
 		{
-#if UE_BUILD_DEBUG
 			glBindBuffer(GL_COPY_WRITE_BUFFER, Buffers[i]);
 			GLint bLocked = GL_FALSE;
 			glGetBufferParameteriv(GL_COPY_WRITE_BUFFER, GL_BUFFER_MAPPED, &bLocked);
 			check(!bLocked);
-#endif
 		}
 	}
 	
-#if UE_BUILD_DEBUG
-	glBindBuffer(GL_COPY_WRITE_BUFFER, 0);
+	glBindBuffer(GL_COPY_WRITE_BUFFER, Buffer);
 #endif
 }
 
diff --git a/Engine/Source/Runtime/OpenGLDrv/Private/Mac/OpenGLMac.h b/Engine/Source/Runtime/OpenGLDrv/Private/Mac/OpenGLMac.h
--- a/Engine/Source/Runtime/OpenGLDrv/Private/Mac/OpenGLMac.h
+++ b/Engine/Source/Runtime/OpenGLDrv/Private/Mac/OpenGLMac.h
@@ -5,36 +5,9 @@
 #include <OpenGL/gl3.h>
 #include <OpenGL/gl3ext.h>
 
-#ifndef GL_VERSION_3_3
-#define glVertexAttribDivisor		glVertexAttribDivisorARB
-#endif
 #ifndef GL_COMPUTE_SHADER
 #define GL_COMPUTE_SHADER 0x91B9
 #endif
-#ifndef GL_VERSION_4_0
-#define GL_TEXTURE_CUBE_MAP_ARRAY         0x9009
-/** GL_ARB_draw_buffers_blend */
-typedef void (APIENTRYP PFNGLBLENDEQUATIONIARBPROC) (GLuint buf, GLenum mode);
-typedef void (APIENTRYP PFNGLBLENDEQUATIONSEPARATEIARBPROC) (GLuint buf, GLenum modeRGB, GLenum modeAlpha);
-typedef void (APIENTRYP PFNGLBLENDFUNCIARBPROC) (GLuint buf, GLenum src, GLenum dst);
-typedef void (APIENTRYP PFNGLBLENDFUNCSEPARATEIARBPROC) (GLuint buf, GLenum srcRGB, GLenum dstRGB, GLenum srcAlpha, GLenum dstAlpha);
-/** GL_EXT_debug_label */
-typedef void (APIENTRYP PFNGLLABELOBJECTEXTPROC) (GLenum type, GLuint object, GLsizei length, const GLchar *label);
-typedef void (APIENTRYP PFNGLGETOBJECTLABELEXTPROC) (GLenum type, GLuint object, GLsizei bufSize, GLsizei *length, GLchar *label);
-/** GL_EXT_debug_marker */
-typedef void (APIENTRYP PFNGLINSERTEVENTMARKEREXTPROC) (GLsizei length, const char *marker);
-typedef void (APIENTRYP PFNGLPUSHGROUPMARKEREXTPROC) (GLsizei length, const char *marker);
-typedef void (APIENTRYP PFNGLPOPGROUPMARKEREXTPROC) (void);
-/** GL_ARB_tessellation_shader */
-#define GL_MAX_TESS_CONTROL_UNIFORM_COMPONENTS 0x8E7F
-#define GL_MAX_TESS_EVALUATION_UNIFORM_COMPONENTS 0x8E80
-typedef void (APIENTRYP PFNGLPATCHPARAMETERIPROC) (GLenum pname, GLint value);
-typedef void (APIENTRYP PFNGLPATCHPARAMETERFVPROC) (GLenum pname, const GLfloat *values);
-#endif
-#ifndef GL_VERSION_4_1
-typedef void (APIENTRYP PFNGLTEXSTORAGE2DPROC) (GLenum target, GLsizei levels, GLenum internalformat, GLsizei width, GLsizei height);
-typedef void (APIENTRYP PFNGLTEXSTORAGE3DPROC) (GLenum target, GLsizei levels, GLenum internalformat, GLsizei width, GLsizei height, GLsizei depth);
-#endif
 
 #include "OpenGL3.h"
 
@@ -199,8 +172,6 @@ struct FMacOpenGL : public FOpenGL3
 		}
 	}
 	
-	static FORCEINLINE bool SupportsTessellation()						{ return bSupportsTessellationShader; }
-	
 	static FORCEINLINE void PatchParameteri(GLenum Pname, GLint Value)
 	{
 		if(glPatchParameteri)
@@ -270,7 +241,6 @@ private:
 	static PFNGLPUSHGROUPMARKEREXTPROC glPushGroupMarkerEXT;
 	static PFNGLPOPGROUPMARKEREXTPROC glPopGroupMarkerEXT;
 	/** GL_ARB_tessellation_shader */
-	static bool bSupportsTessellationShader;
 	static PFNGLPATCHPARAMETERIPROC glPatchParameteri;
 };
 
diff --git a/Engine/Source/Runtime/OpenGLDrv/Private/OpenGLCommands.cpp b/Engine/Source/Runtime/OpenGLDrv/Private/OpenGLCommands.cpp
--- a/Engine/Source/Runtime/OpenGLDrv/Private/OpenGLCommands.cpp
+++ b/Engine/Source/Runtime/OpenGLDrv/Private/OpenGLCommands.cpp
@@ -1295,9 +1295,7 @@ void FOpenGLDynamicRHI::SetPendingBlendStateForActiveRenderTargets( FOpenGLConte
 							);
 					}
 					
-#if PLATFORM_MAC
-					extern bool GIsRunningOnIntelCard;
-					if (GIsRunningOnIntelCard)
+					if (PLATFORM_MAC && IsRHIDeviceIntel())
 					{
 						// Flush the separate blend state changes through GL or on Intel cards under OS X the state change may be silently ignored.
 						if (CachedRenderTargetBlendState.ColorSourceBlendFactor != RenderTargetBlendState.ColorSourceBlendFactor
@@ -1310,7 +1308,6 @@ void FOpenGLDynamicRHI::SetPendingBlendStateForActiveRenderTargets( FOpenGLConte
 							FOpenGL::Flush();
 						}
 					}
-#endif
 				}
 				else
 				{
