Defer the GL_TIMESTAMP emulation on Nvidia GPUs under Mac OS X, it may still crash but it is sufficiently rare that it is possible to profile the game. The number of frames to defer timing is controlled by r.Mac.Nvidia.DelayEmulatingTimeStamp.

diff --git a/Engine/Source/Runtime/OpenGLDrv/Private/Mac/MacOpenGLQuery.cpp b/Engine/Source/Runtime/OpenGLDrv/Private/Mac/MacOpenGLQuery.cpp
--- a/Engine/Source/Runtime/OpenGLDrv/Private/Mac/MacOpenGLQuery.cpp
+++ b/Engine/Source/Runtime/OpenGLDrv/Private/Mac/MacOpenGLQuery.cpp
@@ -9,6 +9,13 @@
  ------------------------------------------------------------------------------*/
 
 bool GIsEmulatingTimestamp = false; // @todo: Now crashing on Nvidia cards, but not on AMD...
+static int32 GNVDelayEmulatingTimeStamp = 60;
+static FAutoConsoleVariableRef CVarNVDelayEmulatingTimeStamp(
+	  TEXT("r.Mac.Nvidia.DelayEmulatingTimeStamp"),
+	  GNVDelayEmulatingTimeStamp,
+	  TEXT("The number of frames to defer emulating GL_TIMESTAMP on Nvidia cards under Mac OS X to avoid crashes on startup, set to 0 to disable. (Default: 60)"),
+	  ECVF_RenderThreadSafe
+	  );
 
 /*------------------------------------------------------------------------------
  OpenGL query emulation.
@@ -24,21 +31,30 @@ FMacOpenGLTimer::FMacOpenGLTimer(FPlatformOpenGLContext* InContext)
 , Running(false)
 {
 	check(Context);
-	glGenQueries(1, &Name);
+	if (!IsRHIDeviceNVIDIA() || GFrameNumberRenderThread > GNVDelayEmulatingTimeStamp)
+	{
+		glGenQueries(1, &Name);
+	}
 }
 
 FMacOpenGLTimer::~FMacOpenGLTimer()
 {
 	End();
 	Next.Reset();
-	glDeleteQueries(1, &Name);
+	if (Name)
+	{
+		glDeleteQueries(1, &Name);
+	}
 }
 
 void FMacOpenGLTimer::Begin(void)
 {
 	if(!Running)
 	{
-		glBeginQuery(GL_TIME_ELAPSED, Name);
+		if (!IsRHIDeviceNVIDIA() || GFrameNumberRenderThread > GNVDelayEmulatingTimeStamp)
+		{
+			glBeginQuery(GL_TIME_ELAPSED, Name);
+		}
 		Running = true;
 	}
 }
@@ -48,7 +64,10 @@ void FMacOpenGLTimer::End(void)
 	if(Running)
 	{
 		Running = false;
-		glEndQuery(GL_TIME_ELAPSED);
+		if (!IsRHIDeviceNVIDIA() || GFrameNumberRenderThread > GNVDelayEmulatingTimeStamp)
+		{
+			glEndQuery(GL_TIME_ELAPSED);
+		}
 	}
 }
 
@@ -71,7 +90,14 @@ int32 FMacOpenGLTimer::GetResultAvailable(void)
 	check(!Running);
 	if(!Available)
 	{
-		glGetQueryObjectuiv(Name, GL_QUERY_RESULT_AVAILABLE, &Available);
+		if (!IsRHIDeviceNVIDIA() || GFrameNumberRenderThread > GNVDelayEmulatingTimeStamp)
+		{
+			glGetQueryObjectuiv(Name, GL_QUERY_RESULT_AVAILABLE, &Available);
+		}
+		else
+		{
+			Available = true;
+		}
 	}
 	return Available;
 }
@@ -81,7 +107,14 @@ void FMacOpenGLTimer::CacheResult(void)
 	check(!Running);
 	if(!Cached)
 	{
-		glGetQueryObjectui64v(Name, GL_QUERY_RESULT, &Result);
+		if (!IsRHIDeviceNVIDIA() || GFrameNumberRenderThread > GNVDelayEmulatingTimeStamp)
+		{
+			glGetQueryObjectui64v(Name, GL_QUERY_RESULT, &Result);
+		}
+		else
+		{
+			Result = 0;
+		}
 		
 		if(Previous.IsValid())
 		{
