Fix for Mac&Linux build

diff --git a/Engine/Source/Runtime/OpenGLDrv/Private/OpenGLDevice.cpp b/Engine/Source/Runtime/OpenGLDrv/Private/OpenGLDevice.cpp
--- a/Engine/Source/Runtime/OpenGLDrv/Private/OpenGLDevice.cpp
+++ b/Engine/Source/Runtime/OpenGLDrv/Private/OpenGLDevice.cpp
@@ -708,7 +708,7 @@ static void InitRHICapabilitiesForGL()
 	SetupTextureFormat( PF_R16G16B16A16_SINT,	FOpenGLTextureFormat( GL_RGBA16I,				GL_NONE,				GL_RGBA_INTEGER,	GL_SHORT,						false,			false));
 	SetupTextureFormat( PF_R5G6B5_UNORM,		FOpenGLTextureFormat( ));
 #if PLATFORM_DESKTOP || PLATFORM_ANDROIDGL4 || PLATFORM_ANDROIDES31
-	if (PLATFORM_DESKTOP || PLATFORM_ANDROIDGL4 || FOpenGL::GetFeatureLevel() >= ERHIFeatureLevel::SM4)
+	if (PLATFORM_DESKTOP != 0 || PLATFORM_ANDROIDGL4 != 0 || FOpenGL::GetFeatureLevel() >= ERHIFeatureLevel::SM4)
 	{
 		// Not supported for rendering:
 		SetupTextureFormat( PF_G16,				FOpenGLTextureFormat( GL_R16,					GL_R16,					GL_RED,			GL_UNSIGNED_SHORT,					false,	false));
