Removed some unnecessary log

diff --git a/Engine/Source/Runtime/OpenGLDrv/Private/OpenGLDevice.cpp b/Engine/Source/Runtime/OpenGLDrv/Private/OpenGLDevice.cpp
--- a/Engine/Source/Runtime/OpenGLDrv/Private/OpenGLDevice.cpp
+++ b/Engine/Source/Runtime/OpenGLDrv/Private/OpenGLDevice.cpp
@@ -734,7 +734,6 @@ static void InitRHICapabilitiesForGL()
 		SetupTextureFormat( PF_FloatR11G11B10,	FOpenGLTextureFormat( GL_R11F_G11F_B10F,		GL_R11F_G11F_B10F,		GL_RGB,			GL_UNSIGNED_INT_10F_11F_11F_REV,	false,	false));
 		if (FOpenGL::GetShaderPlatform() == EShaderPlatform::SP_OPENGL_ES31_EXT)
 		{
-			UE_LOG(LogRHI, Log, TEXT("FOpenGL::GetShaderPlatform() == SP_OPENGL_ES31_EXT"));
 			SetupTextureFormat(PF_B8G8R8A8, FOpenGLTextureFormat(GL_RGBA8, GL_SRGB8_ALPHA8, GL_RGBA, GL_UNSIGNED_BYTE, false, true));
 			SetupTextureFormat(PF_R8G8B8A8, FOpenGLTextureFormat(GL_RGBA8, GL_SRGB8_ALPHA8, GL_RGBA, GL_UNSIGNED_BYTE, false, false));
 			// GL_RG16 is not supported in OpenGL ES. There is currently no extension for it either. The user should check for support and implement a fallback.
@@ -742,7 +741,6 @@ static void InitRHICapabilitiesForGL()
 		}
 		else
 		{
-			UE_LOG(LogRHI, Log, TEXT("FOpenGL::GetShaderPlatform() != SP_OPENGL_ES31_EXT"));
 			SetupTextureFormat(PF_B8G8R8A8, FOpenGLTextureFormat(GL_RGBA8, GL_SRGB8_ALPHA8, GL_BGRA, GL_UNSIGNED_INT_8_8_8_8_REV, false, false));
 			SetupTextureFormat(PF_R8G8B8A8, FOpenGLTextureFormat(GL_RGBA8, GL_SRGB8_ALPHA8, GL_RGBA, GL_UNSIGNED_INT_8_8_8_8_REV, false, false));
 			SetupTextureFormat(PF_G16R16, FOpenGLTextureFormat(GL_RG16, GL_RG16, GL_RG, GL_UNSIGNED_SHORT, false, false));
