Fix for Mac

diff --git a/Engine/Source/Runtime/OpenGLDrv/Private/OpenGLDrvPrivate.h b/Engine/Source/Runtime/OpenGLDrv/Private/OpenGLDrvPrivate.h
--- a/Engine/Source/Runtime/OpenGLDrv/Private/OpenGLDrvPrivate.h
+++ b/Engine/Source/Runtime/OpenGLDrv/Private/OpenGLDrvPrivate.h
@@ -507,6 +507,7 @@ inline bool OpenGLShaderPlatformNeedsBindLocation(const EShaderPlatform InShader
 			return false;
 
 		case SP_OPENGL_SM4:
+		case SP_OPENGL_SM4_MAC:
 		case SP_OPENGL_PCES2:
 		case SP_OPENGL_ES2:
 		case SP_OPENGL_ES2_WEBGL:
@@ -514,7 +515,7 @@ inline bool OpenGLShaderPlatformNeedsBindLocation(const EShaderPlatform InShader
 			return true;
 		default:
 			check(IsOpenGLPlatform(InShaderPlatform));
-			checkf(false, TEXT("invalid shader platform"));
+			checkf(false, TEXT("invalid shader platform (%d)"), int(InShaderPlatform));
 			return true;
 			break;
 	}
