Temporarily disable vertex attrib binding for Linux due to crashes in the wild.

diff --git a/Engine/Source/Runtime/OpenGLDrv/Private/OpenGLCommands.cpp b/Engine/Source/Runtime/OpenGLDrv/Private/OpenGLCommands.cpp
--- a/Engine/Source/Runtime/OpenGLDrv/Private/OpenGLCommands.cpp
+++ b/Engine/Source/Runtime/OpenGLDrv/Private/OpenGLCommands.cpp
@@ -46,7 +46,7 @@ namespace OpenGLConsoleVariables
 		TEXT("If true, don't issue dispatch work.")
 		);
 
-	int32 bUseVAB = 1;
+	int32 bUseVAB = PLATFORM_LINUX ? 0 : 1;	// [RCL] @FIXME - currently (2014-09-19) causes crashes
 	static FAutoConsoleVariableRef CVarUseVAB(
 		TEXT("OpenGL.UseVAB"),
 		bUseVAB,
