UE4 - Compile fix

diff --git a/Engine/Source/Runtime/OpenGLDrv/Public/OpenGL.h b/Engine/Source/Runtime/OpenGLDrv/Public/OpenGL.h
--- a/Engine/Source/Runtime/OpenGLDrv/Public/OpenGL.h
+++ b/Engine/Source/Runtime/OpenGLDrv/Public/OpenGL.h
@@ -34,6 +34,12 @@ struct FPlatformOpenGLContext;
 #define OPENGL_GL4	0
 #endif
 
+/** Official OpenGL definitions */
+#ifndef GL_HALF_FLOAT
+#define GL_HALF_FLOAT 0x140B
+#endif
+
+// Base static class
 class FOpenGLBase
 {
 public:
