UE4 - More compile fixes

diff --git a/Engine/Source/Runtime/OpenGLDrv/Private/OpenGLDevice.cpp b/Engine/Source/Runtime/OpenGLDrv/Private/OpenGLDevice.cpp
--- a/Engine/Source/Runtime/OpenGLDrv/Private/OpenGLDevice.cpp
+++ b/Engine/Source/Runtime/OpenGLDrv/Private/OpenGLDevice.cpp
@@ -259,11 +259,10 @@ static const TCHAR* GetOpenGLDebugSeverityStringARB(GLenum Severity)
 /**
  * OpenGL debug message callback. Conforms to GLDEBUGPROCARB.
  */
-#if PLATFORM_ANDROID && !PLATFORM_ANDROIDGL4
-
-#ifndef GL_APIENTRY
-#define GL_APIENTRY APIENTRY
-#endif
+#if (PLATFORM_ANDROID && !PLATFORM_ANDROIDGL4) || PLATFORM_HTML5
+	#ifndef GL_APIENTRY
+	#define GL_APIENTRY APIENTRY
+	#endif
 static void GL_APIENTRY OpenGLDebugMessageCallbackARB(
 #else
 static void APIENTRY OpenGLDebugMessageCallbackARB(
@@ -668,7 +667,7 @@ static void InitRHICapabilitiesForGL()
 	GShaderPlatformForFeatureLevel[ERHIFeatureLevel::ES2] = (GMaxRHIFeatureLevel == ERHIFeatureLevel::ES2) ? GRHIShaderPlatform : SP_OPENGL_PCES2;
 	GShaderPlatformForFeatureLevel[ERHIFeatureLevel::ES3_1] = SP_NumPlatforms;
 	GShaderPlatformForFeatureLevel[ERHIFeatureLevel::SM4] = PLATFORM_MAC ? SP_OPENGL_SM4_MAC : SP_OPENGL_SM4;
- 	GShaderPlatformForFeatureLevel[ERHIFeatureLevel::SM5] = OPENGL_ES31 ? SP_OPENGL_ES31_EXT : SP_OPENGL_SM5;
+	GShaderPlatformForFeatureLevel[ERHIFeatureLevel::SM5] = OPENGL_ES31 ? SP_OPENGL_ES31_EXT : SP_OPENGL_SM5;
 
 	// Set to same values as in DX11, as for the time being clip space adjustment are done entirely
 	// in HLSLCC-generated shader code and OpenGLDrv.
