Removed unnecessary debug logging

diff --git a/Engine/Source/Runtime/OpenGLDrv/Private/OpenGLTexture.cpp b/Engine/Source/Runtime/OpenGLDrv/Private/OpenGLTexture.cpp
--- a/Engine/Source/Runtime/OpenGLDrv/Private/OpenGLTexture.cpp
+++ b/Engine/Source/Runtime/OpenGLDrv/Private/OpenGLTexture.cpp
@@ -647,9 +647,6 @@ void TOpenGLTexture<RHIResourceType>::Resolve(uint32 MipIndex,uint32 ArrayIndex)
 	
 	glBindBuffer( GL_PIXEL_PACK_BUFFER, PixelBuffer->Resource );
 
-	UE_LOG(LogRHI,Warning,TEXT("Attempting ReadPixels with Format = %x Type = %x Attachement = %x"),GLFormat.Format,GLFormat.Type,Attachment);
-
-	
 #if PLATFORM_MAC || PLATFORM_ANDROIDES31 // glReadPixels is async with PBOs - glGetTexImage is not: radr://16011763
 	if(Attachment == GL_COLOR_ATTACHMENT0 && !GLFormat.bCompressed)
 	{
