Make OpenGL match the D3D11RHI and handle nullptr when setting a texture - this fixes TTP #346500 & TTP #345897 when using an OpenGL renderer.

diff --git a/Engine/Source/Runtime/OpenGLDrv/Private/OpenGLCommands.cpp b/Engine/Source/Runtime/OpenGLDrv/Private/OpenGLCommands.cpp
--- a/Engine/Source/Runtime/OpenGLDrv/Private/OpenGLCommands.cpp
+++ b/Engine/Source/Runtime/OpenGLDrv/Private/OpenGLCommands.cpp
@@ -856,7 +856,14 @@ void FOpenGLDynamicRHI::RHISetShaderTexture(FVertexShaderRHIParamRef VertexShade
 
 	VERIFY_GL_SCOPE();
 	FOpenGLTextureBase* NewTexture = GetOpenGLTextureFromRHITexture(NewTextureRHI);
-	InternalSetShaderTexture(NewTexture, FOpenGL::GetFirstVertexTextureUnit() + TextureIndex, NewTexture->Target, NewTexture->Resource, NewTextureRHI->GetNumMips(), -1);
+	if (NewTexture)
+	{
+		InternalSetShaderTexture(NewTexture, FOpenGL::GetFirstVertexTextureUnit() + TextureIndex, NewTexture->Target, NewTexture->Resource, NewTextureRHI->GetNumMips(), -1);
+	}
+	else
+	{
+		InternalSetShaderTexture(nullptr, FOpenGL::GetFirstVertexTextureUnit() + TextureIndex, 0, 0, 0, -1);
+	}
 }
 
 void FOpenGLDynamicRHI::RHISetShaderTexture(FHullShaderRHIParamRef HullShaderRHI,uint32 TextureIndex,FTextureRHIParamRef NewTextureRHI)
@@ -866,7 +873,14 @@ void FOpenGLDynamicRHI::RHISetShaderTexture(FHullShaderRHIParamRef HullShaderRHI
 	check(GMaxRHIFeatureLevel >= ERHIFeatureLevel::SM5);
 	VERIFY_GL_SCOPE();
 	FOpenGLTextureBase* NewTexture = GetOpenGLTextureFromRHITexture(NewTextureRHI);
-	InternalSetShaderTexture(NewTexture, FOpenGL::GetFirstHullTextureUnit() + TextureIndex, NewTexture->Target, NewTexture->Resource, NewTextureRHI->GetNumMips(), -1);
+	if (NewTexture)
+	{
+		InternalSetShaderTexture(NewTexture, FOpenGL::GetFirstHullTextureUnit() + TextureIndex, NewTexture->Target, NewTexture->Resource, NewTextureRHI->GetNumMips(), -1);
+	}
+	else
+	{
+		InternalSetShaderTexture(nullptr, FOpenGL::GetFirstHullTextureUnit() + TextureIndex, 0, 0, 0, -1);
+	}
 }
 
 void FOpenGLDynamicRHI::RHISetShaderTexture(FDomainShaderRHIParamRef DomainShaderRHI,uint32 TextureIndex,FTextureRHIParamRef NewTextureRHI)
@@ -876,7 +890,14 @@ void FOpenGLDynamicRHI::RHISetShaderTexture(FDomainShaderRHIParamRef DomainShade
 	check(GMaxRHIFeatureLevel >= ERHIFeatureLevel::SM5);
 	VERIFY_GL_SCOPE();
 	FOpenGLTextureBase* NewTexture = GetOpenGLTextureFromRHITexture(NewTextureRHI);
-	InternalSetShaderTexture(NewTexture, FOpenGL::GetFirstDomainTextureUnit() + TextureIndex, NewTexture->Target, NewTexture->Resource, NewTextureRHI->GetNumMips(), -1);
+	if (NewTexture)
+	{
+		InternalSetShaderTexture(NewTexture, FOpenGL::GetFirstDomainTextureUnit() + TextureIndex, NewTexture->Target, NewTexture->Resource, NewTextureRHI->GetNumMips(), -1);
+	}
+	else
+	{
+		InternalSetShaderTexture(nullptr, FOpenGL::GetFirstDomainTextureUnit() + TextureIndex, 0, 0, 0, -1);
+	}
 }
 
 void FOpenGLDynamicRHI::RHISetShaderTexture(FGeometryShaderRHIParamRef GeometryShaderRHI,uint32 TextureIndex,FTextureRHIParamRef NewTextureRHI)
@@ -885,7 +906,14 @@ void FOpenGLDynamicRHI::RHISetShaderTexture(FGeometryShaderRHIParamRef GeometryS
 
 	VERIFY_GL_SCOPE();
 	FOpenGLTextureBase* NewTexture = GetOpenGLTextureFromRHITexture(NewTextureRHI);
-	InternalSetShaderTexture(NewTexture, FOpenGL::GetFirstGeometryTextureUnit() + TextureIndex, NewTexture->Target, NewTexture->Resource, NewTextureRHI->GetNumMips(), -1);
+	if (NewTexture)
+	{
+		InternalSetShaderTexture(NewTexture, FOpenGL::GetFirstGeometryTextureUnit() + TextureIndex, NewTexture->Target, NewTexture->Resource, NewTextureRHI->GetNumMips(), -1);
+	}
+	else
+	{
+		InternalSetShaderTexture(nullptr, FOpenGL::GetFirstGeometryTextureUnit() + TextureIndex, 0, 0, 0, -1);
+	}
 }
 
 void FOpenGLDynamicRHI::RHISetShaderTexture(FPixelShaderRHIParamRef PixelShaderRHI,uint32 TextureIndex,FTextureRHIParamRef NewTextureRHI)
@@ -894,7 +922,14 @@ void FOpenGLDynamicRHI::RHISetShaderTexture(FPixelShaderRHIParamRef PixelShaderR
 
 	VERIFY_GL_SCOPE();
 	FOpenGLTextureBase* NewTexture = GetOpenGLTextureFromRHITexture(NewTextureRHI);
-	InternalSetShaderTexture(NewTexture, FOpenGL::GetFirstPixelTextureUnit() + TextureIndex, NewTexture->Target, NewTexture->Resource, NewTextureRHI->GetNumMips(), -1);
+	if (NewTexture)
+	{
+		InternalSetShaderTexture(NewTexture, FOpenGL::GetFirstPixelTextureUnit() + TextureIndex, NewTexture->Target, NewTexture->Resource, NewTextureRHI->GetNumMips(), -1);
+	}
+	else
+	{
+		InternalSetShaderTexture(nullptr, FOpenGL::GetFirstPixelTextureUnit() + TextureIndex, 0, 0, 0, -1);
+	}
 }
 
 void FOpenGLDynamicRHI::RHISetShaderSampler(FVertexShaderRHIParamRef VertexShaderRHI,uint32 SamplerIndex,FSamplerStateRHIParamRef NewStateRHI)
@@ -995,7 +1030,14 @@ void FOpenGLDynamicRHI::RHISetShaderTexture(FComputeShaderRHIParamRef ComputeSha
 	check(GMaxRHIFeatureLevel >= ERHIFeatureLevel::SM5);
 	VERIFY_GL_SCOPE();
 	FOpenGLTextureBase* NewTexture = GetOpenGLTextureFromRHITexture(NewTextureRHI);
-	InternalSetShaderTexture(NewTexture, FOpenGL::GetFirstComputeTextureUnit() + TextureIndex, NewTexture->Target, NewTexture->Resource, NewTextureRHI->GetNumMips(), -1);
+	if (NewTexture)
+	{
+		InternalSetShaderTexture(NewTexture, FOpenGL::GetFirstComputeTextureUnit() + TextureIndex, NewTexture->Target, NewTexture->Resource, NewTextureRHI->GetNumMips(), -1);
+	}
+	else
+	{
+		InternalSetShaderTexture(nullptr, FOpenGL::GetFirstComputeTextureUnit() + TextureIndex, 0, 0, 0, -1);
+	}
 }
 
 void FOpenGLDynamicRHI::RHISetShaderSampler(FPixelShaderRHIParamRef PixelShaderRHI,uint32 SamplerIndex,FSamplerStateRHIParamRef NewStateRHI)
@@ -2183,7 +2225,14 @@ FORCEINLINE uint32 GetFirstTextureUnit()
 template <EShaderFrequency Frequency>
 FORCEINLINE void SetResource(FOpenGLDynamicRHI* RESTRICT OpenGLRHI, uint32 BindIndex, FOpenGLTextureBase* RESTRICT Texture)
 {
-	OpenGLRHI->InternalSetShaderTexture(Texture, GetFirstTextureUnit<Frequency>() + BindIndex, Texture->Target, Texture->Resource, Texture->NumMips, -1);
+	if (Texture)
+	{
+		OpenGLRHI->InternalSetShaderTexture(Texture, GetFirstTextureUnit<Frequency>() + BindIndex, Texture->Target, Texture->Resource, Texture->NumMips, -1);
+	}
+	else
+	{
+		OpenGLRHI->InternalSetShaderTexture(Texture, GetFirstTextureUnit<Frequency>() + BindIndex, 0, 0, 0, -1);
+	}
 }
 
 template <EShaderFrequency Frequency>
