Build fix - It seems like HTML5 build is picking up thirdparty OpenGL headers even when OpenGL.build.cs says not so. Putting in a temporary fix for now - avoid OpenGL for HTML5 at higher level.

diff --git a/Engine/Source/Runtime/OpenGLDrv/OpenGLDrv.Build.cs b/Engine/Source/Runtime/OpenGLDrv/OpenGLDrv.Build.cs
--- a/Engine/Source/Runtime/OpenGLDrv/OpenGLDrv.Build.cs
+++ b/Engine/Source/Runtime/OpenGLDrv/OpenGLDrv.Build.cs
@@ -22,7 +22,8 @@ public class OpenGLDrv : ModuleRules
 		PrivateIncludePathModuleNames.Add("ImageWrapper");
 		DynamicallyLoadedModuleNames.Add("ImageWrapper");
 
-		AddThirdPartyPrivateStaticDependencies(Target, "OpenGL");
+        if (Target.Platform != UnrealTargetPlatform.HTML5)
+		    AddThirdPartyPrivateStaticDependencies(Target, "OpenGL");
 
 		if (Target.Platform == UnrealTargetPlatform.HTML5 && Target.Architecture == "-win32")
 		{
