Fix broken text rendering in ShooterGame on OS X - specifying GL_R8 as the SRGB internal format isn't behaving itself on OS X so for now just use SRGB8 as we did before.

diff --git a/Engine/Source/Runtime/OpenGLDrv/Private/OpenGLDevice.cpp b/Engine/Source/Runtime/OpenGLDrv/Private/OpenGLDevice.cpp
--- a/Engine/Source/Runtime/OpenGLDrv/Private/OpenGLDevice.cpp
+++ b/Engine/Source/Runtime/OpenGLDrv/Private/OpenGLDevice.cpp
@@ -711,8 +711,11 @@ static void InitRHICapabilitiesForGL()
 	{
 		// Not supported for rendering:
 		SetupTextureFormat( PF_G16,				FOpenGLTextureFormat( GL_R16,					GL_R16,					GL_RED,			GL_UNSIGNED_SHORT,					false,	false));
-
+#if PLATFORM_MAC // @todo On OS X specifying GL_R8 as the internal SRGB format results in incorrect rendering, so specifying SRGB8 as before for now
+		SetupTextureFormat( PF_G8,				FOpenGLTextureFormat( GL_R8,					GL_SRGB8,				GL_RED,			GL_UNSIGNED_BYTE,					false,	false));
+#else
 		SetupTextureFormat( PF_G8,				FOpenGLTextureFormat( GL_R8,					GL_R8,					GL_RED,			GL_UNSIGNED_BYTE,					false,	false));
+#endif
 		SetupTextureFormat( PF_R32_FLOAT,		FOpenGLTextureFormat( GL_R32F,					GL_R32F,				GL_RED,			GL_FLOAT,							false,	false));
 		SetupTextureFormat( PF_G16R16F,			FOpenGLTextureFormat( GL_RG16F,					GL_RG16F,				GL_RG,			GL_HALF_FLOAT,						false,	false));
 		SetupTextureFormat( PF_G16R16F_FILTER,	FOpenGLTextureFormat( GL_RG16F,					GL_RG16F,				GL_RG,			GL_HALF_FLOAT,						false,	false));
