Updated PF_G8 to SRGB8 for all desktop platforms. (Fixes bad text rendering on Mac)

diff --git a/Engine/Source/Runtime/OpenGLDrv/Private/OpenGLDevice.cpp b/Engine/Source/Runtime/OpenGLDrv/Private/OpenGLDevice.cpp
--- a/Engine/Source/Runtime/OpenGLDrv/Private/OpenGLDevice.cpp
+++ b/Engine/Source/Runtime/OpenGLDrv/Private/OpenGLDevice.cpp
@@ -711,11 +711,6 @@ static void InitRHICapabilitiesForGL()
 	{
 		// Not supported for rendering:
 		SetupTextureFormat( PF_G16,				FOpenGLTextureFormat( GL_R16,					GL_R16,					GL_RED,			GL_UNSIGNED_SHORT,					false,	false));
-#if PLATFORM_MAC // @todo On OS X specifying GL_R8 as the internal SRGB format results in incorrect rendering, so specifying SRGB8 as before for now
-		SetupTextureFormat( PF_G8,				FOpenGLTextureFormat( GL_R8,					GL_SRGB8,				GL_RED,			GL_UNSIGNED_BYTE,					false,	false));
-#else
-		SetupTextureFormat( PF_G8,				FOpenGLTextureFormat( GL_R8,					GL_R8,					GL_RED,			GL_UNSIGNED_BYTE,					false,	false));
-#endif
 		SetupTextureFormat( PF_R32_FLOAT,		FOpenGLTextureFormat( GL_R32F,					GL_R32F,				GL_RED,			GL_FLOAT,							false,	false));
 		SetupTextureFormat( PF_G16R16F,			FOpenGLTextureFormat( GL_RG16F,					GL_RG16F,				GL_RG,			GL_HALF_FLOAT,						false,	false));
 		SetupTextureFormat( PF_G16R16F_FILTER,	FOpenGLTextureFormat( GL_RG16F,					GL_RG16F,				GL_RG,			GL_HALF_FLOAT,						false,	false));
@@ -736,6 +731,7 @@ static void InitRHICapabilitiesForGL()
 		SetupTextureFormat( PF_FloatR11G11B10,	FOpenGLTextureFormat( GL_R11F_G11F_B10F,		GL_R11F_G11F_B10F,		GL_RGB,			GL_UNSIGNED_INT_10F_11F_11F_REV,	false,	false));
 		if (FOpenGL::GetShaderPlatform() == EShaderPlatform::SP_OPENGL_ES31_EXT)
 		{
+			SetupTextureFormat(PF_G8, FOpenGLTextureFormat(GL_R8, GL_R8, GL_RED, GL_UNSIGNED_BYTE, false, false));
 			SetupTextureFormat(PF_B8G8R8A8, FOpenGLTextureFormat(GL_RGBA8, GL_SRGB8_ALPHA8, GL_RGBA, GL_UNSIGNED_BYTE, false, true));
 			SetupTextureFormat(PF_R8G8B8A8, FOpenGLTextureFormat(GL_RGBA8, GL_SRGB8_ALPHA8, GL_RGBA, GL_UNSIGNED_BYTE, false, false));
 			// GL_RG16 is not supported in OpenGL ES. There is currently no extension for it either. The user should check for support and implement a fallback.
@@ -743,6 +739,7 @@ static void InitRHICapabilitiesForGL()
 		}
 		else
 		{
+			SetupTextureFormat(PF_G8, FOpenGLTextureFormat(GL_R8, GL_SRGB8, GL_RED, GL_UNSIGNED_BYTE, false, false));
 			SetupTextureFormat(PF_B8G8R8A8, FOpenGLTextureFormat(GL_RGBA8, GL_SRGB8_ALPHA8, GL_BGRA, GL_UNSIGNED_INT_8_8_8_8_REV, false, false));
 			SetupTextureFormat(PF_R8G8B8A8, FOpenGLTextureFormat(GL_RGBA8, GL_SRGB8_ALPHA8, GL_RGBA, GL_UNSIGNED_INT_8_8_8_8_REV, false, false));
 			SetupTextureFormat(PF_G16R16, FOpenGLTextureFormat(GL_RG16, GL_RG16, GL_RG, GL_UNSIGNED_SHORT, false, false));
