UE4 - Fix for some Android devices shader compile issue (TTP 347000)

diff --git a/Engine/Source/Runtime/OpenGLDrv/Private/OpenGLShaders.cpp b/Engine/Source/Runtime/OpenGLDrv/Private/OpenGLShaders.cpp
--- a/Engine/Source/Runtime/OpenGLDrv/Private/OpenGLShaders.cpp
+++ b/Engine/Source/Runtime/OpenGLDrv/Private/OpenGLShaders.cpp
@@ -373,6 +373,7 @@ ShaderType* CompileOpenGLShader(const TArray<uint8>& Code)
 						"#define texture2D texture				\n"
 						"#define texture2DProj textureProj		\n"
 						"#define texture2DLod textureLod			\n"
+						"#define texture2DLodEXT textureLod \n"
 						"#define texture2DProjLod textureProjLod \n"
 						"#define textureCube texture				\n"
 						"#define textureCubeLod textureLod		\n"
