Prevent SlateRHIResourceManager from deallocating its RHI resources directly if the RHI has already been terminated and also prevent OpenGLTexture from trying to access the same deallocated RHI instance on destruction. Fixes crashes on shutdown in ShooterGame.

diff --git a/Engine/Source/Runtime/OpenGLDrv/Private/OpenGLTexture.cpp b/Engine/Source/Runtime/OpenGLDrv/Private/OpenGLTexture.cpp
--- a/Engine/Source/Runtime/OpenGLDrv/Private/OpenGLTexture.cpp
+++ b/Engine/Source/Runtime/OpenGLDrv/Private/OpenGLTexture.cpp
@@ -1385,23 +1385,26 @@ void TOpenGLTexture<RHIResourceType>::CloneViaPBO( TOpenGLTexture<RHIResourceTyp
 template<typename RHIResourceType>
 TOpenGLTexture<RHIResourceType>::~TOpenGLTexture()
 {
-	VERIFY_GL_SCOPE();
+	if ( GIsRHIInitialized )
+	{
+		VERIFY_GL_SCOPE();
 
-	OpenGLTextureDeleted( this );
+		OpenGLTextureDeleted( this );
 
-	if( Resource != 0 )
-	{
-		OpenGLRHI->InvalidateTextureResourceInCache( Resource );
-		FOpenGL::DeleteTextures( 1, &Resource );
-	}
-	
-	if(TextureRange)
-	{
-		delete [] TextureRange;
-		TextureRange = nullptr;
-	}
+		if( Resource != 0 )
+		{
+			OpenGLRHI->InvalidateTextureResourceInCache( Resource );
+			FOpenGL::DeleteTextures( 1, &Resource );
+		}
+		
+		if(TextureRange)
+		{
+			delete [] TextureRange;
+			TextureRange = nullptr;
+		}
 
-	ReleaseOpenGLFramebuffers(OpenGLRHI,this);
+		ReleaseOpenGLFramebuffers(OpenGLRHI,this);
+	}
 }
 
 /*-----------------------------------------------------------------------------
diff --git a/Engine/Source/Runtime/SlateRHIRenderer/Private/SlateRHIResourceManager.cpp b/Engine/Source/Runtime/SlateRHIRenderer/Private/SlateRHIResourceManager.cpp
--- a/Engine/Source/Runtime/SlateRHIRenderer/Private/SlateRHIResourceManager.cpp
+++ b/Engine/Source/Runtime/SlateRHIRenderer/Private/SlateRHIResourceManager.cpp
@@ -130,7 +130,10 @@ FSlateRHIResourceManager::FSlateRHIResourceManager()
 
 FSlateRHIResourceManager::~FSlateRHIResourceManager()
 {
-	DeleteResources();
+	if ( GIsRHIInitialized )
+	{
+		DeleteResources();
+	}
 }
 
 void FSlateRHIResourceManager::CreateTextures( const TArray< const FSlateBrush* >& Resources )
