Remove the GenBuffers & DeleteBuffers overrides on Mac OS X, though they were much faster they are crashing 10.10 and that must to be fixed for 4.5.

diff --git a/Engine/Source/Runtime/OpenGLDrv/Private/Mac/OpenGLMac.cpp b/Engine/Source/Runtime/OpenGLDrv/Private/Mac/OpenGLMac.cpp
--- a/Engine/Source/Runtime/OpenGLDrv/Private/Mac/OpenGLMac.cpp
+++ b/Engine/Source/Runtime/OpenGLDrv/Private/Mac/OpenGLMac.cpp
@@ -57,18 +57,8 @@ static FAutoConsoleVariableRef CVarMacMaxTexturesToDeletePerFrame(
 	ECVF_RenderThreadSafe
 	);
 
-static int32 GMacMinBuffersToDelete = OPENGL_NAME_CACHE_SIZE;
-static FAutoConsoleVariableRef CVarMacMinBuffersToDelete(
-	TEXT("r.Mac.MinBuffersToDelete"),
-	GMacMinBuffersToDelete,
-	TEXT("Specifies how many buffers are required to be waiting before a call to glDeleteBuffers is made in order to amortize the cost. (Default: 1024)"),
-	ECVF_RenderThreadSafe
-	);
-
 static const uint32 GMacTexturePoolNum = 3;
 static TArray<GLuint> GMacTexturesToDelete[GMacTexturePoolNum];
-static TArray<GLuint> GMacBuffers;
-static TArray<GLuint> GMacBuffersToDelete;
 
 /** Used to temporarily disable Cocoa screen updates to make window updates happen only on the render thread. */
 bool GMacEnableCocoaScreenUpdates = true;
@@ -324,12 +314,6 @@ struct FPlatformOpenGLDevice
 					GMacTexturesToDelete[i].Reset();
 				}
 			}
-			
-			glDeleteBuffers(GMacBuffers.Num(), GMacBuffers.GetData());
-			GMacBuffers.Reset();
-			
-			glDeleteBuffers(GMacBuffersToDelete.Num(), GMacBuffersToDelete.GetData());
-			GMacBuffersToDelete.Reset();
 		}
 		
 		FPlatformOpenGLContext::UnregisterGraphicsSwitchingCallback();
@@ -1320,93 +1304,6 @@ bool FMacOpenGL::MustFlushTexStorage(void)
 	return GMacFlushTexStorage && (IsRHIDeviceNVIDIA() || GMacUseMTGL);
 }
 
-void FMacOpenGL::GenBuffers( GLsizei n, GLuint *buffers)
-{
-#if USE_OPENGL_NAME_CACHE
-	if( n < GMacBuffers.Num() )
-	{
-		FMemory::Memcpy( buffers, GMacBuffers.GetData(), sizeof(GLuint)*n);
-		GMacBuffers.RemoveAt(0, n, false);
-	}
-	else if(n < OPENGL_NAME_CACHE_SIZE)
-	{
-		GLsizei Leftover = GMacBuffers.Num();
-		if(Leftover)
-		{
-			FMemory::Memcpy( buffers, GMacBuffers.GetData(), sizeof(GLuint)*Leftover);
-			GMacBuffers.RemoveAt(0, Leftover, false);
-		}
-		
-		GMacBuffers.Reserve(OPENGL_NAME_CACHE_SIZE);
-		GLuint Buffer[OPENGL_NAME_CACHE_SIZE] = {0};
-		glGenBuffers( OPENGL_NAME_CACHE_SIZE, Buffer);
-		GMacBuffers.Append(Buffer, OPENGL_NAME_CACHE_SIZE);
-		
-		n -= Leftover;
-		buffers += Leftover;
-		
-		FMemory::Memcpy( buffers, GMacBuffers.GetData(), sizeof(GLuint)*n);
-		GMacBuffers.RemoveAt(0, n, false);
-	}
-	else
-	{
-		glGenBuffers(n, buffers);
-	}
-#else
-	glGenBuffers( n, buffers);
-#endif
-}
-
-void FMacOpenGL::DeleteBuffers(GLsizei Number, const GLuint* Buffers)
-{
-#if UE_BUILD_DEBUG
-	GLint Buffer = 0;
-	glGetIntegerv(GL_COPY_WRITE_BUFFER, &Buffer);
-#endif
-	
-	uint32 Index = 0;
-	
-#if USE_OPENGL_NAME_CACHE
-	while ( GMacBuffers.Num() < USE_OPENGL_NAME_CACHE && Index > Number )
-	{
-#if UE_BUILD_DEBUG
-		glBindBuffer(GL_COPY_WRITE_BUFFER, Buffers[Index]);
-		GLint bLocked = GL_FALSE;
-		glGetBufferParameteriv(GL_COPY_WRITE_BUFFER, GL_BUFFER_MAPPED, &bLocked);
-		check(!bLocked);
-#endif
-		
-		GMacBuffers.Add(Buffers[Index]);
-		Index++;
-	}
-#endif
-	
-	if( Number > Index )
-	{
-		GMacBuffersToDelete.Append(&Buffers[Index], Number - Index);
-	}
-	
-	if(GMacBuffersToDelete.Num() >= GMacMinBuffersToDelete)
-	{
-		glDeleteBuffers(GMacBuffersToDelete.Num(), GMacBuffersToDelete.GetData());
-		GMacBuffersToDelete.Empty();
-	}
-#if UE_BUILD_DEBUG
-	else
-	{
-		for(uint32 i = Index; i < Number; i++)
-		{
-			glBindBuffer(GL_COPY_WRITE_BUFFER, Buffers[i]);
-			GLint bLocked = GL_FALSE;
-			glGetBufferParameteriv(GL_COPY_WRITE_BUFFER, GL_BUFFER_MAPPED, &bLocked);
-			check(!bLocked);
-		}
-	}
-	
-	glBindBuffer(GL_COPY_WRITE_BUFFER, Buffer);
-#endif
-}
-
 void FMacOpenGL::DeleteTextures(GLsizei Number, const GLuint* Textures)
 {
 	TArray<GLuint>& TexturesToDelete = GMacTexturesToDelete[GFrameNumberRenderThread % GMacTexturePoolNum];
diff --git a/Engine/Source/Runtime/OpenGLDrv/Private/Mac/OpenGLMac.h b/Engine/Source/Runtime/OpenGLDrv/Private/Mac/OpenGLMac.h
--- a/Engine/Source/Runtime/OpenGLDrv/Private/Mac/OpenGLMac.h
+++ b/Engine/Source/Runtime/OpenGLDrv/Private/Mac/OpenGLMac.h
@@ -119,10 +119,6 @@ struct FMacOpenGL : public FOpenGL3
 		}
 	}
 	
-	static void GenBuffers( GLsizei n, GLuint *buffers);
-	
-	static void DeleteBuffers(GLsizei Number, const GLuint* Buffers);
-	
 	static void DeleteTextures(GLsizei Number, const GLuint* Textures);
 
 	static FORCEINLINE bool SupportsSeparateAlphaBlend()				{ return bSupportsDrawBuffersBlend; }
