TArray documentation and a coupld of methods deprecation.

diff --git a/Engine/Source/Runtime/Online/BuildPatchServices/Private/BuildPatchDownloader.cpp b/Engine/Source/Runtime/Online/BuildPatchServices/Private/BuildPatchDownloader.cpp
--- a/Engine/Source/Runtime/Online/BuildPatchServices/Private/BuildPatchDownloader.cpp
+++ b/Engine/Source/Runtime/Online/BuildPatchServices/Private/BuildPatchDownloader.cpp
@@ -239,7 +239,7 @@ uint32 FBuildPatchDownloader::Run()
 				while ( BytesRead < FileSize )
 				{
 					const int64 ReadLen = FMath::Min<int64>( BytesPerCall, FileSize - BytesRead );
-					Reader->Serialize( FileDataArray.GetTypedData() + BytesRead, ReadLen );
+					Reader->Serialize( FileDataArray.GetData() + BytesRead, ReadLen );
 					BytesRead += ReadLen;
 					OnDownloadProgress( NextGuid, BytesRead );
 				}
diff --git a/Engine/Source/Runtime/Online/HTTP/Private/Curl/CurlHttp.cpp b/Engine/Source/Runtime/Online/HTTP/Private/Curl/CurlHttp.cpp
--- a/Engine/Source/Runtime/Online/HTTP/Private/Curl/CurlHttp.cpp
+++ b/Engine/Source/Runtime/Online/HTTP/Private/Curl/CurlHttp.cpp
@@ -201,7 +201,7 @@ void FCurlHttpRequest::SetContentAsString(const FString& ContentString)
 {
 	FTCHARToUTF8 Converter(*ContentString);
 	RequestPayload.SetNum(Converter.Length());
-	FMemory::Memcpy(RequestPayload.GetTypedData(), (uint8*)(ANSICHAR*)Converter.Get(), RequestPayload.Num());
+	FMemory::Memcpy(RequestPayload.GetData(), (uint8*)(ANSICHAR*)Converter.Get(), RequestPayload.Num());
 }
 
 void FCurlHttpRequest::SetHeader(const FString& HeaderName, const FString& HeaderValue)
@@ -486,7 +486,7 @@ bool FCurlHttpRequest::StartRequest()
 		// (if we pass, don't check here and trust the request)
 		check(!GetHeader("Content-Type").IsEmpty() || IsURLEncoded(RequestPayload));
 
-		curl_easy_setopt(EasyHandle, CURLOPT_POSTFIELDS, RequestPayload.GetTypedData());
+		curl_easy_setopt(EasyHandle, CURLOPT_POSTFIELDS, RequestPayload.GetData());
 		curl_easy_setopt(EasyHandle, CURLOPT_POSTFIELDSIZE, RequestPayload.Num());
 	}
 	else if (Verb == TEXT("PUT"))
@@ -517,7 +517,7 @@ bool FCurlHttpRequest::StartRequest()
 		// (if we pass, don't check here and trust the request)
 		check(!GetHeader("Content-Type").IsEmpty() || IsURLEncoded(RequestPayload));
 
-		curl_easy_setopt(EasyHandle, CURLOPT_POSTFIELDS, RequestPayload.GetTypedData());
+		curl_easy_setopt(EasyHandle, CURLOPT_POSTFIELDS, RequestPayload.GetData());
 		curl_easy_setopt(EasyHandle, CURLOPT_POSTFIELDSIZE, RequestPayload.Num());
 	}
 	else
diff --git a/Engine/Source/Runtime/Online/HTTP/Private/Windows/HttpWinInet.cpp b/Engine/Source/Runtime/Online/HTTP/Private/Windows/HttpWinInet.cpp
--- a/Engine/Source/Runtime/Online/HTTP/Private/Windows/HttpWinInet.cpp
+++ b/Engine/Source/Runtime/Online/HTTP/Private/Windows/HttpWinInet.cpp
@@ -51,7 +51,7 @@ FString InternetTranslateError(::DWORD GetLastErrorResult)
 		ExtLength = ExtLength+1;
 		TArray<TCHAR> ExtErrMsg;
 		ExtErrMsg.AddUninitialized(ExtLength);
-		if (!InternetGetLastResponseInfo(&InetError, ExtErrMsg.GetTypedData(), &ExtLength))
+		if (!InternetGetLastResponseInfo(&InetError, ExtErrMsg.GetData(), &ExtLength))
 		{
 			ErrorStr += FString::Printf(TEXT("Call to InternetGetLastResponseInfo() failed: %08X. "), 
 				(uint32) GetLastError());
@@ -419,7 +419,7 @@ void FHttpRequestWinInet::SetContentAsString(const FString& ContentString)
 {
 	FTCHARToUTF8 Converter(*ContentString);
 	RequestPayload.SetNumUninitialized(Converter.Length());
-	FMemory::Memcpy(RequestPayload.GetTypedData(), (const uint8*)Converter.Get(), RequestPayload.Num());
+	FMemory::Memcpy(RequestPayload.GetData(), (const uint8*)Converter.Get(), RequestPayload.Num());
 }
 
 void FHttpRequestWinInet::SetHeader(const FString& HeaderName, const FString& HeaderValue)
@@ -595,7 +595,7 @@ bool FHttpRequestWinInet::StartRequest()
 		RequestHandle, 
 		*Headers, 
 		Headers.Len(), 
-		RequestPayload.Num() > 0 ? RequestPayload.GetTypedData() : NULL,
+		RequestPayload.Num() > 0 ? RequestPayload.GetData() : NULL,
 		RequestPayload.Num());
 
 	if (!bSentRequest &&
@@ -1009,13 +1009,13 @@ void FHttpResponseWinInet::ProcessResponseHeaders()
 		}
 		TArray<TCHAR> HeaderBuffer;
 		HeaderBuffer.AddUninitialized(HeaderSize/sizeof(TCHAR));
-		if (!HttpQueryInfo(Request.RequestHandle, HTTP_QUERY_RAW_HEADERS_CRLF, HeaderBuffer.GetTypedData(), &HeaderSize, NULL))
+		if (!HttpQueryInfo(Request.RequestHandle, HTTP_QUERY_RAW_HEADERS_CRLF, HeaderBuffer.GetData(), &HeaderSize, NULL))
 		{
 			UE_LOG(LogHttp, Warning, TEXT("HttpQueryInfo for all headers failed: %s. %p"), 
 				*InternetTranslateError(GetLastError()), &Request);
 		}
 		// parse all the key/value pairs
-		const TCHAR* HeaderPtr = HeaderBuffer.GetTypedData();
+		const TCHAR* HeaderPtr = HeaderBuffer.GetData();
 		// don't count the terminating NULL character as one to search.
 		const TCHAR* EndPtr = HeaderPtr + HeaderBuffer.Num()-1;
 		while (HeaderPtr < EndPtr)
diff --git a/Engine/Source/Runtime/Online/OnlineSubsystemUtils/Private/Tests/TestVoice.cpp b/Engine/Source/Runtime/Online/OnlineSubsystemUtils/Private/Tests/TestVoice.cpp
--- a/Engine/Source/Runtime/Online/OnlineSubsystemUtils/Private/Tests/TestVoice.cpp
+++ b/Engine/Source/Runtime/Online/OnlineSubsystemUtils/Private/Tests/TestVoice.cpp
@@ -93,7 +93,7 @@ void FTestVoice::SetStaticVoiceData(TArray<uint8>& VoiceData, uint32& TotalVoice
 		VoiceData.Empty(TotalVoiceBytes);
 		VoiceData.AddUninitialized(TotalVoiceBytes);
 
-		FMemory::Memcpy(VoiceData.GetTypedData(), RawVoiceTestData, ARRAY_COUNT(RawVoiceTestData));
+		FMemory::Memcpy(VoiceData.GetData(), RawVoiceTestData, ARRAY_COUNT(RawVoiceTestData));
 
 		LastQueueTime = CurrentTime;
 		bTimeToQueue = false;
@@ -144,10 +144,10 @@ bool FTestVoice::Tick(float DeltaTime)
 
 					if (LastRemainderSize > 0)
 					{
-						FMemory::Memcpy(RawCaptureData.GetTypedData(), Remainder.GetTypedData(), LastRemainderSize);
+						FMemory::Memcpy(RawCaptureData.GetData(), Remainder.GetData(), LastRemainderSize);
 					}
 
-					MicState = VoiceCapture->GetVoiceData(RawCaptureData.GetTypedData() + LastRemainderSize, NewVoiceDataBytes, NewVoiceDataBytes);
+					MicState = VoiceCapture->GetVoiceData(RawCaptureData.GetData() + LastRemainderSize, NewVoiceDataBytes, NewVoiceDataBytes);
 					TotalVoiceBytes = NewVoiceDataBytes + LastRemainderSize;
 					bDoWork = MicState == EVoiceCaptureState::Ok;
 				}
@@ -158,24 +158,24 @@ bool FTestVoice::Tick(float DeltaTime)
 				// ZERO INPUT
 				if (bZeroInput)
 				{
-					FMemory::Memzero(RawCaptureData.GetTypedData(), TotalVoiceBytes);
+					FMemory::Memzero(RawCaptureData.GetData(), TotalVoiceBytes);
 				}	
 				// ZERO INPUT END
 
 				// COMPRESSION BEGIN
 				uint32 CompressedDataSize = MaxCompressedDataSize;
-				LastRemainderSize = VoiceEncoder->Encode(RawCaptureData.GetTypedData(), TotalVoiceBytes, CompressedData.GetTypedData(), CompressedDataSize);
+				LastRemainderSize = VoiceEncoder->Encode(RawCaptureData.GetData(), TotalVoiceBytes, CompressedData.GetData(), CompressedDataSize);
 
 				if (LastRemainderSize > 0)
 				{
-					FMemory::Memcpy(Remainder.GetTypedData(), RawCaptureData.GetTypedData() + (TotalVoiceBytes - LastRemainderSize), LastRemainderSize);
+					FMemory::Memcpy(Remainder.GetData(), RawCaptureData.GetData() + (TotalVoiceBytes - LastRemainderSize), LastRemainderSize);
 				}
 				// COMPRESION END
 
 				// DECOMPRESION BEGIN
 				uint32 UncompressedDataSize = MaxUncompressedDataSize;
-				VoiceDecoder->Decode(CompressedData.GetTypedData(), CompressedDataSize, 
-					UncompressedData.GetTypedData(), UncompressedDataSize);
+				VoiceDecoder->Decode(CompressedData.GetData(), CompressedDataSize, 
+					UncompressedData.GetData(), UncompressedDataSize);
 				// DECOMPRESSION END
 
 				if (bUseDecompressed)
@@ -188,13 +188,13 @@ bool FTestVoice::Tick(float DeltaTime)
 							FMemory::Memzero((uint8*)UncompressedData.GetData(), UncompressedDataSize);
 						}
 
-						SoundStreaming->QueueAudio(UncompressedData.GetTypedData(), UncompressedDataSize);
+						SoundStreaming->QueueAudio(UncompressedData.GetData(), UncompressedDataSize);
 					}
 				}
 				else
 				{
 					//UE_LOG(LogVoice, Log, TEXT("Queueing raw data! %d"), TotalVoiceBytes - LastRemainderSize);
-					SoundStreaming->QueueAudio(RawCaptureData.GetTypedData(), TotalVoiceBytes - LastRemainderSize);
+					SoundStreaming->QueueAudio(RawCaptureData.GetData(), TotalVoiceBytes - LastRemainderSize);
 				}
 			}
 
diff --git a/Engine/Source/Runtime/Online/OnlineSubsystemUtils/Private/VoiceEngineImpl.cpp b/Engine/Source/Runtime/Online/OnlineSubsystemUtils/Private/VoiceEngineImpl.cpp
--- a/Engine/Source/Runtime/Online/OnlineSubsystemUtils/Private/VoiceEngineImpl.cpp
+++ b/Engine/Source/Runtime/Online/OnlineSubsystemUtils/Private/VoiceEngineImpl.cpp
@@ -264,11 +264,11 @@ uint32 FVoiceEngineImpl::ReadLocalVoiceData(uint32 LocalUserNum, uint8* Data, ui
 
 		if (PlayerVoiceData[LocalUserNum].VoiceRemainderSize > 0)
 		{
-			FMemory::Memcpy(DecompressedVoiceBuffer.GetTypedData(), PlayerVoiceData[LocalUserNum].VoiceRemainder.GetTypedData(), PlayerVoiceData[LocalUserNum].VoiceRemainderSize);
+			FMemory::Memcpy(DecompressedVoiceBuffer.GetData(), PlayerVoiceData[LocalUserNum].VoiceRemainder.GetData(), PlayerVoiceData[LocalUserNum].VoiceRemainderSize);
 		}
 
 		// Get new uncompressed data
-		VoiceResult = VoiceCapture->GetVoiceData(DecompressedVoiceBuffer.GetTypedData() + PlayerVoiceData[LocalUserNum].VoiceRemainderSize, 
+		VoiceResult = VoiceCapture->GetVoiceData(DecompressedVoiceBuffer.GetData() + PlayerVoiceData[LocalUserNum].VoiceRemainderSize,
 			NewVoiceDataBytes, NewVoiceDataBytes);
 
 		TotalVoiceBytes = NewVoiceDataBytes + PlayerVoiceData[LocalUserNum].VoiceRemainderSize;
@@ -279,7 +279,7 @@ uint32 FVoiceEngineImpl::ReadLocalVoiceData(uint32 LocalUserNum, uint8* Data, ui
 			CompressedVoiceBuffer.AddUninitialized(MAX_COMPRESSED_VOICE_BUFFER_SIZE);
 
 			PlayerVoiceData[LocalUserNum].VoiceRemainderSize =
-				VoiceEncoder->Encode(DecompressedVoiceBuffer.GetTypedData(), TotalVoiceBytes, CompressedVoiceBuffer.GetTypedData(), CompressedBytesAvailable);
+				VoiceEncoder->Encode(DecompressedVoiceBuffer.GetData(), TotalVoiceBytes, CompressedVoiceBuffer.GetData(), CompressedBytesAvailable);
 
 			// Save off any unencoded remainder
 			if (PlayerVoiceData[LocalUserNum].VoiceRemainderSize > 0)
@@ -291,7 +291,7 @@ uint32 FVoiceEngineImpl::ReadLocalVoiceData(uint32 LocalUserNum, uint8* Data, ui
 				}
 
 				PlayerVoiceData[LocalUserNum].VoiceRemainder.AddUninitialized(MAX_VOICE_REMAINDER_SIZE);
-				FMemory::Memcpy(PlayerVoiceData[LocalUserNum].VoiceRemainder.GetTypedData(), DecompressedVoiceBuffer.GetTypedData() + (TotalVoiceBytes - PlayerVoiceData[LocalUserNum].VoiceRemainderSize), PlayerVoiceData[LocalUserNum].VoiceRemainderSize);
+				FMemory::Memcpy(PlayerVoiceData[LocalUserNum].VoiceRemainder.GetData(), DecompressedVoiceBuffer.GetData() + (TotalVoiceBytes - PlayerVoiceData[LocalUserNum].VoiceRemainderSize), PlayerVoiceData[LocalUserNum].VoiceRemainderSize);
 			}
 
 			static double LastGetVoiceCallTime = 0.0;
@@ -344,7 +344,7 @@ uint32 FVoiceEngineImpl::SubmitRemoteVoiceData(const FUniqueNetId& RemoteTalkerI
 
 	DecompressedVoiceBuffer.Empty(MAX_UNCOMPRESSED_VOICE_BUFFER_SIZE);
 	DecompressedVoiceBuffer.AddUninitialized(MAX_UNCOMPRESSED_VOICE_BUFFER_SIZE);
-	QueuedData->VoiceDecoder->Decode(Data, *Size, DecompressedVoiceBuffer.GetTypedData(), BytesWritten);
+	QueuedData->VoiceDecoder->Decode(Data, *Size, DecompressedVoiceBuffer.GetData(), BytesWritten);
 
 	// If there is no data, return
 	if (BytesWritten <= 0)
diff --git a/Engine/Source/Runtime/Online/Voice/Private/Windows/VoiceCaptureWindows.cpp b/Engine/Source/Runtime/Online/Voice/Private/Windows/VoiceCaptureWindows.cpp
--- a/Engine/Source/Runtime/Online/Voice/Private/Windows/VoiceCaptureWindows.cpp
+++ b/Engine/Source/Runtime/Online/Voice/Private/Windows/VoiceCaptureWindows.cpp
@@ -299,7 +299,7 @@ void FVoiceCaptureWindows::ProcessData()
 			UncompressedAudioBuffer.Empty(MAX_UNCOMPRESSED_VOICE_BUFFER_SIZE);
 			UncompressedAudioBuffer.AddUninitialized(CaptureLength + CaptureLength2);
 
-			uint8* AudioBuffer = UncompressedAudioBuffer.GetTypedData();
+			uint8* AudioBuffer = UncompressedAudioBuffer.GetData();
 			FMemory::Memcpy(AudioBuffer, CaptureData, CaptureLength);
 
 			if (CaptureData2 && CaptureLength2 > 0)
@@ -348,7 +348,7 @@ EVoiceCaptureState::Type FVoiceCaptureWindows::GetVoiceData(uint8* OutVoiceBuffe
 		OutAvailableVoiceData = UncompressedAudioBuffer.Num();
 		if (InVoiceBufferSize >= OutAvailableVoiceData)
 		{
-			FMemory::Memcpy(OutVoiceBuffer, UncompressedAudioBuffer.GetTypedData(), OutAvailableVoiceData);
+			FMemory::Memcpy(OutVoiceBuffer, UncompressedAudioBuffer.GetData(), OutAvailableVoiceData);
 			VoiceCaptureState = EVoiceCaptureState::NoData;
 		}
 		else
diff --git a/Engine/Source/Runtime/OpenGLDrv/Private/OpenGLDebugFrameDump.cpp b/Engine/Source/Runtime/OpenGLDrv/Private/OpenGLDebugFrameDump.cpp
--- a/Engine/Source/Runtime/OpenGLDrv/Private/OpenGLDebugFrameDump.cpp
+++ b/Engine/Source/Runtime/OpenGLDrv/Private/OpenGLDebugFrameDump.cpp
@@ -43,7 +43,7 @@ void appCreatePNGWithAlpha( const TCHAR* File, int32 Width, int32 Height, FColor
 
 		const TArray<uint8>& CompressedData = ImageWrapper->GetCompressed();
 		int32 CompressedSize = CompressedData.Num();
-		Ar->Serialize( (void*)CompressedData.GetTypedData(), CompressedSize );
+		Ar->Serialize( (void*)CompressedData.GetData(), CompressedSize );
 		delete Ar;
 	}
 }
diff --git a/Engine/Source/Runtime/OpenGLDrv/Private/OpenGLDrv.cpp b/Engine/Source/Runtime/OpenGLDrv/Private/OpenGLDrv.cpp
--- a/Engine/Source/Runtime/OpenGLDrv/Private/OpenGLDrv.cpp
+++ b/Engine/Source/Runtime/OpenGLDrv/Private/OpenGLDrv.cpp
@@ -274,7 +274,7 @@ void FOpenGLGPUProfiler::EndFrame()
 				{
 					GPUHitchEventNodeFrames.RemoveAt(0);
 				}
-				GPUHitchEventNodeFrames.AddRawItem((FOpenGLEventNodeFrame*)CurrentEventNodeFrame);
+				GPUHitchEventNodeFrames.Add((FOpenGLEventNodeFrame*)CurrentEventNodeFrame);
 				CurrentEventNodeFrame = NULL;  // prevent deletion of this below; ke kept it in the history
 			}
 		}
diff --git a/Engine/Source/Runtime/OpenGLDrv/Private/OpenGLRenderTarget.cpp b/Engine/Source/Runtime/OpenGLDrv/Private/OpenGLRenderTarget.cpp
--- a/Engine/Source/Runtime/OpenGLDrv/Private/OpenGLRenderTarget.cpp
+++ b/Engine/Source/Runtime/OpenGLDrv/Private/OpenGLRenderTarget.cpp
@@ -704,7 +704,7 @@ void FOpenGLDynamicRHI::RHIReadSurfaceData(FTextureRHIParamRef TextureRHI,FIntRe
 	OutData.Empty();
 	OutData.AddUninitialized(Size);
 
-	FMemory::Memcpy(OutData.GetTypedData(), Temp.GetTypedData(), Size * sizeof(FColor));
+	FMemory::Memcpy(OutData.GetData(), Temp.GetData(), Size * sizeof(FColor));
 }
 
 void FOpenGLDynamicRHI::RHIMapStagingSurface(FTextureRHIParamRef TextureRHI,void*& OutData,int32& OutWidth,int32& OutHeight)
@@ -827,7 +827,7 @@ void FOpenGLDynamicRHI::RHIRead3DSurfaceFloatData(FTextureRHIParamRef TextureRHI
 
 	// Grab the raw data from the temp texture.
 	glPixelStorei( GL_PACK_ALIGNMENT, 1 );
-	FOpenGL::GetTexImage( GL_TEXTURE_3D, 0, GL_RGBA, GL_HALF_FLOAT, OutData.GetTypedData() );
+	FOpenGL::GetTexImage( GL_TEXTURE_3D, 0, GL_RGBA, GL_HALF_FLOAT, OutData.GetData() );
 	glPixelStorei( GL_PACK_ALIGNMENT, 4 );
 
 	// Clean up
diff --git a/Engine/Source/Runtime/OpenGLDrv/Private/OpenGLShaders.cpp b/Engine/Source/Runtime/OpenGLDrv/Private/OpenGLShaders.cpp
--- a/Engine/Source/Runtime/OpenGLDrv/Private/OpenGLShaders.cpp
+++ b/Engine/Source/Runtime/OpenGLDrv/Private/OpenGLShaders.cpp
@@ -267,7 +267,7 @@ ShaderType* CompileOpenGLShader(const TArray<uint8>& Code)
 	}
 
 	int32 CodeOffset = Ar.Tell();
-	const ANSICHAR* GlslCode = (ANSICHAR*)Code.GetTypedData() + CodeOffset;
+	const ANSICHAR* GlslCode = (ANSICHAR*)Code.GetData() + CodeOffset;
 	GLint GlslCodeLength = Code.Num() - CodeOffset - 1;
 	uint32 GlslCodeCRC = FCrc::MemCrc_DEPRECATED(GlslCode,GlslCodeLength);
 
@@ -1147,7 +1147,7 @@ static void VerifyUniformBufferLayouts(GLuint Program)
 			TArray<GLint> ActiveUniformIndices;
 			ActiveUniformIndices.Init(ActiveUniforms);
 
-			glGetActiveUniformBlockiv(Program, BlockIndex, GL_UNIFORM_BLOCK_ACTIVE_UNIFORM_INDICES, ActiveUniformIndices.GetTypedData());
+			glGetActiveUniformBlockiv(Program, BlockIndex, GL_UNIFORM_BLOCK_ACTIVE_UNIFORM_INDICES, ActiveUniformIndices.GetData());
 			
 			TArray<GLint> ActiveUniformOffsets(ActiveUniformIndices);
 			glGetActiveUniformsiv(Program, ActiveUniforms, reinterpret_cast<const GLuint*>(ActiveUniformIndices.GetData()), GL_UNIFORM_OFFSET, ActiveUniformOffsets.GetData());
@@ -1407,7 +1407,7 @@ FComputeShaderRHIRef FOpenGLDynamicRHI::RHICreateComputeShader(const TArray<uint
 #if DEBUG_GL_SHADERS
 		if (ComputeShader->bSuccessfullyCompiled)
 		{
-			UE_LOG(LogRHI,Error,TEXT("Compute Shader:\n%s"),ANSI_TO_TCHAR(ComputeShader->GlslCode.GetTypedData()));
+			UE_LOG(LogRHI,Error,TEXT("Compute Shader:\n%s"),ANSI_TO_TCHAR(ComputeShader->GlslCode.GetData()));
 		}
 #endif //DEBUG_GL_SHADERS
 		checkf(ComputeShader->LinkedProgram, TEXT("Compute shader failed to compile & link."));
@@ -1571,25 +1571,25 @@ FBoundShaderStateRHIRef FOpenGLDynamicRHI::RHICreateBoundShaderState(
 #if DEBUG_GL_SHADERS
 					if (VertexShader->bSuccessfullyCompiled)
 					{
-						UE_LOG(LogRHI,Error,TEXT("Vertex Shader:\n%s"),ANSI_TO_TCHAR(VertexShader->GlslCode.GetTypedData()));
+						UE_LOG(LogRHI,Error,TEXT("Vertex Shader:\n%s"),ANSI_TO_TCHAR(VertexShader->GlslCode.GetData()));
 					}
 					if (PixelShader->bSuccessfullyCompiled)
 					{
-						UE_LOG(LogRHI,Error,TEXT("Pixel Shader:\n%s"),ANSI_TO_TCHAR(PixelShader->GlslCode.GetTypedData()));
+						UE_LOG(LogRHI,Error,TEXT("Pixel Shader:\n%s"),ANSI_TO_TCHAR(PixelShader->GlslCode.GetData()));
 					}
 					if (GeometryShader && GeometryShader->bSuccessfullyCompiled)
 					{
-						UE_LOG(LogRHI,Error,TEXT("Geometry Shader:\n%s"),ANSI_TO_TCHAR(GeometryShader->GlslCode.GetTypedData()));
+						UE_LOG(LogRHI,Error,TEXT("Geometry Shader:\n%s"),ANSI_TO_TCHAR(GeometryShader->GlslCode.GetData()));
 					}
 					if ( FOpenGL::SupportsTessellation() )
 					{
 						if (HullShader && HullShader->bSuccessfullyCompiled)
 						{
-							UE_LOG(LogRHI,Error,TEXT("Hull Shader:\n%s"),ANSI_TO_TCHAR(HullShader->GlslCode.GetTypedData()));
+							UE_LOG(LogRHI,Error,TEXT("Hull Shader:\n%s"),ANSI_TO_TCHAR(HullShader->GlslCode.GetData()));
 						}
 						if (DomainShader && DomainShader->bSuccessfullyCompiled)
 						{
-							UE_LOG(LogRHI,Error,TEXT("Domain Shader:\n%s"),ANSI_TO_TCHAR(DomainShader->GlslCode.GetTypedData()));
+							UE_LOG(LogRHI,Error,TEXT("Domain Shader:\n%s"),ANSI_TO_TCHAR(DomainShader->GlslCode.GetData()));
 						}
 					}
 #endif //DEBUG_GL_SHADERS
@@ -2033,7 +2033,7 @@ void FOpenGLShaderParameterCache::CommitPackedUniformBuffers(FOpenGLLinkedProgra
 		{
 			const FOpenGLUniformBuffer* UniformBuffer = (FOpenGLUniformBuffer*)RHIUniformBuffers[BufferIndex].GetReference();
 			check(UniformBuffer);
-			const uint32* RESTRICT SourceData = UniformBuffer->EmulatedBufferData->Data.GetTypedData();
+			const uint32* RESTRICT SourceData = UniformBuffer->EmulatedBufferData->Data.GetData();
 			for (int32 InfoIndex = LastInfoIndex; InfoIndex < UniformBuffersCopyInfo.Num(); ++InfoIndex)
 			{
 				const FOpenGLUniformBufferCopyInfo& Info = UniformBuffersCopyInfo[InfoIndex];
@@ -2074,7 +2074,7 @@ void FOpenGLShaderParameterCache::CommitPackedUniformBuffers(FOpenGLLinkedProgra
 					const FOpenGLUniformBufferCopyInfo& Info = UniformBuffersCopyInfo[InfoIndex];
 					if (Info.SourceUBIndex == BufferIndex)
 					{
-						const uint32* RESTRICT SourceData = UniformBuffer->EmulatedBufferData->Data.GetTypedData();
+						const uint32* RESTRICT SourceData = UniformBuffer->EmulatedBufferData->Data.GetData();
 						SourceData += Info.SourceOffsetInFloats;
 						float* RESTRICT ScratchMem = (float*)PackedUniformsScratch[Info.DestUBTypeIndex];
 						ScratchMem += Info.DestOffsetInFloats;
diff --git a/Engine/Source/Runtime/OpenGLDrv/Private/OpenGLVertexDeclaration.cpp b/Engine/Source/Runtime/OpenGLDrv/Private/OpenGLVertexDeclaration.cpp
--- a/Engine/Source/Runtime/OpenGLDrv/Private/OpenGLVertexDeclaration.cpp
+++ b/Engine/Source/Runtime/OpenGLDrv/Private/OpenGLVertexDeclaration.cpp
@@ -92,9 +92,9 @@ struct FOpenGLVertexDeclarationKey
 			}
 		};
 		// Sort the FOpenGLVertexElements by stream then offset.
-		Sort( VertexElements.GetTypedData(),InElements.Num(), FCompareFOpenGLVertexElement() );
+		Sort( VertexElements.GetData(),InElements.Num(), FCompareFOpenGLVertexElement() );
 
-		Hash = FCrc::MemCrc_DEPRECATED(VertexElements.GetTypedData(),VertexElements.Num()*sizeof(FOpenGLVertexElement));
+		Hash = FCrc::MemCrc_DEPRECATED(VertexElements.GetData(),VertexElements.Num()*sizeof(FOpenGLVertexElement));
 	}
 };
 
diff --git a/Engine/Source/Runtime/OpenGLDrv/Public/OpenGLShaderResources.h b/Engine/Source/Runtime/OpenGLDrv/Public/OpenGLShaderResources.h
--- a/Engine/Source/Runtime/OpenGLDrv/Public/OpenGLShaderResources.h
+++ b/Engine/Source/Runtime/OpenGLDrv/Public/OpenGLShaderResources.h
@@ -33,7 +33,7 @@ struct FOpenGLShaderResourceTable : public FBaseShaderResourceTable
 		{
 			return false;
 		}
-		if (FMemory::Memcmp(A.TextureMap.GetTypedData(), B.TextureMap.GetTypedData(), A.TextureMap.GetTypeSize()*A.TextureMap.Num()) != 0)
+		if (FMemory::Memcmp(A.TextureMap.GetData(), B.TextureMap.GetData(), A.TextureMap.GetTypeSize()*A.TextureMap.Num()) != 0)
 		{
 			return false;
 		}
@@ -91,14 +91,14 @@ struct FOpenGLShaderBindings
 			return bEqual;
 		}
 
-		bEqual &= FMemory::Memcmp(A.PackedGlobalArrays.GetTypedData(),B.PackedGlobalArrays.GetTypedData(),A.PackedGlobalArrays.GetTypeSize()*A.PackedGlobalArrays.Num()) == 0; 
+		bEqual &= FMemory::Memcmp(A.PackedGlobalArrays.GetData(),B.PackedGlobalArrays.GetData(),A.PackedGlobalArrays.GetTypeSize()*A.PackedGlobalArrays.Num()) == 0; 
 
 		for (int32 Item = 0; Item < A.PackedUniformBuffers.Num(); Item++)
 		{
 			const TArray<CrossCompiler::FPackedArrayInfo> &ArrayA = A.PackedUniformBuffers[Item];
 			const TArray<CrossCompiler::FPackedArrayInfo> &ArrayB = B.PackedUniformBuffers[Item];
 
-			bEqual &= FMemory::Memcmp(ArrayA.GetTypedData(),ArrayB.GetTypedData(),ArrayA.GetTypeSize()*ArrayA.Num()) == 0;
+			bEqual &= FMemory::Memcmp(ArrayA.GetData(),ArrayB.GetData(),ArrayA.GetTypeSize()*ArrayA.Num()) == 0;
 		}
 
 		return bEqual;
@@ -112,14 +112,14 @@ struct FOpenGLShaderBindings
 		Hash |= Binding.NumUniformBuffers << 24;
 		Hash ^= Binding.NumUAVs;
 		Hash ^= Binding.bFlattenUB << 8;
-		Hash ^= FCrc::MemCrc_DEPRECATED( Binding.PackedGlobalArrays.GetTypedData(), Binding.PackedGlobalArrays.GetTypeSize()*Binding.PackedGlobalArrays.Num());
+		Hash ^= FCrc::MemCrc_DEPRECATED( Binding.PackedGlobalArrays.GetData(), Binding.PackedGlobalArrays.GetTypeSize()*Binding.PackedGlobalArrays.Num());
 
 		//@todo-rco: Do we need to calc Binding.ShaderResourceTable.GetTypeHash()?
 
 		for (int32 Item = 0; Item < Binding.PackedUniformBuffers.Num(); Item++)
 		{
 			const TArray<CrossCompiler::FPackedArrayInfo> &Array = Binding.PackedUniformBuffers[Item];
-			Hash ^= FCrc::MemCrc_DEPRECATED( Array.GetTypedData(), Array.GetTypeSize()* Array.Num());
+			Hash ^= FCrc::MemCrc_DEPRECATED( Array.GetData(), Array.GetTypeSize()* Array.Num());
 		}
 		return Hash;
 	}
diff --git a/Engine/Source/Runtime/Renderer/Private/AtmosphereRendering.cpp b/Engine/Source/Runtime/Renderer/Private/AtmosphereRendering.cpp
--- a/Engine/Source/Runtime/Renderer/Private/AtmosphereRendering.cpp
+++ b/Engine/Source/Runtime/Renderer/Private/AtmosphereRendering.cpp
@@ -1450,7 +1450,7 @@ void FAtmosphericFogSceneInfo::Read3DPixelsPtr(FRHICommandListImmediate& RHICmdL
 		Data
 		);
 
-	FMemory::Memcpy( OutData, Data.GetTypedData(), Data.Num() * sizeof(FFloat16Color) );
+	FMemory::Memcpy(OutData, Data.GetData(), Data.Num() * sizeof(FFloat16Color));
 }
 
 void FAtmosphericFogSceneInfo::PrecomputeTextures(FRHICommandListImmediate& RHICmdList, const FViewInfo* View, FSceneViewFamily* ViewFamily)
diff --git a/Engine/Source/Runtime/Renderer/Private/DeferredShadingRenderer.cpp b/Engine/Source/Runtime/Renderer/Private/DeferredShadingRenderer.cpp
--- a/Engine/Source/Runtime/Renderer/Private/DeferredShadingRenderer.cpp
+++ b/Engine/Source/Runtime/Renderer/Private/DeferredShadingRenderer.cpp
@@ -800,7 +800,7 @@ void FDeferredShadingSceneRenderer::Render(FRHICommandListImmediate& RHICmdList)
 	{
 		Scene->FXSystem->PostRenderOpaque(
 			RHICmdList,
-			Views.GetTypedData(),
+			Views.GetData(),
 			GSceneRenderTargets.GetSceneDepthTexture(),
 			GSceneRenderTargets.GetGBufferATexture()
 			);
@@ -906,7 +906,7 @@ void FDeferredShadingSceneRenderer::Render(FRHICommandListImmediate& RHICmdList)
 			if (Scene->bIsEditorScene)
 			{
 				// Precompute Atmospheric Textures
-				Scene->AtmosphericFog->PrecomputeTextures(RHICmdList, Views.GetTypedData(), &ViewFamily);
+				Scene->AtmosphericFog->PrecomputeTextures(RHICmdList, Views.GetData(), &ViewFamily);
 			}
 #endif
 			RenderAtmosphere(RHICmdList, LightShaftOutput);
diff --git a/Engine/Source/Runtime/Renderer/Private/IndirectLightingCache.cpp b/Engine/Source/Runtime/Renderer/Private/IndirectLightingCache.cpp
--- a/Engine/Source/Runtime/Renderer/Private/IndirectLightingCache.cpp
+++ b/Engine/Source/Runtime/Renderer/Private/IndirectLightingCache.cpp
@@ -397,13 +397,13 @@ void FIndirectLightingCache::UpdateCache(FScene* Scene, FSceneRenderer& Renderer
 				}
 			}
 
-			UpdateBlocks(Scene, Renderer.Views.GetTypedData(), BlocksToUpdate);
+			UpdateBlocks(Scene, Renderer.Views.GetData(), BlocksToUpdate);
 
 			UpdateTransitionsOverTime(TransitionsOverTimeToUpdate, Renderer.ViewFamily.DeltaWorldTime);
 
 			if (GCacheDrawLightingSamples || Renderer.ViewFamily.EngineShowFlags.VolumeLightingSamples || GCacheDrawDirectionalShadowing)
 			{
-				FViewElementPDI DebugPDI(Renderer.Views.GetTypedData(), NULL);
+				FViewElementPDI DebugPDI(Renderer.Views.GetData(), NULL);
 
 				for (int32 VolumeIndex = 0; VolumeIndex < Scene->PrecomputedLightVolumes.Num(); VolumeIndex++)
 				{
diff --git a/Engine/Source/Runtime/Renderer/Private/PostProcess/VisualizeTexture.cpp b/Engine/Source/Runtime/Renderer/Private/PostProcess/VisualizeTexture.cpp
--- a/Engine/Source/Runtime/Renderer/Private/PostProcess/VisualizeTexture.cpp
+++ b/Engine/Source/Runtime/Renderer/Private/PostProcess/VisualizeTexture.cpp
@@ -464,7 +464,7 @@ void FVisualizeTexture::GenerateContent(FRHICommandListImmediate& RHICmdList, co
 			uint32 ExtendXWithMSAA = Bitmap.Num() / Extent.Y;
 
 			// Save the contents of the array to a bitmap file. (24bit only so alpha channel is dropped)
-			FFileHelper::CreateBitmap(*ScreenFileName, ExtendXWithMSAA, Extent.Y, Bitmap.GetTypedData());	
+			FFileHelper::CreateBitmap(*ScreenFileName, ExtendXWithMSAA, Extent.Y, Bitmap.GetData());	
 
 			UE_LOG(LogConsoleResponse, Display, TEXT("Content was saved to \"%s\""), *FPaths::ScreenShotDir());
 		}
diff --git a/Engine/Source/Runtime/Renderer/Private/StaticMeshDrawList.inl b/Engine/Source/Runtime/Renderer/Private/StaticMeshDrawList.inl
--- a/Engine/Source/Runtime/Renderer/Private/StaticMeshDrawList.inl
+++ b/Engine/Source/Runtime/Renderer/Private/StaticMeshDrawList.inl
@@ -258,10 +258,10 @@ bool TStaticMeshDrawList<DrawingPolicyType>::DrawVisibleInner(
 	{
 		FDrawingPolicyLink* DrawingPolicyLink = &DrawingPolicySet[OrderedDrawingPolicies[Index]];
 		bool bDrawnShared = false;
-		FPlatformMisc::Prefetch(DrawingPolicyLink->CompactElements.GetTypedData());
+		FPlatformMisc::Prefetch(DrawingPolicyLink->CompactElements.GetData());
 		const int32 NumElements = DrawingPolicyLink->Elements.Num();
-		FPlatformMisc::Prefetch(&DrawingPolicyLink->CompactElements.GetTypedData()->MeshId);
-		const FElementCompact* CompactElementPtr = DrawingPolicyLink->CompactElements.GetTypedData();
+		FPlatformMisc::Prefetch(&DrawingPolicyLink->CompactElements.GetData()->MeshId);
+		const FElementCompact* CompactElementPtr = DrawingPolicyLink->CompactElements.GetData();
 		for (int32 ElementIndex = 0; ElementIndex < NumElements; ElementIndex++, CompactElementPtr++)
 		{
 			if (StaticMeshVisibilityMap.AccessCorrespondingBit(FRelativeBitReference(CompactElementPtr->MeshId)))
@@ -412,10 +412,10 @@ int32 TStaticMeshDrawList<DrawingPolicyType>::DrawVisibleFrontToBack(
 	{
 		FDrawingPolicyLink* DrawingPolicyLink = &DrawingPolicySet[*PolicyIt];
 		FVector DrawingPolicyCenter = DrawingPolicyLink->CachedBoundingSphere.Center;
-		FPlatformMisc::Prefetch(DrawingPolicyLink->CompactElements.GetTypedData());
+		FPlatformMisc::Prefetch(DrawingPolicyLink->CompactElements.GetData());
 		const int32 NumElements = DrawingPolicyLink->Elements.Num();
-		FPlatformMisc::Prefetch(&DrawingPolicyLink->CompactElements.GetTypedData()->MeshId);
-		const FElementCompact* CompactElementPtr = DrawingPolicyLink->CompactElements.GetTypedData();
+		FPlatformMisc::Prefetch(&DrawingPolicyLink->CompactElements.GetData()->MeshId);
+		const FElementCompact* CompactElementPtr = DrawingPolicyLink->CompactElements.GetData();
 		for(int32 ElementIndex = 0; ElementIndex < NumElements; ElementIndex++, CompactElementPtr++)
 		{
 			if(StaticMeshVisibilityMap.AccessCorrespondingBit(FRelativeBitReference(CompactElementPtr->MeshId)))
diff --git a/Engine/Source/Runtime/ShaderCore/Public/ShaderCore.h b/Engine/Source/Runtime/ShaderCore/Public/ShaderCore.h
--- a/Engine/Source/Runtime/ShaderCore/Public/ShaderCore.h
+++ b/Engine/Source/Runtime/ShaderCore/Public/ShaderCore.h
@@ -257,10 +257,10 @@ struct FBaseShaderResourceTable
 		{
 			return false;
 		}
-		bEqual &= (FMemory::Memcmp(A.ShaderResourceViewMap.GetTypedData(), B.ShaderResourceViewMap.GetTypedData(), A.ShaderResourceViewMap.GetTypeSize()*A.ShaderResourceViewMap.Num()) == 0);
-		bEqual &= (FMemory::Memcmp(A.SamplerMap.GetTypedData(), B.SamplerMap.GetTypedData(), A.SamplerMap.GetTypeSize()*A.SamplerMap.Num()) == 0);
-		bEqual &= (FMemory::Memcmp(A.UnorderedAccessViewMap.GetTypedData(), B.UnorderedAccessViewMap.GetTypedData(), A.UnorderedAccessViewMap.GetTypeSize()*A.UnorderedAccessViewMap.Num()) == 0);
-		bEqual &= (FMemory::Memcmp(A.ResourceTableLayoutHashes.GetTypedData(), B.ResourceTableLayoutHashes.GetTypedData(), A.ResourceTableLayoutHashes.GetTypeSize()*A.ResourceTableLayoutHashes.Num()) == 0);
+		bEqual &= (FMemory::Memcmp(A.ShaderResourceViewMap.GetData(), B.ShaderResourceViewMap.GetData(), A.ShaderResourceViewMap.GetTypeSize()*A.ShaderResourceViewMap.Num()) == 0);
+		bEqual &= (FMemory::Memcmp(A.SamplerMap.GetData(), B.SamplerMap.GetData(), A.SamplerMap.GetTypeSize()*A.SamplerMap.Num()) == 0);
+		bEqual &= (FMemory::Memcmp(A.UnorderedAccessViewMap.GetData(), B.UnorderedAccessViewMap.GetData(), A.UnorderedAccessViewMap.GetTypeSize()*A.UnorderedAccessViewMap.Num()) == 0);
+		bEqual &= (FMemory::Memcmp(A.ResourceTableLayoutHashes.GetData(), B.ResourceTableLayoutHashes.GetData(), A.ResourceTableLayoutHashes.GetTypeSize()*A.ResourceTableLayoutHashes.Num()) == 0);
 		return bEqual;
 	}
 };
diff --git a/Engine/Source/Runtime/SlateCore/Public/Rendering/DrawElements.h b/Engine/Source/Runtime/SlateCore/Public/Rendering/DrawElements.h
--- a/Engine/Source/Runtime/SlateCore/Public/Rendering/DrawElements.h
+++ b/Engine/Source/Runtime/SlateCore/Public/Rendering/DrawElements.h
@@ -593,7 +593,7 @@ public:
 	FSlateDrawElement& AddUninitialized()
 	{
 		const int32 InsertIdx = DrawElements.AddUninitialized();
-		return *(new(DrawElements.GetTypedData() + InsertIdx) FSlateDrawElement());
+		return *(new(DrawElements.GetData() + InsertIdx) FSlateDrawElement());
 	}
 
 	/**
diff --git a/Engine/Source/Runtime/SlateCore/Public/Textures/SlateTextureData.h b/Engine/Source/Runtime/SlateCore/Public/Textures/SlateTextureData.h
--- a/Engine/Source/Runtime/SlateCore/Public/Textures/SlateTextureData.h
+++ b/Engine/Source/Runtime/SlateCore/Public/Textures/SlateTextureData.h
@@ -81,7 +81,7 @@ struct SLATECORE_API FSlateTextureData
 	/** Accesses the raw bytes of already sized texture data */
 	uint8* GetRawBytesPtr( )
 	{
-		return Bytes.GetTypedData();
+		return Bytes.GetData();
 	}
 
 private:
diff --git a/Engine/Source/Runtime/SlateRHIRenderer/Private/SlateRHIRenderingPolicy.cpp b/Engine/Source/Runtime/SlateRHIRenderer/Private/SlateRHIRenderingPolicy.cpp
--- a/Engine/Source/Runtime/SlateRHIRenderer/Private/SlateRHIRenderingPolicy.cpp
+++ b/Engine/Source/Runtime/SlateRHIRenderer/Private/SlateRHIRenderingPolicy.cpp
@@ -75,7 +75,7 @@ void FSlateElementIndexBuffer::FillBuffer( const TArray<SlateIndex>& InIndices,
 
 		void* IndicesPtr = RHILockIndexBuffer( IndexBufferRHI, 0, RequiredBufferSize, RLM_WriteOnly );
 
-		FMemory::Memcpy( IndicesPtr, InIndices.GetTypedData(), RequiredBufferSize );
+		FMemory::Memcpy( IndicesPtr, InIndices.GetData(), RequiredBufferSize );
 
 		RHIUnlockIndexBuffer(IndexBufferRHI);
 	}
diff --git a/Engine/Source/Runtime/Sockets/Private/NetworkMessage.cpp b/Engine/Source/Runtime/Sockets/Private/NetworkMessage.cpp
--- a/Engine/Source/Runtime/Sockets/Private/NetworkMessage.cpp
+++ b/Engine/Source/Runtime/Sockets/Private/NetworkMessage.cpp
@@ -83,7 +83,7 @@ bool FNFSMessageHeader::ReceivePayload(FArrayReader& OutPayload, const FSimpleAb
 	int32 Size = sizeof(FNFSMessageHeader);
 	HeaderBytes.AddZeroed(Size);
 		
-	if (!Socket.Receive(HeaderBytes.GetTypedData(), Size))
+	if (!Socket.Receive(HeaderBytes.GetData(), Size))
 	{
 		UE_LOG(LogSockets, Error, TEXT("Unable to read full header"));
 		return false;
@@ -116,14 +116,14 @@ bool FNFSMessageHeader::ReceivePayload(FArrayReader& OutPayload, const FSimpleAb
 	OutPayload.Seek(PayloadOffset);
 
 	// receive the payload
-	if (!Socket.Receive(OutPayload.GetTypedData() + PayloadOffset, Header.PayloadSize))
+	if (!Socket.Receive(OutPayload.GetData() + PayloadOffset, Header.PayloadSize))
 	{
 		UE_LOG(LogSockets, Error, TEXT("Unable to read full payload"));
 		return false;
 	}
 
 	// make sure it's valid
-	uint32 ActualPayloadCrc = FCrc::MemCrc_DEPRECATED(OutPayload.GetTypedData() + PayloadOffset, Header.PayloadSize);
+	uint32 ActualPayloadCrc = FCrc::MemCrc_DEPRECATED(OutPayload.GetData() + PayloadOffset, Header.PayloadSize);
 	if (Header.PayloadCrc != ActualPayloadCrc)
 	{
 		UE_LOG(LogSockets, Error, TEXT("Payload Crc failure."));
diff --git a/Engine/Source/Runtime/Sockets/Public/MultichannelTcpSender.h b/Engine/Source/Runtime/Sockets/Public/MultichannelTcpSender.h
--- a/Engine/Source/Runtime/Sockets/Public/MultichannelTcpSender.h
+++ b/Engine/Source/Runtime/Sockets/Public/MultichannelTcpSender.h
@@ -91,7 +91,7 @@ public:
 
 		SendBuffer->AddUninitialized(Count);
 
-		FMemory::Memcpy(SendBuffer->GetTypedData() + OldCount, Data, Count);
+		FMemory::Memcpy(SendBuffer->GetData() + OldCount, Data, Count);
 
 		AttemptResumeSendingInternal();
 	}
@@ -139,7 +139,7 @@ public:
 					if (OkToSendDelegate.Execute(Size, Channel))
 					{
 						Data.AddUninitialized(Size);
-						FMemory::Memcpy(Data.GetTypedData(), SendBuffer->GetTypedData(), Size);
+						FMemory::Memcpy(Data.GetData(), SendBuffer->GetData(), Size);
 
 						if (Size < Num)
 						{
diff --git a/Engine/Source/Runtime/Sockets/Public/MultichannelTcpSocket.h b/Engine/Source/Runtime/Sockets/Public/MultichannelTcpSocket.h
--- a/Engine/Source/Runtime/Sockets/Public/MultichannelTcpSocket.h
+++ b/Engine/Source/Runtime/Sockets/Public/MultichannelTcpSocket.h
@@ -71,7 +71,7 @@ public:
 				
 				if (ChannelBuffer->Buffer.Num() >= Count)
 				{
-					FMemory::Memcpy(Data, ChannelBuffer->Buffer.GetTypedData(), Count);
+					FMemory::Memcpy(Data, ChannelBuffer->Buffer.GetData(), Count);
 
 					if (ChannelBuffer->Buffer.Num() == Count)
 					{
@@ -145,7 +145,7 @@ public:
 		
 		int32 Count = FMath::Min<int32>(ChannelBuffer->Buffer.Num(), MaxCount);
 
-		FMemory::Memcpy(Data, ChannelBuffer->Buffer.GetTypedData(), Count);
+		FMemory::Memcpy(Data, ChannelBuffer->Buffer.GetData(), Count);
 		
 		if (ChannelBuffer->Buffer.Num() == Count)
 		{
@@ -241,7 +241,7 @@ private:
 			FBufferArchive Ar;
 			uint64 BytesReceived = Receiver.GetBytesReceived();
 			Ar << BytesReceived;
-			Sender.Send(Ar.GetTypedData(), Ar.Num(), ControlChannel);
+			Sender.Send(Ar.GetData(), Ar.Num(), ControlChannel);
 		}
 	}
 
diff --git a/Engine/Source/Runtime/Windows/D3D11RHI/Private/D3D11RHI.cpp b/Engine/Source/Runtime/Windows/D3D11RHI/Private/D3D11RHI.cpp
--- a/Engine/Source/Runtime/Windows/D3D11RHI/Private/D3D11RHI.cpp
+++ b/Engine/Source/Runtime/Windows/D3D11RHI/Private/D3D11RHI.cpp
@@ -398,7 +398,7 @@ void FD3DGPUProfiler::EndFrame()
 				{
 					GPUHitchEventNodeFrames.RemoveAt(0);
 				}
-				GPUHitchEventNodeFrames.AddRawItem((FD3D11EventNodeFrame*)CurrentEventNodeFrame);
+				GPUHitchEventNodeFrames.Add((FD3D11EventNodeFrame*)CurrentEventNodeFrame);
 				CurrentEventNodeFrame = NULL;  // prevent deletion of this below; ke kept it in the history
 			}
 		}
diff --git a/Engine/Source/Runtime/Windows/D3D11RHI/Private/D3D11RenderTarget.cpp b/Engine/Source/Runtime/Windows/D3D11RHI/Private/D3D11RenderTarget.cpp
--- a/Engine/Source/Runtime/Windows/D3D11RHI/Private/D3D11RenderTarget.cpp
+++ b/Engine/Source/Runtime/Windows/D3D11RHI/Private/D3D11RenderTarget.cpp
@@ -669,7 +669,7 @@ void FD3D11DynamicRHI::ReadSurfaceDataNoMSAARaw(FTextureRHIParamRef TextureRHI,F
 	VERIFYD3D11RESULT(Direct3DDeviceIMContext->Map(TempTexture2D,0,D3D11_MAP_READ,0,&LockedRect));
 
 	uint32 BytesPerLine = BytesPerPixel * InRect.Width();
-	uint8* DestPtr = OutData.GetTypedData();
+	uint8* DestPtr = OutData.GetData();
 	uint8* SrcPtr = (uint8*)LockedRect.pData + StagingRect.Min.X * BytesPerPixel +  StagingRect.Min.Y * LockedRect.RowPitch;
 	for(uint32 Y = 0; Y < SizeY; Y++)
 	{
@@ -1012,7 +1012,7 @@ void FD3D11DynamicRHI::RHIReadSurfaceData(FTextureRHIParamRef TextureRHI,FIntRec
 	uint32 BytesPerPixel = ComputeBytesPerPixel(TextureDesc.Format);
 	uint32 SrcPitch = SizeX * BytesPerPixel;
 
-	ConvertRAWSurfaceDataToFColor(TextureDesc.Format, SizeX, SizeY, OutDataRaw.GetTypedData(), SrcPitch, OutData.GetTypedData(), InFlags);
+	ConvertRAWSurfaceDataToFColor(TextureDesc.Format, SizeX, SizeY, OutDataRaw.GetData(), SrcPitch, OutData.GetData(), InFlags);
 }
 
 void FD3D11DynamicRHI::ReadSurfaceDataMSAARaw(FRHICommandList_RecursiveHazardous& RHICmdList, FTextureRHIParamRef TextureRHI,FIntRect InRect,TArray<uint8>& OutData, FReadSurfaceDataFlags InFlags)
diff --git a/Engine/Source/Runtime/Windows/D3D11RHI/Private/D3D11Shaders.cpp b/Engine/Source/Runtime/Windows/D3D11RHI/Private/D3D11Shaders.cpp
--- a/Engine/Source/Runtime/Windows/D3D11RHI/Private/D3D11Shaders.cpp
+++ b/Engine/Source/Runtime/Windows/D3D11RHI/Private/D3D11Shaders.cpp
@@ -15,7 +15,7 @@ FVertexShaderRHIRef FD3D11DynamicRHI::RHICreateVertexShader(const TArray<uint8>&
 	FMemoryReader Ar( Code, true );
 	Ar << Shader->ShaderResourceTable;
 	int32 Offset = Ar.Tell();
-	const uint8* CodePtr = Code.GetTypedData() + Offset;
+	const uint8* CodePtr = Code.GetData() + Offset;
 	const size_t CodeSize = Code.Num() - Offset - 1;
 
 	VERIFYD3D11RESULT( Direct3DDevice->CreateVertexShader( (void*)CodePtr, CodeSize, NULL, Shader->Resource.GetInitReference() ) );
@@ -39,7 +39,7 @@ FGeometryShaderRHIRef FD3D11DynamicRHI::RHICreateGeometryShader(const TArray<uin
 	FMemoryReader Ar( Code, true );
 	Ar << Shader->ShaderResourceTable;
 	int32 Offset = Ar.Tell();
-	const uint8* CodePtr = Code.GetTypedData() + Offset;
+	const uint8* CodePtr = Code.GetData() + Offset;
 	const size_t CodeSize = Code.Num() - Offset - 1;
 
 	VERIFYD3D11RESULT( Direct3DDevice->CreateGeometryShader( (void*)CodePtr, CodeSize, NULL, Shader->Resource.GetInitReference() ) );
@@ -59,7 +59,7 @@ FGeometryShaderRHIRef FD3D11DynamicRHI::RHICreateGeometryShaderWithStreamOutput(
 	FMemoryReader Ar( Code, true );
 	Ar << Shader->ShaderResourceTable;
 	int32 Offset = Ar.Tell();
-	const uint8* CodePtr = Code.GetTypedData() + Offset;
+	const uint8* CodePtr = Code.GetData() + Offset;
 	const size_t CodeSize = Code.Num() - Offset - 1;
 
 	uint32 D3DRasterizedStream = RasterizedStream;
@@ -107,7 +107,7 @@ FHullShaderRHIRef FD3D11DynamicRHI::RHICreateHullShader(const TArray<uint8>& Cod
 	FMemoryReader Ar( Code, true );
 	Ar << Shader->ShaderResourceTable;
 	int32 Offset = Ar.Tell();
-	const uint8* CodePtr = Code.GetTypedData() + Offset;
+	const uint8* CodePtr = Code.GetData() + Offset;
 	const size_t CodeSize = Code.Num() - Offset - 1;
 
 	VERIFYD3D11RESULT( Direct3DDevice->CreateHullShader( (void*)CodePtr, CodeSize, NULL, Shader->Resource.GetInitReference() ) );
@@ -127,7 +127,7 @@ FDomainShaderRHIRef FD3D11DynamicRHI::RHICreateDomainShader(const TArray<uint8>&
 	FMemoryReader Ar( Code, true );
 	Ar << Shader->ShaderResourceTable;
 	int32 Offset = Ar.Tell();
-	const uint8* CodePtr = Code.GetTypedData() + Offset;
+	const uint8* CodePtr = Code.GetData() + Offset;
 	const size_t CodeSize = Code.Num() - Offset - 1;
 
 	VERIFYD3D11RESULT( Direct3DDevice->CreateDomainShader( (void*)CodePtr, CodeSize, NULL, Shader->Resource.GetInitReference() ) );
@@ -147,7 +147,7 @@ FPixelShaderRHIRef FD3D11DynamicRHI::RHICreatePixelShader(const TArray<uint8>& C
 	FMemoryReader Ar( Code, true );
 	Ar << Shader->ShaderResourceTable;
 	int32 Offset = Ar.Tell();
-	const uint8* CodePtr = Code.GetTypedData() + Offset;
+	const uint8* CodePtr = Code.GetData() + Offset;
 	const size_t CodeSize = Code.Num() - Offset - 1;
 
 	VERIFYD3D11RESULT( Direct3DDevice->CreatePixelShader( (void*)CodePtr, CodeSize, NULL, Shader->Resource.GetInitReference() ) );
@@ -167,7 +167,7 @@ FComputeShaderRHIRef FD3D11DynamicRHI::RHICreateComputeShader(const TArray<uint8
 	FMemoryReader Ar( Code, true );
 	Ar << Shader->ShaderResourceTable;
 	int32 Offset = Ar.Tell();
-	const uint8* CodePtr = Code.GetTypedData() + Offset;
+	const uint8* CodePtr = Code.GetData() + Offset;
 	const size_t CodeSize = Code.Num() - Offset - 1;
 
 	VERIFYD3D11RESULT( Direct3DDevice->CreateComputeShader( (void*)CodePtr, CodeSize, NULL, Shader->Resource.GetInitReference() ) );
@@ -215,7 +215,7 @@ FD3D11BoundShaderState::FD3D11BoundShaderState(
 	FMemory::Memzero(&NullInputElement,sizeof(D3D11_INPUT_ELEMENT_DESC));
 
 	VERIFYD3D11RESULT(Direct3DDevice->CreateInputLayout(
-		InVertexDeclaration ? InVertexDeclaration->VertexElements.GetTypedData() : &NullInputElement,
+		InVertexDeclaration ? InVertexDeclaration->VertexElements.GetData() : &NullInputElement,
 		InVertexDeclaration ? InVertexDeclaration->VertexElements.Num() : 0,
 		&InVertexShader->Code[ InVertexShader->Offset ],			// TEMP ugly
 		InVertexShader->Code.Num() - 1 - InVertexShader->Offset,
diff --git a/Engine/Source/Runtime/Windows/D3D11RHI/Private/D3D11VertexDeclaration.cpp b/Engine/Source/Runtime/Windows/D3D11RHI/Private/D3D11VertexDeclaration.cpp
--- a/Engine/Source/Runtime/Windows/D3D11RHI/Private/D3D11VertexDeclaration.cpp
+++ b/Engine/Source/Runtime/Windows/D3D11RHI/Private/D3D11VertexDeclaration.cpp
@@ -60,10 +60,10 @@ struct FD3D11VertexDeclarationKey
 				return ((int32)A.AlignedByteOffset + A.InputSlot * MAX_uint16) < ((int32)B.AlignedByteOffset + B.InputSlot * MAX_uint16);
 			}
 		};
-		Sort(VertexElements.GetTypedData(), VertexElements.Num(), FCompareDesc());
+		Sort(VertexElements.GetData(), VertexElements.Num(), FCompareDesc());
 
 		// Hash once.
-		Hash = FCrc::MemCrc_DEPRECATED(VertexElements.GetTypedData(),VertexElements.Num()*sizeof(D3D11_INPUT_ELEMENT_DESC));
+		Hash = FCrc::MemCrc_DEPRECATED(VertexElements.GetData(),VertexElements.Num()*sizeof(D3D11_INPUT_ELEMENT_DESC));
 	}
 };
 
