Replace deprecated GetTypedData function calls.

diff --git a/Engine/Source/Runtime/OpenGLDrv/Private/OpenGLShaders.cpp b/Engine/Source/Runtime/OpenGLDrv/Private/OpenGLShaders.cpp
--- a/Engine/Source/Runtime/OpenGLDrv/Private/OpenGLShaders.cpp
+++ b/Engine/Source/Runtime/OpenGLDrv/Private/OpenGLShaders.cpp
@@ -468,8 +468,8 @@ ShaderType* CompileOpenGLShader(const TArray<uint8>& Code)
 #if DEBUG_GL_SHADERS
 	Shader->GlslCode.Empty(GlslCodeLength + 1);
 	Shader->GlslCode.AddUninitialized(GlslCodeLength + 1);
-	FMemory::Memcpy(Shader->GlslCode.GetTypedData(), GlslCode, GlslCodeLength + 1);
-	Shader->GlslCodeString = (ANSICHAR*)Shader->GlslCode.GetTypedData();
+	FMemory::Memcpy(Shader->GlslCode.GetData(), GlslCode, GlslCodeLength + 1);
+	Shader->GlslCodeString = (ANSICHAR*)Shader->GlslCode.GetData();
 #endif
 
 #if CHECK_FOR_GL_SHADERS_TO_REPLACE
