UE4 - Integrate Android fix for skylights

diff --git a/Engine/Source/Runtime/OpenGLDrv/Private/Android/AndroidOpenGL.h b/Engine/Source/Runtime/OpenGLDrv/Private/Android/AndroidOpenGL.h
--- a/Engine/Source/Runtime/OpenGLDrv/Private/Android/AndroidOpenGL.h
+++ b/Engine/Source/Runtime/OpenGLDrv/Private/Android/AndroidOpenGL.h
@@ -217,6 +217,9 @@ struct FAndroidOpenGL : public FOpenGLES2
 		}
 	}
 
+	// Adreno doesn't support HALF_FLOAT
+	static FORCEINLINE int32 GetReadHalfFloatPixelsEnum()				{ return GL_FLOAT; }
+
 	static FORCEINLINE bool UseES30ShadingLanguage()
 	{
 		return bUseES30ShadingLanguage;
diff --git a/Engine/Source/Runtime/OpenGLDrv/Private/OpenGLRenderTarget.cpp b/Engine/Source/Runtime/OpenGLDrv/Private/OpenGLRenderTarget.cpp
--- a/Engine/Source/Runtime/OpenGLDrv/Private/OpenGLRenderTarget.cpp
+++ b/Engine/Source/Runtime/OpenGLDrv/Private/OpenGLRenderTarget.cpp
@@ -605,9 +605,13 @@ void FOpenGLDynamicRHI::ReadSurfaceDataRaw(FOpenGLContextState& ContextState, FT
 		int32 FloatBGRADataSize = sizeof(float) * PixelComponentCount;
 		float* FloatBGRAData = (float*)FMemory::Malloc( FloatBGRADataSize );
 		if ( FOpenGL::SupportsBGRA8888() )
+		{
 			glReadPixels(Rect.Min.X, Rect.Min.Y, SizeX, SizeY, GL_BGRA, GL_FLOAT, FloatBGRAData );
+		}
 		else 
+		{
 			glReadPixels(Rect.Min.X, Rect.Min.Y, SizeX, SizeY, GL_RGBA, GL_FLOAT, FloatBGRAData );
+		}
 		// Determine minimal and maximal float values present in received data. Treat each component separately.
 		float MinValue[4] = { 0.0f, 0.0f, 0.0f, 0.0f };
 		float MaxValue[4] = { 1.0f, 1.0f, 1.0f, 1.0f };
@@ -773,7 +777,25 @@ void FOpenGLDynamicRHI::RHIReadSurfaceFloatData(FTextureRHIParamRef TextureRHI,F
 	glBindFramebuffer(UGL_READ_FRAMEBUFFER, SourceFramebuffer);
 	FOpenGL::ReadBuffer(SourceFramebuffer == 0 ? GL_BACK : GL_COLOR_ATTACHMENT0);
 	glPixelStorei(GL_PACK_ALIGNMENT, 1);	
-	glReadPixels(Rect.Min.X, Rect.Min.Y, SizeX, SizeY, GL_RGBA, FOpenGL::GetReadHalfFloatPixelsEnum(), OutData.GetData());
+
+	if (FOpenGL::GetReadHalfFloatPixelsEnum() == GL_FLOAT)
+	{
+		// Slow path: Some Adreno devices won't work with HALF_FLOAT ReadPixels
+		TArray<FLinearColor> FloatData;
+		// 4 float components per texel (RGBA)
+		FloatData.AddUninitialized(SizeX * SizeY);
+		glReadPixels(Rect.Min.X, Rect.Min.Y, SizeX, SizeY, GL_RGBA, GL_FLOAT, FloatData.GetData());
+		FLinearColor* FloatDataPtr = FloatData.GetTypedData();
+		for (int32 Index = 0; Index < SizeX * SizeY; ++Index, ++FloatDataPtr)
+		{
+			OutData[Index] = FFloat16Color(*FloatDataPtr);
+		}
+	}
+	else
+	{
+		glReadPixels(Rect.Min.X, Rect.Min.Y, SizeX, SizeY, GL_RGBA, FOpenGL::GetReadHalfFloatPixelsEnum(), OutData.GetData());
+	}
+
 	glPixelStorei(GL_PACK_ALIGNMENT, 4);
 
 	if (bTempFBO)
@@ -795,9 +817,9 @@ void FOpenGLDynamicRHI::RHIRead3DSurfaceFloatData(FTextureRHIParamRef TextureRHI
 	FRHITexture3D* Texture3DRHI = TextureRHI->GetTexture3D();
 	FOpenGLTextureBase* Texture = GetOpenGLTextureFromRHITexture(TextureRHI);
 
- 	uint32 SizeX = Rect.Width();
- 	uint32 SizeY = Rect.Height();
- 	uint32 SizeZ = ZMinMax.Y - ZMinMax.X;
+	uint32 SizeX = Rect.Width();
+	uint32 SizeY = Rect.Height();
+	uint32 SizeZ = ZMinMax.Y - ZMinMax.X;
 
 	// Allocate the output buffer.
 	OutData.Empty(SizeX * SizeY * SizeZ * sizeof(FFloat16Color));
