UE4 - Compile fix

diff --git a/Engine/Source/Runtime/OpenGLDrv/Private/OpenGLRenderTarget.cpp b/Engine/Source/Runtime/OpenGLDrv/Private/OpenGLRenderTarget.cpp
--- a/Engine/Source/Runtime/OpenGLDrv/Private/OpenGLRenderTarget.cpp
+++ b/Engine/Source/Runtime/OpenGLDrv/Private/OpenGLRenderTarget.cpp
@@ -786,7 +786,7 @@ void FOpenGLDynamicRHI::RHIReadSurfaceFloatData(FTextureRHIParamRef TextureRHI,F
 		FloatData.AddUninitialized(SizeX * SizeY);
 		glReadPixels(Rect.Min.X, Rect.Min.Y, SizeX, SizeY, GL_RGBA, GL_FLOAT, FloatData.GetData());
 		FLinearColor* FloatDataPtr = FloatData.GetTypedData();
-		for (int32 Index = 0; Index < SizeX * SizeY; ++Index, ++FloatDataPtr)
+		for (uint32 Index = 0; Index < SizeX * SizeY; ++Index, ++FloatDataPtr)
 		{
 			OutData[Index] = FFloat16Color(*FloatDataPtr);
 		}
