Fix OpenGL name collision issues.

- Fixes problems of UT modular client builds (also encountered in editor builds) where due to name collision between our function pointer names and GL functions exported by the drivers we were getting wrong addresses when looking for GL extensions and/or overwriting wrong memory.

diff --git a/Engine/Source/Runtime/OpenGLDrv/Private/Linux/OpenGLLinux.cpp b/Engine/Source/Runtime/OpenGLDrv/Private/Linux/OpenGLLinux.cpp
--- a/Engine/Source/Runtime/OpenGLDrv/Private/Linux/OpenGLLinux.cpp
+++ b/Engine/Source/Runtime/OpenGLDrv/Private/Linux/OpenGLLinux.cpp
@@ -11,8 +11,10 @@
 	OpenGL function pointers.
 ------------------------------------------------------------------------------*/
 #define DEFINE_GL_ENTRYPOINTS(Type,Func) Type Func = NULL;
-ENUM_GL_ENTRYPOINTS_ALL(DEFINE_GL_ENTRYPOINTS);
-
+namespace GLFuncPointers	// see explanation in OpenGLLinux.h why we need the namespace
+{
+	ENUM_GL_ENTRYPOINTS_ALL(DEFINE_GL_ENTRYPOINTS);
+};
 
 typedef SDL_Window*		SDL_HWindow;
 typedef SDL_GLContext	SDL_HGLContext;
@@ -778,7 +780,7 @@ bool PlatformInitOpenGL()
 		if (bOpenGLSupported)
 		{
 			// Initialize all entry points required by Unreal.
-			#define GET_GL_ENTRYPOINTS(Type,Func) Func = reinterpret_cast<Type>(SDL_GL_GetProcAddress(#Func));
+			#define GET_GL_ENTRYPOINTS(Type,Func) GLFuncPointers::Func = reinterpret_cast<Type>(SDL_GL_GetProcAddress(#Func));
 			ENUM_GL_ENTRYPOINTS(GET_GL_ENTRYPOINTS);
 			ENUM_GL_ENTRYPOINTS_OPTIONAL(GET_GL_ENTRYPOINTS);
 		
diff --git a/Engine/Source/Runtime/OpenGLDrv/Private/Linux/OpenGLLinux.h b/Engine/Source/Runtime/OpenGLDrv/Private/Linux/OpenGLLinux.h
--- a/Engine/Source/Runtime/OpenGLDrv/Private/Linux/OpenGLLinux.h
+++ b/Engine/Source/Runtime/OpenGLDrv/Private/Linux/OpenGLLinux.h
@@ -345,7 +345,15 @@
 
 /** Declare all GL functions. */
 #define DECLARE_GL_ENTRYPOINTS(Type,Func) extern Type Func;
-ENUM_GL_ENTRYPOINTS_ALL(DECLARE_GL_ENTRYPOINTS);
+// We need to make pointer names different from GL functions otherwise we may end up getting
+// addresses of those symbols when looking for extensions.
+namespace GLFuncPointers
+{
+	ENUM_GL_ENTRYPOINTS_ALL(DECLARE_GL_ENTRYPOINTS);
+};
+
+// this using is needed since the rest of code uses plain GL names
+using namespace GLFuncPointers;
 
 //========================================================================
 
