Fix uninitialized auxSurface variable
#ue4
#android

diff --git a/Engine/Source/Runtime/OpenGLDrv/Private/Android/AndroidEGL.cpp b/Engine/Source/Runtime/OpenGLDrv/Private/Android/AndroidEGL.cpp
--- a/Engine/Source/Runtime/OpenGLDrv/Private/Android/AndroidEGL.cpp
+++ b/Engine/Source/Runtime/OpenGLDrv/Private/Android/AndroidEGL.cpp
@@ -420,6 +420,7 @@ eglDisplay(EGL_NO_DISPLAY)
 	,eglFormat(-1)
 	,eglConfigParam(NULL)
 	,eglSurface(EGL_NO_SURFACE)
+	,auxSurface(EGL_NO_SURFACE)
 	,eglWidth(0)
 	,eglHeight(0)
 	,eglRatio(0)
