Fix for TTP 348451 - Crash when selecting ES2 from the feature level preview menu when running in opengl

diff --git a/Engine/Source/Runtime/OpenGLDrv/Private/OpenGLShaders.cpp b/Engine/Source/Runtime/OpenGLDrv/Private/OpenGLShaders.cpp
--- a/Engine/Source/Runtime/OpenGLDrv/Private/OpenGLShaders.cpp
+++ b/Engine/Source/Runtime/OpenGLDrv/Private/OpenGLShaders.cpp
@@ -1465,6 +1465,9 @@ FBoundShaderStateRHIRef FOpenGLDynamicRHI::RHICreateBoundShaderState(
 
 		FOpenGLLinkedProgramConfiguration Config;
 
+		check(VertexShader);
+		check(PixelShader);
+
 		// Fill-in the configuration
 		Config.Shaders[CrossCompiler::SHADER_STAGE_VERTEX].Bindings = VertexShader->Bindings;
 		Config.Shaders[CrossCompiler::SHADER_STAGE_VERTEX].Resource = VertexShader->Resource;
diff --git a/Engine/Source/Runtime/ShaderCore/Private/Shader.cpp b/Engine/Source/Runtime/ShaderCore/Private/Shader.cpp
--- a/Engine/Source/Runtime/ShaderCore/Private/Shader.cpp
+++ b/Engine/Source/Runtime/ShaderCore/Private/Shader.cpp
@@ -419,7 +419,7 @@ bool ArePlatformsCompatible(EShaderPlatform CurrentPlatform, EShaderPlatform Tar
 								CurrentPlatform == SP_PCD3D_SM4 ||
 								CurrentPlatform == SP_PCD3D_ES2;
 
-		bFeatureLevelCompatible = bFeatureLevelCompatible && bIsCurrentPlatformD3D && bIsTargetD3D;
+		bFeatureLevelCompatible = bFeatureLevelCompatible && (bIsCurrentPlatformD3D == bIsTargetD3D);
 	}
 
 	return bFeatureLevelCompatible;
