UE4 - Fix for TTP 348619, OpenGL tracking down dirty ranges in global uniforms

diff --git a/Engine/Source/Runtime/OpenGLDrv/Private/OpenGLShaders.cpp b/Engine/Source/Runtime/OpenGLDrv/Private/OpenGLShaders.cpp
--- a/Engine/Source/Runtime/OpenGLDrv/Private/OpenGLShaders.cpp
+++ b/Engine/Source/Runtime/OpenGLDrv/Private/OpenGLShaders.cpp
@@ -19,6 +19,29 @@
 const uint32 SizeOfFloat4 = 16;
 const uint32 NumFloatsInFloat4 = 4;
 
+FORCEINLINE void FOpenGLShaderParameterCache::FRange::MarkDirtyRange(uint32 NewStartVector, uint32 NewNumVectors)
+{
+	if (NumVectors > 0)
+	{
+		uint32 High = StartVector + NumVectors;
+		uint32 NewHigh = NewStartVector + NewNumVectors;
+		if (NewStartVector < StartVector)
+		{
+			StartVector = NewStartVector;
+		}
+
+		if (NewHigh > High)
+		{
+			NumVectors = (High - StartVector) + 1;
+		}
+	}
+	else
+	{
+		StartVector = NewStartVector;
+		NumVectors = NewNumVectors;
+	}
+}
+
 /**
  * Verify that an OpenGL shader has compiled successfully. 
  */
@@ -1880,8 +1903,8 @@ FOpenGLShaderParameterCache::FOpenGLShaderParameterCache() :
 {
 	for (int32 ArrayIndex = 0; ArrayIndex < CrossCompiler::PACKED_TYPEINDEX_MAX; ++ArrayIndex)
 	{
-		PackedGlobalUniformDirty[ArrayIndex].LowVector = 0;
-		PackedGlobalUniformDirty[ArrayIndex].HighVector = MAX_uint32;
+		PackedGlobalUniformDirty[ArrayIndex].StartVector = 0;
+		PackedGlobalUniformDirty[ArrayIndex].NumVectors = 0;
 	}
 }
 
@@ -1906,8 +1929,8 @@ void FOpenGLShaderParameterCache::InitializeResources(int32 UniformArraySize)
 
 	for (int32 ArrayIndex = 0; ArrayIndex < CrossCompiler::PACKED_TYPEINDEX_MAX; ++ArrayIndex)
 	{
-		PackedGlobalUniformDirty[ArrayIndex].LowVector = 0;
-		PackedGlobalUniformDirty[ArrayIndex].HighVector = UniformArraySize / SizeOfFloat4;
+		PackedGlobalUniformDirty[ArrayIndex].StartVector = 0;
+		PackedGlobalUniformDirty[ArrayIndex].NumVectors = UniformArraySize / SizeOfFloat4;
 	}
 }
 
@@ -1933,8 +1956,8 @@ void FOpenGLShaderParameterCache::MarkAllDirty()
 {
 	for (int32 ArrayIndex = 0; ArrayIndex < CrossCompiler::PACKED_TYPEINDEX_MAX; ++ArrayIndex)
 	{
-		PackedGlobalUniformDirty[ArrayIndex].LowVector = 0;
-		PackedGlobalUniformDirty[ArrayIndex].HighVector = GlobalUniformArraySize / SizeOfFloat4;
+		PackedGlobalUniformDirty[ArrayIndex].StartVector = 0;
+		PackedGlobalUniformDirty[ArrayIndex].NumVectors = GlobalUniformArraySize / SizeOfFloat4;
 	}
 }
 
@@ -1947,8 +1970,7 @@ void FOpenGLShaderParameterCache::Set(uint32 BufferIndexName, uint32 ByteOffset,
 	check(GlobalUniformArraySize != -1);
 	check(BufferIndex < CrossCompiler::PACKED_TYPEINDEX_MAX);
 	check(ByteOffset + NumBytes <= (uint32)GlobalUniformArraySize);
-	PackedGlobalUniformDirty[BufferIndex].LowVector = FMath::Min(PackedGlobalUniformDirty[BufferIndex].LowVector, ByteOffset / SizeOfFloat4);
-	PackedGlobalUniformDirty[BufferIndex].HighVector = FMath::Max(PackedGlobalUniformDirty[BufferIndex].HighVector, (ByteOffset + NumBytes + SizeOfFloat4 - 1) / SizeOfFloat4);
+	PackedGlobalUniformDirty[BufferIndex].MarkDirtyRange(ByteOffset / SizeOfFloat4, (NumBytes + SizeOfFloat4 - 1) / SizeOfFloat4);
 	FMemory::Memcpy(PackedGlobalUniforms[BufferIndex] + ByteOffset, NewValues, NumBytes);
 }
 
@@ -1986,15 +2008,12 @@ void FOpenGLShaderParameterCache::CommitPackedGlobals(const FOpenGLLinkedProgram
 		const void* UniformData = PackedGlobalUniforms[ArrayIndex];
 
 		// This has to be >=. If LowVector == HighVector it means that particular vector was written to.
-		if (PackedGlobalUniformDirty[ArrayIndex].HighVector >= PackedGlobalUniformDirty[ArrayIndex].LowVector)
+		if (PackedGlobalUniformDirty[ArrayIndex].NumVectors > 0)
 		{
-			const uint32 StartVector = PackedGlobalUniformDirty[ArrayIndex].LowVector;
-			// The number of dirty vectors is the index of the highest vector written minus the first plus one.
-			// The plus one is important so that we upload the highest vector written.
-			uint32 NumDirtyVectors = PackedGlobalUniformDirty[ArrayIndex].HighVector - StartVector + 1;
-			NumDirtyVectors = FMath::Min(NumDirtyVectors, NumVectors - StartVector);
+			const int32 StartVector = PackedGlobalUniformDirty[ArrayIndex].StartVector;
+			int32 NumDirtyVectors = FMath::Min((int32)PackedGlobalUniformDirty[ArrayIndex].NumVectors, NumVectors - StartVector);
 			check(NumDirtyVectors);
-			UniformData = (uint8*)UniformData + PackedGlobalUniformDirty[ArrayIndex].LowVector * SizeOfFloat4;
+			UniformData = (uint8*)UniformData + StartVector * SizeOfFloat4;
 			Location += StartVector;
 			switch (UniformInfo.Index)
 			{
@@ -2013,8 +2032,8 @@ void FOpenGLShaderParameterCache::CommitPackedGlobals(const FOpenGLLinkedProgram
 				break;
 			}
 
-			PackedGlobalUniformDirty[ArrayIndex].LowVector = GlobalUniformArraySize / SizeOfFloat4;
-			PackedGlobalUniformDirty[ArrayIndex].HighVector = 0;
+			PackedGlobalUniformDirty[ArrayIndex].StartVector = 0;
+			PackedGlobalUniformDirty[ArrayIndex].NumVectors = 0;
 		}
 	}
 }
@@ -2043,12 +2062,10 @@ void FOpenGLShaderParameterCache::CommitPackedUniformBuffers(FOpenGLLinkedProgra
 				if (Info.SourceUBIndex == BufferIndex)
 				{
 					check((Info.DestOffsetInFloats + Info.SizeInFloats) * sizeof(float) <= (uint32)GlobalUniformArraySize);
-
 					float* RESTRICT ScratchMem = (float*)PackedGlobalUniforms[Info.DestUBTypeIndex];
 					ScratchMem += Info.DestOffsetInFloats;
 					FMemory::Memcpy(ScratchMem, SourceData + Info.SourceOffsetInFloats, Info.SizeInFloats * sizeof(float));
-					PackedGlobalUniformDirty[Info.DestUBTypeIndex].LowVector = FMath::Min(PackedGlobalUniformDirty[Info.DestUBTypeIndex].LowVector, uint32(Info.DestOffsetInFloats / NumFloatsInFloat4));
-					PackedGlobalUniformDirty[Info.DestUBTypeIndex].HighVector = FMath::Max(PackedGlobalUniformDirty[Info.DestUBTypeIndex].HighVector, uint32((Info.DestOffsetInFloats + Info.SizeInFloats + NumFloatsInFloat4 - 1) / NumFloatsInFloat4));
+					PackedGlobalUniformDirty[Info.DestUBTypeIndex].MarkDirtyRange(Info.DestOffsetInFloats / NumFloatsInFloat4, (Info.SizeInFloats + NumFloatsInFloat4 - 1) / NumFloatsInFloat4);
 				}
 				else
 				{
diff --git a/Engine/Source/Runtime/OpenGLDrv/Public/OpenGLShaderResources.h b/Engine/Source/Runtime/OpenGLDrv/Public/OpenGLShaderResources.h
--- a/Engine/Source/Runtime/OpenGLDrv/Public/OpenGLShaderResources.h
+++ b/Engine/Source/Runtime/OpenGLDrv/Public/OpenGLShaderResources.h
@@ -315,8 +315,10 @@ private:
 	
 	struct FRange
 	{
-		uint32	LowVector;
-		uint32	HighVector;
+		uint32	StartVector;
+		uint32	NumVectors;
+
+		void MarkDirtyRange(uint32 NewStartVector, uint32 NewNumVectors);
 	};
 	/** Dirty ranges for each uniform array. */
 	FRange	PackedGlobalUniformDirty[CrossCompiler::PACKED_TYPEINDEX_MAX];
