- Added a localized error message for when running a game that doesn't have ES2 support on a device that doesn't have Metal support
#include peter.sauerbrei

diff --git a/Engine/Source/Runtime/OpenGLDrv/Private/IOS/IOSOpenGL.cpp b/Engine/Source/Runtime/OpenGLDrv/Private/IOS/IOSOpenGL.cpp
--- a/Engine/Source/Runtime/OpenGLDrv/Private/IOS/IOSOpenGL.cpp
+++ b/Engine/Source/Runtime/OpenGLDrv/Private/IOS/IOSOpenGL.cpp
@@ -28,6 +28,21 @@ struct FPlatformOpenGLDevice
 
 	FPlatformOpenGLDevice()
 	{
+		// make sure the app is compiled with ES2 support (we fallback tro ES2 mode in IOSView, but that is too early to show a warning
+		// so, we will wait until now to show it
+		// if we get here, and don't support ES2, then something has gone very wrong, and we need to alert the user
+
+		bool bSupportsES2 = false;
+		GConfig->GetBool(TEXT("/Script/IOSRuntimeSettings.IOSRuntimeSettings"), TEXT("bSupportsOpenGLES2"), bSupportsES2, GEngineIni);
+		if (!bSupportsES2)
+		{
+			NSLog(@"App requires Metal but it doesn't exist");
+			FPlatformMisc::MessageBoxExt(EAppMsgType::Ok, 
+				*(NSLOCTEXT("Renderer", "MetalAPIMissingInfo", "This application requires the Metal API which is not available on this device. Metal requires an A7 processor and iOS 8, or later.\n\nDevices that have an A7 are iPhone 5S, iPad Air, and iPad mini with Retina display. Any device older than those will not work.").ToString()), 
+				*(NSLOCTEXT("Renderer", "MetalAPIMissingTitle", "Metal required but not available").ToString()));
+			abort();
+		}
+
 		//SetRealTimeMode(pthread_self(), 20, 60);
 
 		bSingleContext = !GUseThreadedRendering;
