Added the frame pacer limit to the IOS runtime settings.
It is now configurable through->

Project Settings | IOS | Power Usage | Frame Rate Lock

The frame pacer is on by default, for IOS. Games are limited to 30 Frames per second. Tappy is set to 60 fps.

diff --git a/Engine/Source/Runtime/Core/Public/GenericPlatform/GenericPlatformFramePacer.h b/Engine/Source/Runtime/Core/Public/GenericPlatform/GenericPlatformFramePacer.h
--- a/Engine/Source/Runtime/Core/Public/GenericPlatform/GenericPlatformFramePacer.h
+++ b/Engine/Source/Runtime/Core/Public/GenericPlatform/GenericPlatformFramePacer.h
@@ -25,7 +25,7 @@ struct CORE_API FGenericPlatformRHIFramePacer
      * @param InTriggeredEvent - The event we wish to trigger when the frame interval has been triggered by the hardware.
      * @param InFrameInterval - How often should the event be triggered, in Frames.
      */
-	static void InitWithEvent(class FEvent* InTriggeredEvent, uint32 InFrameInterval) {}
+	static void InitWithEvent(class FEvent* InTriggeredEvent) {}
     
     /**
      * Teardown the Frame Pacer.
diff --git a/Engine/Source/Runtime/Core/Public/IOS/IOSPlatformFramePacer.h b/Engine/Source/Runtime/Core/Public/IOS/IOSPlatformFramePacer.h
--- a/Engine/Source/Runtime/Core/Public/IOS/IOSPlatformFramePacer.h
+++ b/Engine/Source/Runtime/Core/Public/IOS/IOSPlatformFramePacer.h
@@ -19,7 +19,7 @@ struct FIOSPlatformRHIFramePacer : public FGenericPlatformRHIFramePacer
 {
     // FGenericPlatformRHIFramePacer interface
     static bool IsEnabled();
-	static void InitWithEvent(class FEvent* TriggeredEvent, uint32 InFrameInterval);
+	static void InitWithEvent(class FEvent* TriggeredEvent);
     static void Destroy();
     
     /** Access to the IOS Frame Pacer: CADisplayLink */
diff --git a/Engine/Source/Runtime/IOS/IOSRuntimeSettings/Classes/IOSRuntimeSettings.h b/Engine/Source/Runtime/IOS/IOSRuntimeSettings/Classes/IOSRuntimeSettings.h
--- a/Engine/Source/Runtime/IOS/IOSRuntimeSettings/Classes/IOSRuntimeSettings.h
+++ b/Engine/Source/Runtime/IOS/IOSRuntimeSettings/Classes/IOSRuntimeSettings.h
@@ -5,6 +5,23 @@
 #include "IOSRuntimeSettings.generated.h"
 
 
+UENUM()
+enum class EPowerUsageFrameRateLock : uint8
+{
+    /** Frame rate is not limited */
+    PUFRL_None = 0 UMETA(DisplayName="None"),
+        
+    /** Frame rate is limited to a maximum of 20 frames per second */
+    PUFRL_20 = 20 UMETA(DisplayName="20 FPS"),
+    
+    /** Frame rate is limited to a maximum of 30 frames per second */
+    PUFRL_30 = 30 UMETA(DisplayName="30 FPS"),
+    
+    /** Frame rate is limited to a maximum of 60 frames per second */
+    PUFRL_60 = 60 UMETA(DisplayName="60 FPS"),
+};
+
+
 /**
  * Implements the settings for the iOS target platform.
  */
@@ -57,6 +74,10 @@ public:
 	// version info
 	UPROPERTY(GlobalConfig, EditAnywhere, Category = BundleInfo)
 	FString VersionInfo;
+    
+    /** Set the maximum frame rate to save on power consumption */
+    UPROPERTY(GlobalConfig, EditAnywhere, Category = PowerUsage)
+    TEnumAsByte<EPowerUsageFrameRateLock> FrameRateLock;
 
 #if WITH_EDITOR
 	// UObject interface
diff --git a/Engine/Source/Runtime/IOS/IOSRuntimeSettings/Private/IOSRuntimeSettings.cpp b/Engine/Source/Runtime/IOS/IOSRuntimeSettings/Private/IOSRuntimeSettings.cpp
--- a/Engine/Source/Runtime/IOS/IOSRuntimeSettings/Private/IOSRuntimeSettings.cpp
+++ b/Engine/Source/Runtime/IOS/IOSRuntimeSettings/Private/IOSRuntimeSettings.cpp
@@ -13,6 +13,7 @@ UIOSRuntimeSettings::UIOSRuntimeSettings(const class FPostConstructInitializePro
 	BundleName = TEXT("MyUE4Game");
 	BundleIdentifier = TEXT("com.YourCompany.GameNameNoSpaces");
 	VersionInfo = TEXT("1.0.0");
+    FrameRateLock = EPowerUsageFrameRateLock::PUFRL_30;
 }
 
 #if WITH_EDITOR
diff --git a/Engine/Source/Runtime/IOS/MetalRHI/Private/MetalManager.cpp b/Engine/Source/Runtime/IOS/MetalRHI/Private/MetalManager.cpp
--- a/Engine/Source/Runtime/IOS/MetalRHI/Private/MetalManager.cpp
+++ b/Engine/Source/Runtime/IOS/MetalRHI/Private/MetalManager.cpp
@@ -191,9 +191,14 @@ FMetalManager::FMetalManager()
 	CommandBufferSemaphore = dispatch_semaphore_create(FParse::Param(FCommandLine::Get(),TEXT("gpulockstep")) ? 1 : 3);
 
 	AutoReleasePoolTLSSlot = FPlatformTLS::AllocTlsSlot();
-
-    FrameReadyEvent = FPlatformProcess::CreateSynchEvent();
-    FIOSPlatformRHIFramePacer::InitWithEvent( FrameReadyEvent, 1 );
+    
+    // Hook into the ios framepacer, if it's enabled for this platform.
+    FrameReadyEvent = NULL;
+    if( FIOSPlatformRHIFramePacer::IsEnabled() )
+    {
+        FrameReadyEvent = FPlatformProcess::CreateSynchEvent();
+        FIOSPlatformRHIFramePacer::InitWithEvent( FrameReadyEvent );
+    }
     
 	InitFrame();
 }
@@ -323,8 +328,11 @@ void FMetalManager::EndFrame(bool bPresent)
 		dispatch_semaphore_signal(CommandBufferSemaphore);
 	 }];
     
-	// Wait until at least a VBlank has passed since last time
-    FrameReadyEvent->Wait();
+    // We may be limiting our framerate to the display link
+    if( FrameReadyEvent != nullptr )
+    {
+        FrameReadyEvent->Wait();
+    }
 
 	// Commit before waiting to avoid leaving the gpu idle
 	[CurrentCommandBuffer commit];
diff --git a/Engine/Source/Runtime/IOS/MetalRHI/Private/MetalManager.h b/Engine/Source/Runtime/IOS/MetalRHI/Private/MetalManager.h
--- a/Engine/Source/Runtime/IOS/MetalRHI/Private/MetalManager.h
+++ b/Engine/Source/Runtime/IOS/MetalRHI/Private/MetalManager.h
@@ -245,6 +245,8 @@ protected:
 	void CommitNonComputeShaderConstants();
 
 private:
+    
+    /** Event for coordinating pausing of render thread to keep inline with the ios display link. */
     FEvent* FrameReadyEvent;
 };
 
diff --git a/Engine/Source/Runtime/OpenGLDrv/Private/IOS/IOSOpenGL.cpp b/Engine/Source/Runtime/OpenGLDrv/Private/IOS/IOSOpenGL.cpp
--- a/Engine/Source/Runtime/OpenGLDrv/Private/IOS/IOSOpenGL.cpp
+++ b/Engine/Source/Runtime/OpenGLDrv/Private/IOS/IOSOpenGL.cpp
@@ -20,6 +20,9 @@ struct FPlatformOpenGLContext
 	}
 };
 
+// Event for coordinating pausing of render thread to keep inline with the ios display link.
+static FEvent* FrameReadyEvent = NULL;
+
 struct FPlatformOpenGLDevice
 {
 	FPlatformOpenGLContext RenderingContext;
@@ -76,6 +79,16 @@ struct FPlatformOpenGLDevice
 
 		PlatformSharedContextSetup(this);
 		InitDefaultGLContextState();
+        
+        // Hook into the ios framepacer, if it's enabled for this platform.
+        if( FIOSPlatformRHIFramePacer::IsEnabled() )
+        {
+            if( FrameReadyEvent == nullptr )
+            {
+                FrameReadyEvent = FPlatformProcess::CreateSynchEvent();
+                FIOSPlatformRHIFramePacer::InitWithEvent( FrameReadyEvent );
+            }
+        }
 	}
 
 	bool SetRealTimeMode(pthread_t ThreadHandle, uint32 NormalProcessingTimeMs, uint32 ConstraintProcessingTimeMs)
@@ -130,20 +143,10 @@ void* PlatformGetWindow(FPlatformOpenGLContext* Context, void** AddParam)
 	return (void*)Context->Context;
 }
 
-
-// Event for coordinating pausing of render thread to keep inline with the ios display link.
-static FEvent* FrameReadyEvent = NULL;
-
 bool PlatformBlitToViewport( FPlatformOpenGLDevice* Device, const FOpenGLViewport& Viewport, uint32 BackbufferSizeX, uint32 BackbufferSizeY, bool bPresent,bool bLockToVsync, int32 SyncInterval )
 {
-    if( FIOSPlatformRHIFramePacer::IsEnabled() )
+    if( FrameReadyEvent != nullptr )
     {
-        if( FrameReadyEvent == NULL )
-        {
-            FrameReadyEvent = FPlatformProcess::CreateSynchEvent();
-            FIOSPlatformRHIFramePacer::InitWithEvent( FrameReadyEvent, 1 );
-        }
-    
         FrameReadyEvent->Wait();
     }
     
