Due to radr://18379338 (UE4 Skeletal meshes broken on NVIDIA/Yosemite) we need to forcibly rebind SRVs whenever their backing buffer is modified.
To do this requires:
- A modification count in TOpenGLBuffer updated when the buffer is written to.
- A strong reference to the buffer, the desired format & the current modification version in the SRV.
- An extra internal function to rebind GL_TEXTURE_BUFFER SRVs inside the OpenGLRHI, which is compiled out on non-Mac OS X platforms.

diff --git a/Engine/Source/Runtime/OpenGLDrv/Private/OpenGLCommands.cpp b/Engine/Source/Runtime/OpenGLDrv/Private/OpenGLCommands.cpp
--- a/Engine/Source/Runtime/OpenGLDrv/Private/OpenGLCommands.cpp
+++ b/Engine/Source/Runtime/OpenGLDrv/Private/OpenGLCommands.cpp
@@ -73,6 +73,13 @@ namespace OpenGLConsoleVariables
 		TEXT("If true, use GL_ARB_bindless_texture over traditional glBindTexture/glBindSampler."),
 		ECVF_ReadOnly
 		);
+	
+	int32 bRebindTextureBuffers = PLATFORM_MAC ? 1 : 0;
+	static FAutoConsoleVariableRef CVarRebindTextureBuffers(
+		TEXT("OpenGL.RebindTextureBuffers"),
+		bRebindTextureBuffers,
+		TEXT("If true, rebind GL_TEXTURE_BUFFER's to their GL_TEXTURE name whenever the buffer is modified.")
+		);
 };
 
 TGlobalResource<FVector4VertexDeclaration> GOpenGLVector4VertexDeclaration;
@@ -439,10 +446,38 @@ void FOpenGLDynamicRHI::RHISetUAVParameter(FComputeShaderRHIParamRef ComputeShad
 	check(0);
 }
 
-void FOpenGLDynamicRHI::InternalSetShaderTexture(FOpenGLTextureBase* Texture, GLint TextureIndex, GLenum Target, GLuint Resource, int NumMips, int LimitMip)
+#if PLATFORM_MAC
+void FOpenGLDynamicRHI::InternalUpdateTextureBuffer( FOpenGLContextState& ContextState, FOpenGLShaderResourceView* SRV, GLint TextureIndex )
+{
+	if ( SRV && SRV->Target == GL_TEXTURE_BUFFER && IsValidRef(SRV->VertexBuffer) && OpenGLConsoleVariables::bRebindTextureBuffers )
+	{
+		FOpenGLVertexBuffer* VB = (FOpenGLVertexBuffer*)SRV->VertexBuffer.GetReference();
+		if (SRV->ModificationVersion != VB->ModificationCount)
+		{
+			const uint32 FormatBPP = GPixelFormats[SRV->Format].BlockBytes;
+			const FOpenGLTextureFormat& GLFormat = GOpenGLTextureFormats[SRV->Format];
+			
+			if( ContextState.ActiveTexture != TextureIndex )
+			{
+				glActiveTexture(GL_TEXTURE0 + TextureIndex);
+				ContextState.ActiveTexture = TextureIndex;
+			}
+			
+			FOpenGL::TexBuffer(GL_TEXTURE_BUFFER, GLFormat.InternalFormat[0], VB->Resource);
+			
+			SRV->ModificationVersion = VB->ModificationCount;
+		}
+	}
+}
+#else
+#define InternalUpdateTextureBuffer(SRV)
+#endif
+
+void FOpenGLDynamicRHI::InternalSetShaderTexture(FOpenGLTextureBase* Texture, FOpenGLShaderResourceView* SRV, GLint TextureIndex, GLenum Target, GLuint Resource, int NumMips, int LimitMip)
 {
 	auto& PendingTextureState = PendingState.Textures[TextureIndex];
 	PendingTextureState.Texture = Texture;
+	PendingTextureState.SRV = SRV;
 	PendingTextureState.Target = Target;
 	PendingTextureState.Resource = Resource;
 	PendingTextureState.LimitMip = LimitMip;
@@ -672,8 +707,11 @@ void FOpenGLDynamicRHI::SetupTexturesForDraw( FOpenGLContextState& ContextState,
 				}
 			}
 #endif
-			
 			CachedSetupTextureStage( ContextState, TextureStageIndex, TextureStage.Target, TextureStage.Resource, TextureStage.LimitMip, TextureStage.NumMips );
+			if (PLATFORM_MAC && OpenGLConsoleVariables::bRebindTextureBuffers && TextureStage.SRV)
+			{
+				InternalUpdateTextureBuffer(ContextState, TextureStage.SRV, TextureStageIndex);
+			}
 			if (bNeedsSetupSamplerStage)
 			{
 				ApplyTextureStage( ContextState, TextureStageIndex, TextureStage, PendingState.SamplerStates[TextureStageIndex] );
@@ -747,7 +785,7 @@ void FOpenGLDynamicRHI::RHISetShaderResourceViewParameter(FPixelShaderRHIParamRe
 		Target = SRV->Target;
 		LimitMip = SRV->LimitMip;
 	}
-	InternalSetShaderTexture(NULL, FOpenGL::GetFirstPixelTextureUnit() + TextureIndex, Target, Resource, 0, LimitMip);
+	InternalSetShaderTexture(NULL, SRV, FOpenGL::GetFirstPixelTextureUnit() + TextureIndex, Target, Resource, 0, LimitMip);
 	RHISetShaderSampler(PixelShaderRHI,TextureIndex,PointSamplerState);
 }
 
@@ -767,7 +805,7 @@ void FOpenGLDynamicRHI::RHISetShaderResourceViewParameter(FVertexShaderRHIParamR
 		Target = SRV->Target;
 		LimitMip = SRV->LimitMip;
 	}
-	InternalSetShaderTexture(NULL, FOpenGL::GetFirstVertexTextureUnit() + TextureIndex, Target, Resource, 0, LimitMip);
+	InternalSetShaderTexture(NULL, SRV, FOpenGL::GetFirstVertexTextureUnit() + TextureIndex, Target, Resource, 0, LimitMip);
 	RHISetShaderSampler(VertexShaderRHI,TextureIndex,PointSamplerState);
 }
 
@@ -786,7 +824,7 @@ void FOpenGLDynamicRHI::RHISetShaderResourceViewParameter(FComputeShaderRHIParam
 		Target = SRV->Target;
 		LimitMip = SRV->LimitMip;
 	}
-	InternalSetShaderTexture(NULL, FOpenGL::GetFirstComputeTextureUnit() + TextureIndex, Target, Resource, 0, LimitMip);
+	InternalSetShaderTexture(NULL, SRV, FOpenGL::GetFirstComputeTextureUnit() + TextureIndex, Target, Resource, 0, LimitMip);
 	RHISetShaderSampler(ComputeShaderRHI,TextureIndex,PointSamplerState);
 }
 
@@ -807,7 +845,7 @@ void FOpenGLDynamicRHI::RHISetShaderResourceViewParameter(FHullShaderRHIParamRef
 		Target = SRV->Target;
 		LimitMip = SRV->LimitMip;
 	}
-	InternalSetShaderTexture(NULL, FOpenGL::GetFirstHullTextureUnit() + TextureIndex, Target, Resource, 0, LimitMip);
+	InternalSetShaderTexture(NULL, SRV, FOpenGL::GetFirstHullTextureUnit() + TextureIndex, Target, Resource, 0, LimitMip);
 }
 
 void FOpenGLDynamicRHI::RHISetShaderResourceViewParameter(FDomainShaderRHIParamRef DomainShaderRHI,uint32 TextureIndex,FShaderResourceViewRHIParamRef SRVRHI)
@@ -827,7 +865,7 @@ void FOpenGLDynamicRHI::RHISetShaderResourceViewParameter(FDomainShaderRHIParamR
 		Target = SRV->Target;
 		LimitMip = SRV->LimitMip;
 	}
-	InternalSetShaderTexture(NULL, FOpenGL::GetFirstDomainTextureUnit() + TextureIndex, Target, Resource, 0, LimitMip);
+	InternalSetShaderTexture(NULL, SRV, FOpenGL::GetFirstDomainTextureUnit() + TextureIndex, Target, Resource, 0, LimitMip);
 }
 
 void FOpenGLDynamicRHI::RHISetShaderResourceViewParameter(FGeometryShaderRHIParamRef GeometryShaderRHI,uint32 TextureIndex,FShaderResourceViewRHIParamRef SRVRHI)
@@ -846,7 +884,7 @@ void FOpenGLDynamicRHI::RHISetShaderResourceViewParameter(FGeometryShaderRHIPara
 		Target = SRV->Target;
 		LimitMip = SRV->LimitMip;
 	}
-	InternalSetShaderTexture(NULL, FOpenGL::GetFirstGeometryTextureUnit() + TextureIndex, Target, Resource, 0, LimitMip);
+	InternalSetShaderTexture(NULL, SRV, FOpenGL::GetFirstGeometryTextureUnit() + TextureIndex, Target, Resource, 0, LimitMip);
 	RHISetShaderSampler(GeometryShaderRHI,TextureIndex,PointSamplerState);
 }
 
@@ -858,11 +896,11 @@ void FOpenGLDynamicRHI::RHISetShaderTexture(FVertexShaderRHIParamRef VertexShade
 	FOpenGLTextureBase* NewTexture = GetOpenGLTextureFromRHITexture(NewTextureRHI);
 	if (NewTexture)
 	{
-		InternalSetShaderTexture(NewTexture, FOpenGL::GetFirstVertexTextureUnit() + TextureIndex, NewTexture->Target, NewTexture->Resource, NewTextureRHI->GetNumMips(), -1);
+		InternalSetShaderTexture(NewTexture, nullptr, FOpenGL::GetFirstVertexTextureUnit() + TextureIndex, NewTexture->Target, NewTexture->Resource, NewTextureRHI->GetNumMips(), -1);
 	}
 	else
 	{
-		InternalSetShaderTexture(nullptr, FOpenGL::GetFirstVertexTextureUnit() + TextureIndex, 0, 0, 0, -1);
+		InternalSetShaderTexture(nullptr, nullptr, FOpenGL::GetFirstVertexTextureUnit() + TextureIndex, 0, 0, 0, -1);
 	}
 }
 
@@ -875,11 +913,11 @@ void FOpenGLDynamicRHI::RHISetShaderTexture(FHullShaderRHIParamRef HullShaderRHI
 	FOpenGLTextureBase* NewTexture = GetOpenGLTextureFromRHITexture(NewTextureRHI);
 	if (NewTexture)
 	{
-		InternalSetShaderTexture(NewTexture, FOpenGL::GetFirstHullTextureUnit() + TextureIndex, NewTexture->Target, NewTexture->Resource, NewTextureRHI->GetNumMips(), -1);
+		InternalSetShaderTexture(NewTexture, nullptr, FOpenGL::GetFirstHullTextureUnit() + TextureIndex, NewTexture->Target, NewTexture->Resource, NewTextureRHI->GetNumMips(), -1);
 	}
 	else
 	{
-		InternalSetShaderTexture(nullptr, FOpenGL::GetFirstHullTextureUnit() + TextureIndex, 0, 0, 0, -1);
+		InternalSetShaderTexture(nullptr, nullptr, FOpenGL::GetFirstHullTextureUnit() + TextureIndex, 0, 0, 0, -1);
 	}
 }
 
@@ -892,11 +930,11 @@ void FOpenGLDynamicRHI::RHISetShaderTexture(FDomainShaderRHIParamRef DomainShade
 	FOpenGLTextureBase* NewTexture = GetOpenGLTextureFromRHITexture(NewTextureRHI);
 	if (NewTexture)
 	{
-		InternalSetShaderTexture(NewTexture, FOpenGL::GetFirstDomainTextureUnit() + TextureIndex, NewTexture->Target, NewTexture->Resource, NewTextureRHI->GetNumMips(), -1);
+		InternalSetShaderTexture(NewTexture, nullptr, FOpenGL::GetFirstDomainTextureUnit() + TextureIndex, NewTexture->Target, NewTexture->Resource, NewTextureRHI->GetNumMips(), -1);
 	}
 	else
 	{
-		InternalSetShaderTexture(nullptr, FOpenGL::GetFirstDomainTextureUnit() + TextureIndex, 0, 0, 0, -1);
+		InternalSetShaderTexture(nullptr, nullptr, FOpenGL::GetFirstDomainTextureUnit() + TextureIndex, 0, 0, 0, -1);
 	}
 }
 
@@ -908,11 +946,11 @@ void FOpenGLDynamicRHI::RHISetShaderTexture(FGeometryShaderRHIParamRef GeometryS
 	FOpenGLTextureBase* NewTexture = GetOpenGLTextureFromRHITexture(NewTextureRHI);
 	if (NewTexture)
 	{
-		InternalSetShaderTexture(NewTexture, FOpenGL::GetFirstGeometryTextureUnit() + TextureIndex, NewTexture->Target, NewTexture->Resource, NewTextureRHI->GetNumMips(), -1);
+		InternalSetShaderTexture(NewTexture, nullptr, FOpenGL::GetFirstGeometryTextureUnit() + TextureIndex, NewTexture->Target, NewTexture->Resource, NewTextureRHI->GetNumMips(), -1);
 	}
 	else
 	{
-		InternalSetShaderTexture(nullptr, FOpenGL::GetFirstGeometryTextureUnit() + TextureIndex, 0, 0, 0, -1);
+		InternalSetShaderTexture(nullptr, nullptr, FOpenGL::GetFirstGeometryTextureUnit() + TextureIndex, 0, 0, 0, -1);
 	}
 }
 
@@ -924,11 +962,11 @@ void FOpenGLDynamicRHI::RHISetShaderTexture(FPixelShaderRHIParamRef PixelShaderR
 	FOpenGLTextureBase* NewTexture = GetOpenGLTextureFromRHITexture(NewTextureRHI);
 	if (NewTexture)
 	{
-		InternalSetShaderTexture(NewTexture, FOpenGL::GetFirstPixelTextureUnit() + TextureIndex, NewTexture->Target, NewTexture->Resource, NewTextureRHI->GetNumMips(), -1);
+		InternalSetShaderTexture(NewTexture, nullptr, FOpenGL::GetFirstPixelTextureUnit() + TextureIndex, NewTexture->Target, NewTexture->Resource, NewTextureRHI->GetNumMips(), -1);
 	}
 	else
 	{
-		InternalSetShaderTexture(nullptr, FOpenGL::GetFirstPixelTextureUnit() + TextureIndex, 0, 0, 0, -1);
+		InternalSetShaderTexture(nullptr, nullptr, FOpenGL::GetFirstPixelTextureUnit() + TextureIndex, 0, 0, 0, -1);
 	}
 }
 
@@ -1032,11 +1070,11 @@ void FOpenGLDynamicRHI::RHISetShaderTexture(FComputeShaderRHIParamRef ComputeSha
 	FOpenGLTextureBase* NewTexture = GetOpenGLTextureFromRHITexture(NewTextureRHI);
 	if (NewTexture)
 	{
-		InternalSetShaderTexture(NewTexture, FOpenGL::GetFirstComputeTextureUnit() + TextureIndex, NewTexture->Target, NewTexture->Resource, NewTextureRHI->GetNumMips(), -1);
+		InternalSetShaderTexture(NewTexture, nullptr, FOpenGL::GetFirstComputeTextureUnit() + TextureIndex, NewTexture->Target, NewTexture->Resource, NewTextureRHI->GetNumMips(), -1);
 	}
 	else
 	{
-		InternalSetShaderTexture(nullptr, FOpenGL::GetFirstComputeTextureUnit() + TextureIndex, 0, 0, 0, -1);
+		InternalSetShaderTexture(nullptr, nullptr, FOpenGL::GetFirstComputeTextureUnit() + TextureIndex, 0, 0, 0, -1);
 	}
 }
 
@@ -2227,11 +2265,11 @@ FORCEINLINE void SetResource(FOpenGLDynamicRHI* RESTRICT OpenGLRHI, uint32 BindI
 {
 	if (Texture)
 	{
-		OpenGLRHI->InternalSetShaderTexture(Texture, GetFirstTextureUnit<Frequency>() + BindIndex, Texture->Target, Texture->Resource, Texture->NumMips, -1);
+		OpenGLRHI->InternalSetShaderTexture(Texture, nullptr, GetFirstTextureUnit<Frequency>() + BindIndex, Texture->Target, Texture->Resource, Texture->NumMips, -1);
 	}
 	else
 	{
-		OpenGLRHI->InternalSetShaderTexture(Texture, GetFirstTextureUnit<Frequency>() + BindIndex, 0, 0, 0, -1);
+		OpenGLRHI->InternalSetShaderTexture(Texture, nullptr, GetFirstTextureUnit<Frequency>() + BindIndex, 0, 0, 0, -1);
 	}
 }
 
@@ -2252,7 +2290,7 @@ FORCEINLINE void SetResource(FOpenGLDynamicRHI* RESTRICT OpenGLRHI, uint32 BindI
 template <EShaderFrequency Frequency>
 FORCEINLINE void SetResource(FOpenGLDynamicRHI* RESTRICT OpenGLRHI, uint32 BindIndex, FOpenGLShaderResourceView* RESTRICT SRV)
 {
-	OpenGLRHI->InternalSetShaderTexture(NULL, GetFirstTextureUnit<Frequency>() + BindIndex, SRV->Target, SRV->Resource, 0, SRV->LimitMip);
+	OpenGLRHI->InternalSetShaderTexture(NULL, SRV, GetFirstTextureUnit<Frequency>() + BindIndex, SRV->Target, SRV->Resource, 0, SRV->LimitMip);
 	SetResource<Frequency>(OpenGLRHI,BindIndex,OpenGLRHI->GetPointSamplerState());
 }
 
diff --git a/Engine/Source/Runtime/OpenGLDrv/Private/OpenGLUAV.cpp b/Engine/Source/Runtime/OpenGLDrv/Private/OpenGLUAV.cpp
--- a/Engine/Source/Runtime/OpenGLDrv/Private/OpenGLUAV.cpp
+++ b/Engine/Source/Runtime/OpenGLDrv/Private/OpenGLUAV.cpp
@@ -31,7 +31,7 @@ FShaderResourceViewRHIRef FOpenGLDynamicRHI::RHICreateShaderResourceView(FVertex
 	// and the next draw will take care of cleaning it up; or
 	// next operation that needs the stage will switch something else in on it.
 
-	return new FOpenGLShaderResourceView(this,TextureID,GL_TEXTURE_BUFFER);
+	return new FOpenGLShaderResourceView(this,TextureID,GL_TEXTURE_BUFFER,Format,VertexBufferRHI);
 }
 
 FOpenGLShaderResourceView::~FOpenGLShaderResourceView()
diff --git a/Engine/Source/Runtime/OpenGLDrv/Public/OpenGLDrv.h b/Engine/Source/Runtime/OpenGLDrv/Public/OpenGLDrv.h
--- a/Engine/Source/Runtime/OpenGLDrv/Public/OpenGLDrv.h
+++ b/Engine/Source/Runtime/OpenGLDrv/Public/OpenGLDrv.h
@@ -469,9 +469,13 @@ private:
 
 public:
 	/** Remember what RHI user wants set on a specific OpenGL texture stage, translating from Stage and TextureIndex for stage pair. */
-	void InternalSetShaderTexture(FOpenGLTextureBase* Texture, GLint TextureIndex, GLenum Target, GLuint Resource, int NumMips, int LimitMip);
+	void InternalSetShaderTexture(FOpenGLTextureBase* Texture, FOpenGLShaderResourceView* SRV, GLint TextureIndex, GLenum Target, GLuint Resource, int NumMips, int LimitMip);
 	void InternalSetShaderUAV(GLint UAVIndex, GLenum Format, GLuint Resource);
 	void InternalSetSamplerStates(GLint TextureIndex, FOpenGLSamplerState* SamplerState);
+#if PLATFORM_MAC
+	/** On OS X force a rebind of the texture buffer to the texture name to workaround radr://18379338 */
+	void InternalUpdateTextureBuffer( FOpenGLContextState& ContextState, FOpenGLShaderResourceView* SRV, GLint TextureIndex );
+#endif
 
 private:
 	void ApplyTextureStage(FOpenGLContextState& ContextState, GLint TextureIndex, const FTextureStage& TextureStage, FOpenGLSamplerState* SamplerState);
diff --git a/Engine/Source/Runtime/OpenGLDrv/Public/OpenGLResources.h b/Engine/Source/Runtime/OpenGLDrv/Public/OpenGLResources.h
--- a/Engine/Source/Runtime/OpenGLDrv/Public/OpenGLResources.h
+++ b/Engine/Source/Runtime/OpenGLDrv/Public/OpenGLResources.h
@@ -91,11 +91,15 @@ class TOpenGLBuffer : public BaseType
 public:
 
 	GLuint Resource;
+	
+	/** Needed on OS X to force a rebind of the texture buffer to the texture name to workaround radr://18379338 */
+	uint64 ModificationCount;
 
 	TOpenGLBuffer(uint32 InStride,uint32 InSize,uint32 InUsage,
 		const void *InData = NULL, bool bStreamedDraw = false, GLuint ResourceToUse = 0, uint32 ResourceSize = 0)
 	: BaseType(InStride,InSize,InUsage)
 	, Resource(0)
+	, ModificationCount(0)
 	, bIsLocked(false)
 	, bIsLockReadOnly(false)
 	, bStreamDraw(bStreamedDraw)
@@ -322,6 +326,7 @@ public:
 				LockBuffer = NULL;
 				bLockBufferWasAllocated = false;
 			}
+			ModificationCount += (bIsLockReadOnly ? 0 : 1);
 			bIsLocked = false;
 		}
 	}
@@ -337,6 +342,7 @@ public:
 #else
 		LoadData( InOffset, InSize, InData);
 #endif
+		ModificationCount++;
 	}
 
 	bool IsDynamic() const { return (this->GetUsage() & BUF_AnyDynamic) != 0; }
@@ -1215,19 +1221,43 @@ public:
 	GLenum Target;
 
 	int32 LimitMip;
+	
+	/** Needed on OS X to force a rebind of the texture buffer to the texture name to workaround radr://18379338 */
+	FVertexBufferRHIRef VertexBuffer;
+	uint64 ModificationVersion;
+	uint8 Format;
 
 	FOpenGLShaderResourceView( FOpenGLDynamicRHI* InOpenGLRHI, GLuint InResource, GLenum InTarget )
 	:	Resource(InResource)
 	,	Target(InTarget)
 	,	LimitMip(-1)
+	,	ModificationVersion(0)
+	,	Format(0)
 	,	OpenGLRHI(InOpenGLRHI)
 	,	OwnsResource(true)
 	{}
+	
+	FOpenGLShaderResourceView( FOpenGLDynamicRHI* InOpenGLRHI, GLuint InResource, GLenum InTarget, uint8 InFormat, FVertexBufferRHIParamRef InVertexBuffer )
+	:	Resource(InResource)
+	,	Target(InTarget)
+	,	LimitMip(-1)
+	,	VertexBuffer(InVertexBuffer)
+	,	ModificationVersion(0)
+	,	Format(InFormat)
+	,	OpenGLRHI(InOpenGLRHI)
+	,	OwnsResource(true)
+	{
+		check(IsValidRef(VertexBuffer));
+		FOpenGLVertexBuffer* VB = (FOpenGLVertexBuffer*)VertexBuffer.GetReference();
+		ModificationVersion = VB->ModificationCount;
+	}
 
 	FOpenGLShaderResourceView( FOpenGLDynamicRHI* InOpenGLRHI, GLuint InResource, GLenum InTarget, GLuint Mip, bool InOwnsResource )
 	:	Resource(InResource)
 	,	Target(InTarget)
 	,	LimitMip(Mip)
+	,	ModificationVersion(0)
+	,	Format(0)
 	,	OpenGLRHI(InOpenGLRHI)
 	,	OwnsResource(InOwnsResource)
 	{}
diff --git a/Engine/Source/Runtime/OpenGLDrv/Public/OpenGLState.h b/Engine/Source/Runtime/OpenGLDrv/Public/OpenGLState.h
--- a/Engine/Source/Runtime/OpenGLDrv/Public/OpenGLState.h
+++ b/Engine/Source/Runtime/OpenGLDrv/Public/OpenGLState.h
@@ -180,6 +180,7 @@ public:
 struct FTextureStage
 {
 	class FOpenGLTextureBase* Texture;
+	class FOpenGLShaderResourceView* SRV;
 	GLenum Target;
 	GLuint Resource;
 	int32 LimitMip;
@@ -188,6 +189,7 @@ struct FTextureStage
 
 	FTextureStage()
 	:	Texture(NULL)
+	,	SRV(NULL)
 	,	Target(GL_NONE)
 	,	Resource(0)
 	,	LimitMip(-1)
