Fix the accidental errors in OpenGLCommands from my last commit.

diff --git a/Engine/Source/Runtime/OpenGLDrv/Private/OpenGLCommands.cpp b/Engine/Source/Runtime/OpenGLDrv/Private/OpenGLCommands.cpp
--- a/Engine/Source/Runtime/OpenGLDrv/Private/OpenGLCommands.cpp
+++ b/Engine/Source/Runtime/OpenGLDrv/Private/OpenGLCommands.cpp
@@ -470,7 +470,7 @@ void FOpenGLDynamicRHI::InternalUpdateTextureBuffer( FOpenGLContextState& Contex
 	}
 }
 #else
-#define InternalUpdateTextureBuffer(SRV)
+#define InternalUpdateTextureBuffer(ContextState, SRV, TextureIndex)
 #endif
 
 void FOpenGLDynamicRHI::InternalSetShaderTexture(FOpenGLTextureBase* Texture, FOpenGLShaderResourceView* SRV, GLint TextureIndex, GLenum Target, GLuint Resource, int NumMips, int LimitMip)
