Invert argument order in new constructor to silence MSVC.

diff --git a/Engine/Source/Runtime/OpenGLDrv/Private/OpenGLUAV.cpp b/Engine/Source/Runtime/OpenGLDrv/Private/OpenGLUAV.cpp
--- a/Engine/Source/Runtime/OpenGLDrv/Private/OpenGLUAV.cpp
+++ b/Engine/Source/Runtime/OpenGLDrv/Private/OpenGLUAV.cpp
@@ -31,7 +31,7 @@ FShaderResourceViewRHIRef FOpenGLDynamicRHI::RHICreateShaderResourceView(FVertex
 	// and the next draw will take care of cleaning it up; or
 	// next operation that needs the stage will switch something else in on it.
 
-	return new FOpenGLShaderResourceView(this,TextureID,GL_TEXTURE_BUFFER,Format,VertexBufferRHI);
+	return new FOpenGLShaderResourceView(this,TextureID,GL_TEXTURE_BUFFER,VertexBufferRHI,Format);
 }
 
 FOpenGLShaderResourceView::~FOpenGLShaderResourceView()
diff --git a/Engine/Source/Runtime/OpenGLDrv/Public/OpenGLResources.h b/Engine/Source/Runtime/OpenGLDrv/Public/OpenGLResources.h
--- a/Engine/Source/Runtime/OpenGLDrv/Public/OpenGLResources.h
+++ b/Engine/Source/Runtime/OpenGLDrv/Public/OpenGLResources.h
@@ -1237,7 +1237,7 @@ public:
 	,	OwnsResource(true)
 	{}
 	
-	FOpenGLShaderResourceView( FOpenGLDynamicRHI* InOpenGLRHI, GLuint InResource, GLenum InTarget, uint8 InFormat, FVertexBufferRHIParamRef InVertexBuffer )
+	FOpenGLShaderResourceView( FOpenGLDynamicRHI* InOpenGLRHI, GLuint InResource, GLenum InTarget, FVertexBufferRHIParamRef InVertexBuffer, uint8 InFormat )
 	:	Resource(InResource)
 	,	Target(InTarget)
 	,	LimitMip(-1)
