- Added custom present support to AndroidOpenGL

diff --git a/Engine/Source/Runtime/OpenGLDrv/Private/Android/AndroidOpenGL.cpp b/Engine/Source/Runtime/OpenGLDrv/Private/Android/AndroidOpenGL.cpp
--- a/Engine/Source/Runtime/OpenGLDrv/Private/Android/AndroidOpenGL.cpp
+++ b/Engine/Source/Runtime/OpenGLDrv/Private/Android/AndroidOpenGL.cpp
@@ -117,7 +117,14 @@ void* PlatformGetWindow(FPlatformOpenGLContext* Context, void** AddParam)
 
 bool PlatformBlitToViewport( FPlatformOpenGLDevice* Device, const FOpenGLViewport& Viewport, uint32 BackbufferSizeX, uint32 BackbufferSizeY, bool bPresent,bool bLockToVsync, int32 SyncInterval )
 {
-	AndroidEGL::GetInstance()->SwapBuffers();
+	if (Viewport.GetCustomPresent())
+	{
+		Viewport.GetCustomPresent()->Present(SyncInterval);
+	}
+	else
+	{
+		AndroidEGL::GetInstance()->SwapBuffers();
+	}
 	return true;
 }
 
