Disable OpenGL workarounds no longer required on Yosemite, though they will still be used on Mavericks.
- No longer need to flush after creating a texture name & before submitting data, this was fixed in 10.10 for Nvidia drivers and MTGL.
- No longer need to flush after changing blend state on Intel cards, fixed in 10.10, significantly improving performance.
- GSupportsDepthFetchDuringDepthTest should always be supported from 10.10 onward, as the AMD X4000 driver appears to be fixed.
- Combined depth-stencil clears now work on the AMD X4000 driver as of 10.10.

diff --git a/Engine/Source/Runtime/OpenGLDrv/Private/Mac/MacOpenGLContext.cpp b/Engine/Source/Runtime/OpenGLDrv/Private/Mac/MacOpenGLContext.cpp
--- a/Engine/Source/Runtime/OpenGLDrv/Private/Mac/MacOpenGLContext.cpp
+++ b/Engine/Source/Runtime/OpenGLDrv/Private/Mac/MacOpenGLContext.cpp
@@ -219,7 +219,7 @@ FPlatformOpenGLContext::FPlatformOpenGLContext()
 , VirtualScreen(0)
 , VendorID(0)
 , RendererID(0)
-, SupportsDepthFetchDuringDepthTest(false)
+, SupportsDepthFetchDuringDepthTest(true)
 {
 }
 
@@ -433,19 +433,22 @@ void FPlatformOpenGLContext::VerifyCurrentContext()
 				PlatformContext->VendorID = 0x1002;
 			}
 			
-			// Renderer IDs matchup to driver kexts, so switching based on them will allow us to target workarouds to many GPUs
-			// which exhibit the same unfortunate driver bugs without having to parse their individual ID strings.
-			switch((PlatformContext->RendererID & kCGLRendererIDMatchingMask))
+			if(FPlatformMisc::IsRunningOnMavericks())
 			{
-				case kCGLRendererATIRadeonX4000ID:
+				// Renderer IDs matchup to driver kexts, so switching based on them will allow us to target workarouds to many GPUs
+				// which exhibit the same unfortunate driver bugs without having to parse their individual ID strings.
+				switch((PlatformContext->RendererID & kCGLRendererIDMatchingMask))
 				{
-					// @todo: remove once AMD fix the AMDX4000 driver for GCN cards so that it is possible to sample the depth while also stencil testing to the same DEPTH_STENCIL texture - it works on all other cards we have.
-					PlatformContext->SupportsDepthFetchDuringDepthTest = false;
-					break;
-				}
-				default:
-				{
-					PlatformContext->SupportsDepthFetchDuringDepthTest = true;
+					case kCGLRendererATIRadeonX4000ID:
+					{
+						// @todo: remove once AMD fix the AMDX4000 driver for GCN cards so that it is possible to sample the depth while also stencil testing to the same DEPTH_STENCIL texture - it works on all other cards we have.
+						PlatformContext->SupportsDepthFetchDuringDepthTest = false;
+						break;
+					}
+					default:
+					{
+						PlatformContext->SupportsDepthFetchDuringDepthTest = true;
+					}
 				}
 			}
 		}
diff --git a/Engine/Source/Runtime/OpenGLDrv/Private/Mac/OpenGLMac.cpp b/Engine/Source/Runtime/OpenGLDrv/Private/Mac/OpenGLMac.cpp
--- a/Engine/Source/Runtime/OpenGLDrv/Private/Mac/OpenGLMac.cpp
+++ b/Engine/Source/Runtime/OpenGLDrv/Private/Mac/OpenGLMac.cpp
@@ -1301,7 +1301,7 @@ bool FMacOpenGL::MustFlushTexStorage(void)
 {
 	// @todo There is a bug in Apple's GL with TexStorage calls, on Nvidia and when using MTGL that can see the texture never be created, which then subsequently causes crashes
 	FPlatformOpenGLContext::VerifyCurrentContext();
-	return GMacFlushTexStorage && (IsRHIDeviceNVIDIA() || GMacUseMTGL);
+	return GMacFlushTexStorage && FPlatformMisc::IsRunningOnMavericks() && (IsRHIDeviceNVIDIA() || GMacUseMTGL);
 }
 
 void FMacOpenGL::DeleteTextures(GLsizei Number, const GLuint* Textures)
diff --git a/Engine/Source/Runtime/OpenGLDrv/Private/Mac/OpenGLMac.h b/Engine/Source/Runtime/OpenGLDrv/Private/Mac/OpenGLMac.h
--- a/Engine/Source/Runtime/OpenGLDrv/Private/Mac/OpenGLMac.h
+++ b/Engine/Source/Runtime/OpenGLDrv/Private/Mac/OpenGLMac.h
@@ -23,6 +23,40 @@ struct FMacOpenGL : public FOpenGL3
 	/** Must we employ a workaround for radr://15553950, TTP# 315197 */
 	static bool MustFlushTexStorage(void);
 	
+	static FORCEINLINE void ClearBufferfi(GLenum Buffer, GLint DrawBufferIndex, GLfloat Depth, GLint Stencil)
+	{
+		if(FPlatformMisc::IsRunningOnMavericks())
+		{
+			switch (Buffer)
+			{
+				case GL_DEPTH_STENCIL:	// Clear depth and stencil separately to avoid an AMD Dx00 bug which causes depth to clear, but not stencil.
+					// Especially irritatingly this bug will not manifest when stepping though the program with GL Profiler.
+					// This was a bug found by me during dev. on Tropico 3 circa. Q4 2011 in the ATi Mac 2xx0 & 4xx0 drivers.
+					// It was never fixed & has re-emerged in the AMD Mac FirePro Dx00 drivers.
+					// Also, on NVIDIA depth must be cleared first.
+				case GL_DEPTH:	// Clear depth only
+					ClearBufferfv(GL_DEPTH, 0, &Depth);
+					
+					// If not also clearing depth break
+					if(Buffer == GL_DEPTH)
+					{
+						break;
+					}
+					// Otherwise fall through to perform a separate stencil clear.
+				case GL_STENCIL:	// Clear stencil only
+					ClearBufferiv(GL_STENCIL, 0, (const GLint*)&Stencil);
+					break;
+					
+				default:
+					break;	// impossible anyway
+			}
+		}
+		else
+		{
+			glClearBufferfi(Buffer, DrawBufferIndex, Depth, Stencil);
+		}
+	}
+
 	static FORCEINLINE void GetQueryObject(GLuint QueryId, EQueryMode QueryMode, GLuint *OutResult)
 	{
 		MacGetQueryObject(QueryId, QueryMode, OutResult);
diff --git a/Engine/Source/Runtime/OpenGLDrv/Private/OpenGLCommands.cpp b/Engine/Source/Runtime/OpenGLDrv/Private/OpenGLCommands.cpp
--- a/Engine/Source/Runtime/OpenGLDrv/Private/OpenGLCommands.cpp
+++ b/Engine/Source/Runtime/OpenGLDrv/Private/OpenGLCommands.cpp
@@ -1415,8 +1415,8 @@ void FOpenGLDynamicRHI::SetPendingBlendStateForActiveRenderTargets( FOpenGLConte
 							RenderTargetBlendState.AlphaBlendOperation
 							);
 					}
-					
-					if (PLATFORM_MAC && IsRHIDeviceIntel())
+#if PLATFORM_MAC // Only need this workaround for Mavericks. Fixed in Yosemite.
+					if (IsRHIDeviceIntel() && FPlatformMisc::IsRunningOnMavericks())
 					{
 						// Flush the separate blend state changes through GL or on Intel cards under OS X the state change may be silently ignored.
 						if (CachedRenderTargetBlendState.ColorSourceBlendFactor != RenderTargetBlendState.ColorSourceBlendFactor
@@ -1429,6 +1429,7 @@ void FOpenGLDynamicRHI::SetPendingBlendStateForActiveRenderTargets( FOpenGLConte
 							FOpenGL::Flush();
 						}
 					}
+#endif
 				}
 				else
 				{
@@ -2906,31 +2907,6 @@ void FOpenGLDynamicRHI::RHIClear(bool bClearColor,const FLinearColor& Color,bool
 
 static inline void ClearCurrentDepthStencilWithCurrentScissor( int8 ClearType, float Depth, uint32 Stencil )
 {
-#if PLATFORM_MAC
-	switch (ClearType)
-	{
-		case CT_DepthStencil:	// Clear depth and stencil separately to avoid an AMD Dx00 bug which causes depth to clear, but not stencil.
-								// Especially irritatingly this bug will not manifest when stepping though the program with GL Profiler.
-								// This was a bug found by me during dev. on Tropico 3 circa. Q4 2011 in the ATi Mac 2xx0 & 4xx0 drivers.
-								// It was never fixed & has re-emerged in the AMD Mac FirePro Dx00 drivers.
-								// Also, on NVIDIA depth must be cleared first.
-		case CT_Depth:	// Clear depth only
-			FOpenGL::ClearBufferfv(GL_DEPTH, 0, &Depth);
-
-			// If not also clearing depth break
-			if(!(ClearType & CT_Stencil))
-			{
-				break;
-			}
-			// Otherwise fall through to perform a separate stencil clear.
-		case CT_Stencil:	// Clear stencil only
-			FOpenGL::ClearBufferiv(GL_STENCIL, 0, (const GLint*)&Stencil);
-			break;
-
-		default:
-			break;	// impossible anyway
-	}
-#else
 	switch (ClearType)
 	{
 	case CT_DepthStencil:	// Clear depth and stencil
@@ -2948,7 +2924,6 @@ static inline void ClearCurrentDepthStencilWithCurrentScissor( int8 ClearType, f
 	default:
 		break;	// impossible anyway
 	}
-#endif
 }
 
 void FOpenGLDynamicRHI::ClearCurrentFramebufferWithCurrentScissor(FOpenGLContextState& ContextState, int8 ClearType, int32 NumClearColors, const FLinearColor* ClearColorArray, float Depth, uint32 Stencil)
