Restore the texture-flush workaround on Yosemite, it is still required when using MTGL.

diff --git a/Engine/Source/Runtime/OpenGLDrv/Private/Mac/OpenGLMac.cpp b/Engine/Source/Runtime/OpenGLDrv/Private/Mac/OpenGLMac.cpp
--- a/Engine/Source/Runtime/OpenGLDrv/Private/Mac/OpenGLMac.cpp
+++ b/Engine/Source/Runtime/OpenGLDrv/Private/Mac/OpenGLMac.cpp
@@ -1299,9 +1299,10 @@ void FMacOpenGL::MacGetQueryObject(GLuint QueryId, EQueryMode QueryMode, GLuint
 
 bool FMacOpenGL::MustFlushTexStorage(void)
 {
-	// @todo There is a bug in Apple's GL with TexStorage calls, on Nvidia and when using MTGL that can see the texture never be created, which then subsequently causes crashes
+	// @todo There is a bug in Apple's GL with TexStorage calls when using MTGL that can see the texture never be created, which then subsequently causes crashes
+	// @todo This bug also affects Nvidia cards under Mavericks without MTGL.
 	FPlatformOpenGLContext::VerifyCurrentContext();
-	return GMacFlushTexStorage && FPlatformMisc::IsRunningOnMavericks() && (IsRHIDeviceNVIDIA() || GMacUseMTGL);
+	return GMacFlushTexStorage && ((FPlatformMisc::IsRunningOnMavericks() && IsRHIDeviceNVIDIA()) || GMacUseMTGL);
 }
 
 void FMacOpenGL::DeleteTextures(GLsizei Number, const GLuint* Textures)
