Prevent crashes on OS X when using MTGL due to unsafe/out-of-spec. usage of OpenGL & synchronisation bugs within glBufferSubData.
- In MapBufferRange we can't assume that RLM_WriteOnly rewrites the whole buffer, so instead use INVALIDATE_RANGE or reading outside the last written range will crash.
- In MapBufferRange using UNSYNCHRONIZED_BIT with RLM_WriteOnly is also out-of-spec. as we aren't guaranteeing that the caller has synchronised access to the buffer with the GPU - on OS X this causes a race condition in the GL (esp. when using MTGL) that results in corrupt rendering.
- Replace all usage of glBufferSubData with our wrapper.
- On OS X when using MTGL map calls to glBufferSubData to MapBufferRange(RLM_WriteOnly)->MemCpy->UnmapBuffer to avoid unsynchronised behaviour in glBufferSubData that subsequently causes the MTGL thread to trip over itself & crash the app.

diff --git a/Engine/Source/Runtime/OpenGLDrv/Private/Mac/OpenGLMac.cpp b/Engine/Source/Runtime/OpenGLDrv/Private/Mac/OpenGLMac.cpp
--- a/Engine/Source/Runtime/OpenGLDrv/Private/Mac/OpenGLMac.cpp
+++ b/Engine/Source/Runtime/OpenGLDrv/Private/Mac/OpenGLMac.cpp
@@ -1311,3 +1311,16 @@ void FMacOpenGL::DeleteTextures(GLsizei Number, const GLuint* Textures)
 	TexturesToDelete.Append(Textures, Number);
 }
 
+void FMacOpenGL::BufferSubData(GLenum Target, GLintptr Offset, GLsizeiptr Size, const GLvoid* Data)
+{
+	if(GMacUseMTGL)
+	{
+		void* Dest = MapBufferRange(Target, Offset, Size, FOpenGLBase::RLM_WriteOnly);
+		FMemory::Memcpy(Dest, Data, Size);
+		UnmapBuffer(Target);
+	}
+	else
+	{
+		glBufferSubData(Target, Offset, Size, Data);
+	}
+}
diff --git a/Engine/Source/Runtime/OpenGLDrv/Private/Mac/OpenGLMac.h b/Engine/Source/Runtime/OpenGLDrv/Private/Mac/OpenGLMac.h
--- a/Engine/Source/Runtime/OpenGLDrv/Private/Mac/OpenGLMac.h
+++ b/Engine/Source/Runtime/OpenGLDrv/Private/Mac/OpenGLMac.h
@@ -23,6 +23,8 @@ struct FMacOpenGL : public FOpenGL3
 	/** Must we employ a workaround for radr://15553950, TTP# 315197 */
 	static bool MustFlushTexStorage(void);
 	
+	static void BufferSubData(GLenum Target, GLintptr Offset, GLsizeiptr Size, const GLvoid* Data);
+	
 	static FORCEINLINE void ClearBufferfi(GLenum Buffer, GLint DrawBufferIndex, GLfloat Depth, GLint Stencil)
 	{
 		if(FPlatformMisc::IsRunningOnMavericks())
diff --git a/Engine/Source/Runtime/OpenGLDrv/Private/OpenGLUniformBuffer.cpp b/Engine/Source/Runtime/OpenGLDrv/Private/OpenGLUniformBuffer.cpp
--- a/Engine/Source/Runtime/OpenGLDrv/Private/OpenGLUniformBuffer.cpp
+++ b/Engine/Source/Runtime/OpenGLDrv/Private/OpenGLUniformBuffer.cpp
@@ -537,7 +537,7 @@ FUniformBufferRHIRef FOpenGLDynamicRHI::RHICreateUniformBuffer(const void* Conte
 		}
 		else
 		{
-			glBufferSubData(GL_UNIFORM_BUFFER, 0, Layout.ConstantBufferSize, Contents);
+			FOpenGL::BufferSubData(GL_UNIFORM_BUFFER, 0, Layout.ConstantBufferSize, Contents);
 		}
 	}
 
diff --git a/Engine/Source/Runtime/OpenGLDrv/Public/OpenGL3.h b/Engine/Source/Runtime/OpenGLDrv/Public/OpenGL3.h
--- a/Engine/Source/Runtime/OpenGLDrv/Public/OpenGL3.h
+++ b/Engine/Source/Runtime/OpenGLDrv/Public/OpenGL3.h
@@ -128,8 +128,8 @@ struct FOpenGL3 : public FOpenGLBase
 				Access = GL_MAP_READ_BIT;
 				break;
 			case RLM_WriteOnly:
-				Access = (GL_MAP_INVALIDATE_BUFFER_BIT | GL_MAP_WRITE_BIT );
-#if 1
+				Access = (GL_MAP_INVALIDATE_RANGE_BIT | GL_MAP_WRITE_BIT);
+#if !PLATFORM_MAC // On OS X using the UNSYNCHRONIZED_BIT here is unsafe & out-of-spec. which will lead to corrupt rendering when using MTGL.
 				// Temp workaround for synchrnoization when a UBO is discarded while being referenced
 				Access |= GL_MAP_UNSYNCHRONIZED_BIT;
 #endif
diff --git a/Engine/Source/Runtime/OpenGLDrv/Public/OpenGLResources.h b/Engine/Source/Runtime/OpenGLDrv/Public/OpenGLResources.h
--- a/Engine/Source/Runtime/OpenGLDrv/Public/OpenGLResources.h
+++ b/Engine/Source/Runtime/OpenGLDrv/Public/OpenGLResources.h
@@ -64,7 +64,7 @@ class TOpenGLBuffer : public BaseType
 			{
 				const uint32 Size = FMath::Min<uint32>( BlockSize, InSize);
 			
-				glBufferSubData( Type, InOffset, Size, Data);
+				FOpenGL::BufferSubData( Type, InOffset, Size, Data);
 
 				InOffset += Size;
 				InSize -= Size;
@@ -73,7 +73,7 @@ class TOpenGLBuffer : public BaseType
 		}
 		else
 		{
-			glBufferSubData( Type, InOffset, InSize, InData);
+			FOpenGL::BufferSubData( Type, InOffset, InSize, InData);
 		}
 	}
 
@@ -118,7 +118,7 @@ public:
 				Resource = ResourceToUse;
 				check( Type != GL_UNIFORM_BUFFER || !IsUniformBufferBound(Resource) );
 				Bind();
-				glBufferSubData(Type, 0, InSize, InData);
+				FOpenGL::BufferSubData(Type, 0, InSize, InData);
 			}
 			else
 			{
@@ -135,7 +135,7 @@ public:
 					else
 					{
 						glBufferData(Type, RealSize, NULL, GetAccess());
-						glBufferSubData(Type, 0, InSize, InData);
+						FOpenGL::BufferSubData(Type, 0, InSize, InData);
 					}
 #else
 					glBufferData(Type, RealSize, NULL, GetAccess());
@@ -313,7 +313,7 @@ public:
 					else
 					{
 						// Only updating a subset of the data
-						glBufferSubData(Type, LockOffset, LockSize, LockBuffer);
+						FOpenGL::BufferSubData(Type, LockOffset, LockSize, LockBuffer);
 						check( LockBuffer != NULL );
 					}
 #else
@@ -338,7 +338,7 @@ public:
 		VERIFY_GL_SCOPE();
 		Bind();
 #if !RESTRICT_SUBDATA_SIZE
-		glBufferSubData(Type, InOffset, InSize, InData);
+		FOpenGL::BufferSubData(Type, InOffset, InSize, InData);
 #else
 		LoadData( InOffset, InSize, InData);
 #endif
