Make it possible to use apitrace on Mac again - apitrace doesn't understand APPLE_texture_range or APPLE_client_storage and explodes if you use them together.

diff --git a/Engine/Source/Runtime/OpenGLDrv/Private/Mac/OpenGLMac.cpp b/Engine/Source/Runtime/OpenGLDrv/Private/Mac/OpenGLMac.cpp
--- a/Engine/Source/Runtime/OpenGLDrv/Private/Mac/OpenGLMac.cpp
+++ b/Engine/Source/Runtime/OpenGLDrv/Private/Mac/OpenGLMac.cpp
@@ -11,6 +11,7 @@
 #include "CocoaThread.h"
 #include "CocoaWindow.h"
 #include <IOKit/IOKitLib.h>
+#include <mach-o/dyld.h>
 
 /*------------------------------------------------------------------------------
  OpenGL static variables.
@@ -1071,6 +1072,7 @@ FRHITexture* PlatformCreateBuiltinBackBuffer(FOpenGLDynamicRHI* OpenGLRHI, uint3
 	return NULL;
 }
 
+bool FMacOpenGL::bUsingApitrace = false;
 bool FMacOpenGL::bSupportsTextureCubeMapArray = false;
 PFNGLTEXSTORAGE2DPROC FMacOpenGL::glTexStorage2D = NULL;
 PFNGLTEXSTORAGE3DPROC FMacOpenGL::glTexStorage3D = NULL;
@@ -1138,6 +1140,18 @@ void FMacOpenGL::ProcessExtensions(const FString& ExtensionsString)
 	ProcessQueryGLInt();
 	FOpenGL3::ProcessExtensions(ExtensionsString);
 	
+	// Check for Apitrace, which doesn't understand some Apple extensions.
+	uint32 ModuleCount = _dyld_image_count();
+	for(uint32 Index = 0; Index < ModuleCount; Index++)
+	{
+		ANSICHAR const* ModulePath = (ANSICHAR const*)_dyld_get_image_name(Index);
+		if(FCStringAnsi::Strstr(ModulePath, "OpenGL.framework/Versions/A/OpenGL") && FCStringAnsi::Strcmp(ModulePath, "/System/Library/Frameworks/OpenGL.framework/Versions/A/OpenGL"))
+		{
+			bUsingApitrace = true;
+			break;
+		}
+	}
+	
 	if(GIsEmulatingTimestamp)
 	{
 		TimestampQueryBits = 64;
diff --git a/Engine/Source/Runtime/OpenGLDrv/Private/Mac/OpenGLMac.h b/Engine/Source/Runtime/OpenGLDrv/Private/Mac/OpenGLMac.h
--- a/Engine/Source/Runtime/OpenGLDrv/Private/Mac/OpenGLMac.h
+++ b/Engine/Source/Runtime/OpenGLDrv/Private/Mac/OpenGLMac.h
@@ -227,8 +227,8 @@ struct FMacOpenGL : public FOpenGL3
 	}
 	
 	static FORCEINLINE bool SupportsSeamlessCubeMap()					{ return true; }
-	static FORCEINLINE bool SupportsClientStorage()						{ return true; }
-	static FORCEINLINE bool SupportsTextureRange()						{ return true; }
+	static FORCEINLINE bool SupportsClientStorage()						{ return !bUsingApitrace; }
+	static FORCEINLINE bool SupportsTextureRange()						{ return !bUsingApitrace; }
 	
 	static FORCEINLINE EShaderPlatform GetShaderPlatform()
 	{
@@ -256,6 +256,8 @@ struct FMacOpenGL : public FOpenGL3
 	static void ProcessExtensions(const FString& ExtensionsString);
 	
 private:
+	/** Detects apitrace interception library & equivalents */
+	static bool bUsingApitrace;
 	/** GL_ARB_texture_cube_map_array */
 	static bool bSupportsTextureCubeMapArray;
 	/** GL_ARB_texture_storage */
