Some tidying up of OpenGLMac, plus exporting GL_EXT_depth_bounds_test which is always available and GL_ARB_draw_indirect which is available on GL4.1 renderers.

diff --git a/Engine/Source/Runtime/OpenGLDrv/Private/Mac/OpenGLMac.cpp b/Engine/Source/Runtime/OpenGLDrv/Private/Mac/OpenGLMac.cpp
--- a/Engine/Source/Runtime/OpenGLDrv/Private/Mac/OpenGLMac.cpp
+++ b/Engine/Source/Runtime/OpenGLDrv/Private/Mac/OpenGLMac.cpp
@@ -1085,6 +1085,9 @@ PFNGLLABELOBJECTEXTPROC FMacOpenGL::glLabelObjectEXT = NULL;
 PFNGLPUSHGROUPMARKEREXTPROC FMacOpenGL::glPushGroupMarkerEXT = NULL;
 PFNGLPOPGROUPMARKEREXTPROC FMacOpenGL::glPopGroupMarkerEXT = NULL;
 PFNGLPATCHPARAMETERIPROC FMacOpenGL::glPatchParameteri = NULL;
+bool FMacOpenGL::bSupportsDrawIndirect = false;
+PFNGLDRAWARRAYSINDIRECTPROC FMacOpenGL::glDrawArraysIndirect = nullptr;
+PFNGLDRAWELEMENTSINDIRECTPROC FMacOpenGL::glDrawElementsIndirect = nullptr;
 
 uint64 FMacOpenGL::GetVideoMemorySize()
 {
@@ -1197,6 +1200,13 @@ void FMacOpenGL::ProcessExtensions(const FString& ExtensionsString)
 	{
 		glPatchParameteri = (PFNGLPATCHPARAMETERIPROC)dlsym(RTLD_SELF, "glPatchParameteri");
 	}
+	
+	if(ExtensionsString.Contains(TEXT("GL_ARB_draw_indirect")))
+	{
+		bSupportsDrawIndirect = true;
+		glDrawArraysIndirect = (PFNGLDRAWARRAYSINDIRECTPROC)dlsym(RTLD_SELF, "glDrawArraysIndirect");
+		glDrawElementsIndirect = (PFNGLDRAWELEMENTSINDIRECTPROC)dlsym(RTLD_SELF, "glDrawElementsIndirect");
+	}
 }
 
 void FMacOpenGL::MacQueryTimestampCounter(GLuint QueryID)
diff --git a/Engine/Source/Runtime/OpenGLDrv/Private/Mac/OpenGLMac.h b/Engine/Source/Runtime/OpenGLDrv/Private/Mac/OpenGLMac.h
--- a/Engine/Source/Runtime/OpenGLDrv/Private/Mac/OpenGLMac.h
+++ b/Engine/Source/Runtime/OpenGLDrv/Private/Mac/OpenGLMac.h
@@ -14,6 +14,7 @@
 
 struct FMacOpenGL : public FOpenGL3
 {
+private:
 	static void MacGetQueryObject(GLuint QueryId, EQueryMode QueryMode, GLuint *OutResult);
 	static void MacQueryTimestampCounter(GLuint QueryID);
 	static void MacBeginQuery(GLenum QueryType, GLuint QueryId);
@@ -23,6 +24,45 @@
 	/** Must we employ a workaround for radr://15553950, TTP# 315197 */
 	static bool MustFlushTexStorage(void);
 	
+public:
+	static void ProcessQueryGLInt();
+	static void ProcessExtensions(const FString& ExtensionsString);
+	static uint64 GetVideoMemorySize();
+	
+	static FORCEINLINE EShaderPlatform GetShaderPlatform()
+	{
+		static bool bForceFeatureLevelES2 = FParse::Param(FCommandLine::Get(), TEXT("FeatureLevelES2"));
+		if (bForceFeatureLevelES2)
+		{
+			return SP_OPENGL_PCES2;
+		}
+		
+		// Shader platform
+		switch(GetMajorVersion())
+		{
+			case 3:
+				return SP_OPENGL_SM4_MAC;
+			case 4:
+				return GetMinorVersion() > 2 ? SP_OPENGL_SM5 : SP_OPENGL_SM4_MAC;
+			default:
+				return SP_OPENGL_SM4_MAC;
+		}
+	}
+	
+	static FORCEINLINE bool SupportsSeparateAlphaBlend()				{ return bSupportsDrawBuffersBlend; }
+	static FORCEINLINE bool SupportsSeamlessCubeMap()					{ return true; }
+	static FORCEINLINE bool SupportsClientStorage()						{ return !bUsingApitrace; }
+	static FORCEINLINE bool SupportsTextureRange()						{ return !bUsingApitrace; }
+	static FORCEINLINE bool SupportsDepthBoundsTest()					{ return true; }
+	static FORCEINLINE bool SupportsDrawIndirect()						{ return bSupportsDrawIndirect; }
+	
+	static FORCEINLINE void Flush()
+	{
+		glFlushRenderAPPLE();
+	}
+	
+	static void DeleteTextures(GLsizei Number, const GLuint* Textures);
+	
 	static void BufferSubData(GLenum Target, GLintptr Offset, GLsizeiptr Size, const GLvoid* Data);
 	
 	static FORCEINLINE void ClearBufferfi(GLenum Buffer, GLint DrawBufferIndex, GLfloat Depth, GLint Stencil)
@@ -79,8 +119,6 @@
 	{
 		MacGetQueryObject(QueryId, QueryMode, OutResult);
 	}
-	
-	static FORCEINLINE void InitDebugContext() UGL_OPTIONAL_VOID
 
 	static FORCEINLINE void LabelObject(GLenum Type, GLuint Object, const ANSICHAR* Name)
 	{
@@ -155,9 +193,10 @@
 		}
 	}
 	
-	static void DeleteTextures(GLsizei Number, const GLuint* Textures);
-
-	static FORCEINLINE bool SupportsSeparateAlphaBlend()				{ return bSupportsDrawBuffersBlend; }
+	static FORCEINLINE void TextureRange(GLenum Target, GLsizei Length, const GLvoid *Pointer)
+	{
+		glTextureRangeAPPLE(Target, Length, Pointer);
+	}
 	
 	static FORCEINLINE void BlendFuncSeparatei(GLuint Buf, GLenum SrcRGB, GLenum DstRGB, GLenum SrcAlpha, GLenum DstAlpha)
 	{
@@ -170,6 +209,7 @@
 			UE_LOG(LogRHI, Fatal, TEXT("OpenGL state on draw requires setting different blend operation or factors to different render targets. This is not supported on Mac OS X GL 3.2 contexts!"));
 		}
 	}
+	
 	static FORCEINLINE void BlendEquationSeparatei(GLuint Buf, GLenum ModeRGB, GLenum ModeAlpha)
 	{
 		if(glBlendEquationSeparatei)
@@ -181,6 +221,7 @@
 			UE_LOG(LogRHI, Fatal, TEXT("OpenGL state on draw requires setting different blend operation or factors to different render targets. This is not supported on Mac OS X GL 3.2 contexts!"));
 		}
 	}
+	
 	static FORCEINLINE void BlendFunci(GLuint Buf, GLenum Src, GLenum Dst)
 	{
 		if(glBlendFunci)
@@ -192,6 +233,7 @@
 			UE_LOG(LogRHI, Fatal, TEXT("OpenGL state on draw requires setting different blend operation or factors to different render targets. This is not supported on Mac OS X GL 3.2 contexts!"));
 		}
 	}
+	
 	static FORCEINLINE void BlendEquationi(GLuint Buf, GLenum Mode)
 	{
 		if(glBlendEquationi)
@@ -216,45 +258,11 @@
 		}
 	}
 	
-	static FORCEINLINE void TextureRange(GLenum Target, GLsizei Length, const GLvoid *Pointer)
+	static FORCEINLINE void DepthBounds(GLfloat Min, GLfloat Max)
 	{
-		glTextureRangeAPPLE(Target, Length, Pointer);
+		glDepthBoundsEXT( Min, Max);
 	}
 	
-	static FORCEINLINE void Flush()
-	{
-		glFlushRenderAPPLE();
-	}
-	
-	static FORCEINLINE bool SupportsSeamlessCubeMap()					{ return true; }
-	static FORCEINLINE bool SupportsClientStorage()						{ return !bUsingApitrace; }
-	static FORCEINLINE bool SupportsTextureRange()						{ return !bUsingApitrace; }
-	
-	static FORCEINLINE EShaderPlatform GetShaderPlatform()
-	{
-		static bool bForceFeatureLevelES2 = FParse::Param(FCommandLine::Get(), TEXT("FeatureLevelES2"));
-		if (bForceFeatureLevelES2)
-		{
-			return SP_OPENGL_PCES2;
-		}
-		
-		// Shader platform
-		switch(GetMajorVersion())
-		{
-			case 3:
-				return SP_OPENGL_SM4_MAC;
-			case 4:
-				return GetMinorVersion() > 2 ? SP_OPENGL_SM5 : SP_OPENGL_SM4_MAC;
-			default:
-				return SP_OPENGL_SM4_MAC;
-		}
-	}
-	
-	static uint64 GetVideoMemorySize();
-	
-	static void ProcessQueryGLInt();
-	static void ProcessExtensions(const FString& ExtensionsString);
-	
 private:
 	/** Detects apitrace interception library & equivalents */
 	static bool bUsingApitrace;
@@ -276,6 +284,10 @@
 	static PFNGLPOPGROUPMARKEREXTPROC glPopGroupMarkerEXT;
 	/** GL_ARB_tessellation_shader */
 	static PFNGLPATCHPARAMETERIPROC glPatchParameteri;
+	/** GL_ARB_draw_indirect */
+	static bool bSupportsDrawIndirect;
+	static PFNGLDRAWARRAYSINDIRECTPROC glDrawArraysIndirect;
+	static PFNGLDRAWELEMENTSINDIRECTPROC glDrawElementsIndirect;
 };
 
 
