UE4 - Fix for mobile dynamic shadows on iOS OpenGL

diff --git a/Engine/Source/Runtime/OpenGLDrv/Private/HTML5/HTML5OpenGL.h b/Engine/Source/Runtime/OpenGLDrv/Private/HTML5/HTML5OpenGL.h
--- a/Engine/Source/Runtime/OpenGLDrv/Private/HTML5/HTML5OpenGL.h
+++ b/Engine/Source/Runtime/OpenGLDrv/Private/HTML5/HTML5OpenGL.h
@@ -71,6 +71,7 @@ struct FHTML5OpenGL : public FOpenGLES2
 	static FORCEINLINE bool SupportsMapBuffer()							{ return false; }
 	static FORCEINLINE bool SupportsCombinedDepthStencilAttachment()	{ return bCombinedDepthStencilAttachment; }
 	static FORCEINLINE GLenum GetDepthFormat()							{ return GL_DEPTH_COMPONENT; }
+	static FORCEINLINE GLenum GetShadowDepthFormat()					{ return GL_DEPTH_COMPONENT; }
 	// Optional
 	static FORCEINLINE void BeginQuery(GLenum QueryType, GLuint QueryId) UGL_OPTIONAL_VOID
 	static FORCEINLINE void EndQuery(GLenum QueryType) UGL_OPTIONAL_VOID
diff --git a/Engine/Source/Runtime/OpenGLDrv/Private/IOS/IOSOpenGL.h b/Engine/Source/Runtime/OpenGLDrv/Private/IOS/IOSOpenGL.h
--- a/Engine/Source/Runtime/OpenGLDrv/Private/IOS/IOSOpenGL.h
+++ b/Engine/Source/Runtime/OpenGLDrv/Private/IOS/IOSOpenGL.h
@@ -107,6 +107,8 @@ struct FIOSOpenGL : public FOpenGLES2
 	}
 
 	static FORCEINLINE bool SupportsFramebufferSRGBEnable()				{ return false; }
+
+	static FORCEINLINE GLenum GetShadowDepthFormat()					{ return GL_DEPTH_COMPONENT16; }
 };
 
 typedef FIOSOpenGL FOpenGL;
diff --git a/Engine/Source/Runtime/OpenGLDrv/Private/OpenGLDevice.cpp b/Engine/Source/Runtime/OpenGLDrv/Private/OpenGLDevice.cpp
--- a/Engine/Source/Runtime/OpenGLDrv/Private/OpenGLDevice.cpp
+++ b/Engine/Source/Runtime/OpenGLDrv/Private/OpenGLDevice.cpp
@@ -693,13 +693,14 @@ static void InitRHICapabilitiesForGL()
 	}
 
 	GLenum DepthFormat = FOpenGL::GetDepthFormat();
+	GLenum ShadowDepthFormat = FOpenGL::GetShadowDepthFormat();
 
 	// Initialize the platform pixel format map.					InternalFormat				InternalFormatSRGB		Format				Type							bCompressed		bBGRA
 	SetupTextureFormat( PF_Unknown,				FOpenGLTextureFormat( ));
 	SetupTextureFormat( PF_A32B32G32R32F,		FOpenGLTextureFormat( GL_RGBA32F,				GL_NONE,				GL_RGBA,			GL_FLOAT,						false,			false));
 	SetupTextureFormat( PF_UYVY,				FOpenGLTextureFormat( ));
 	//@todo: ES2 requires GL_OES_depth_texture extension to support depth textures of any kind.
-	SetupTextureFormat( PF_ShadowDepth,			FOpenGLTextureFormat( DepthFormat,				GL_NONE,				GL_DEPTH_COMPONENT,	GL_UNSIGNED_INT,				false,			false));
+	SetupTextureFormat( PF_ShadowDepth,			FOpenGLTextureFormat( ShadowDepthFormat,		GL_NONE,				GL_DEPTH_COMPONENT,	GL_UNSIGNED_INT,				false,			false));
 	SetupTextureFormat( PF_D24,					FOpenGLTextureFormat( DepthFormat,				GL_NONE,				GL_DEPTH_COMPONENT,	GL_UNSIGNED_INT,				false,			false));
 	SetupTextureFormat( PF_A16B16G16R16,		FOpenGLTextureFormat( GL_RGBA16,				GL_RGBA16,				GL_RGBA,			GL_UNSIGNED_SHORT,				false,			false));
 	SetupTextureFormat( PF_A1,					FOpenGLTextureFormat( ));
diff --git a/Engine/Source/Runtime/OpenGLDrv/Public/OpenGL.h b/Engine/Source/Runtime/OpenGLDrv/Public/OpenGL.h
--- a/Engine/Source/Runtime/OpenGLDrv/Public/OpenGL.h
+++ b/Engine/Source/Runtime/OpenGLDrv/Public/OpenGL.h
@@ -147,8 +147,8 @@ public:
 	static FORCEINLINE bool HasHardwareHiddenSurfaceRemoval()			{ return false; }
 	static FORCEINLINE bool AmdWorkaround()								{ return false; }
 
-
 	static FORCEINLINE GLenum GetDepthFormat()							{ return GL_DEPTH_COMPONENT16; }
+	static FORCEINLINE GLenum GetShadowDepthFormat()					{ return GL_DEPTH_COMPONENT16; }
 	static FORCEINLINE GLenum GetVertexHalfFloatFormat()				{ return GL_HALF_FLOAT; }
 
 	static FORCEINLINE GLint GetMaxTextureImageUnits()			{ check(MaxTextureImageUnits != -1); return MaxTextureImageUnits; }
diff --git a/Engine/Source/Runtime/OpenGLDrv/Public/OpenGLES2.h b/Engine/Source/Runtime/OpenGLDrv/Public/OpenGLES2.h
--- a/Engine/Source/Runtime/OpenGLDrv/Public/OpenGLES2.h
+++ b/Engine/Source/Runtime/OpenGLDrv/Public/OpenGLES2.h
@@ -108,6 +108,7 @@ struct FOpenGLES2 : public FOpenGLBase
 	static FORCEINLINE bool SupportsShaderTextureCubeLod()				{ return bSupportsShaderTextureCubeLod; }
 	static FORCEINLINE bool SupportsCopyTextureLevels()					{ return bSupportsCopyTextureLevels; }
 	static FORCEINLINE GLenum GetDepthFormat()							{ return GL_DEPTH_COMPONENT; }
+	static FORCEINLINE GLenum GetShadowDepthFormat()					{ return GL_DEPTH_COMPONENT; }
 
 	static FORCEINLINE bool RequiresDontEmitPrecisionForTextureSamplers() { return bRequiresDontEmitPrecisionForTextureSamplers; }
 	static FORCEINLINE bool RequiresTextureCubeLodEXTToTextureCubeLodDefine() { return bRequiresTextureCubeLodEXTToTextureCubeLodDefine; }
diff --git a/Engine/Source/Runtime/OpenGLDrv/Public/OpenGLES31.h b/Engine/Source/Runtime/OpenGLDrv/Public/OpenGLES31.h
--- a/Engine/Source/Runtime/OpenGLDrv/Public/OpenGLES31.h
+++ b/Engine/Source/Runtime/OpenGLDrv/Public/OpenGLES31.h
@@ -79,6 +79,7 @@ struct FOpenGLES31 : public FOpenGLBase
 	static FORCEINLINE bool SupportsShaderTextureCubeLod()				{ return bES2Fallback ? bSupportsShaderTextureCubeLod : true; }
 	static FORCEINLINE bool SupportsCopyTextureLevels()					{ return bSupportsCopyTextureLevels; }
 	static FORCEINLINE GLenum GetDepthFormat()							{ return GL_DEPTH_COMPONENT16; }
+	static FORCEINLINE GLenum GetShadowDepthFormat()					{ return GL_DEPTH_COMPONENT16; }
 
 	static FORCEINLINE bool RequiresDontEmitPrecisionForTextureSamplers() { return bRequiresDontEmitPrecisionForTextureSamplers; }
 	static FORCEINLINE bool RequiresTextureCubeLodEXTToTextureCubeLodDefine() { return bRequiresTextureCubeLodEXTToTextureCubeLodDefine; }
