remove GRHIShaderPlatform_DEPRECATED usage from OpenGL RHI (replaced with GMaxRHIShaderPlatform)

diff --git a/Engine/Source/Runtime/OpenGLDrv/Private/OpenGLCommands.cpp b/Engine/Source/Runtime/OpenGLDrv/Private/OpenGLCommands.cpp
--- a/Engine/Source/Runtime/OpenGLDrv/Private/OpenGLCommands.cpp
+++ b/Engine/Source/Runtime/OpenGLDrv/Private/OpenGLCommands.cpp
@@ -1605,7 +1605,7 @@ void FOpenGLDynamicRHI::RHISetRenderTargets(
 
 	FOpenGLTextureBase* NewDepthStencilRT = GetOpenGLTextureFromRHITexture(NewDepthStencilTargetRHI);
 
-	if (IsES2Platform(GRHIShaderPlatform_DEPRECATED) && !IsPCPlatform(GRHIShaderPlatform_DEPRECATED))
+	if (IsES2Platform(GMaxRHIShaderPlatform) && !IsPCPlatform(GMaxRHIShaderPlatform))
 	{
 		// @todo-mobile
 
diff --git a/Engine/Source/Runtime/OpenGLDrv/Private/OpenGLDevice.cpp b/Engine/Source/Runtime/OpenGLDrv/Private/OpenGLDevice.cpp
--- a/Engine/Source/Runtime/OpenGLDrv/Private/OpenGLDevice.cpp
+++ b/Engine/Source/Runtime/OpenGLDrv/Private/OpenGLDevice.cpp
@@ -664,7 +664,7 @@ static void InitRHICapabilitiesForGL()
 	GSupportsVertexInstancing = true;
 	GHardwareHiddenSurfaceRemoval = FOpenGL::HasHardwareHiddenSurfaceRemoval();
 
-	GShaderPlatformForFeatureLevel[ERHIFeatureLevel::ES2] = (GMaxRHIFeatureLevel == ERHIFeatureLevel::ES2) ? GRHIShaderPlatform_DEPRECATED : SP_OPENGL_PCES2;
+	GShaderPlatformForFeatureLevel[ERHIFeatureLevel::ES2] = (GMaxRHIFeatureLevel == ERHIFeatureLevel::ES2) ? GMaxRHIShaderPlatform : SP_OPENGL_PCES2;
 	GShaderPlatformForFeatureLevel[ERHIFeatureLevel::ES3_1] = SP_NumPlatforms;
 	GShaderPlatformForFeatureLevel[ERHIFeatureLevel::SM4] = PLATFORM_MAC ? SP_OPENGL_SM4_MAC : SP_OPENGL_SM4;
 	GShaderPlatformForFeatureLevel[ERHIFeatureLevel::SM5] = OPENGL_ES31 ? SP_OPENGL_ES31_EXT : SP_OPENGL_SM5;
@@ -935,7 +935,7 @@ static bool VerifyCompiledShader(GLuint Shader, const ANSICHAR* GlslCode, bool I
 static void CheckTextureCubeLodSupport()
 {
 #if PLATFORM_ANDROID
-	if (IsES2Platform(GRHIShaderPlatform_DEPRECATED))
+	if (IsES2Platform(GMaxRHIShaderPlatform))
 	{
 		UE_LOG(LogRHI, Display, TEXT("Testing for shader compiler compatibility"));
 		// This code creates a sample program and finds out which hacks are required to compile it
diff --git a/Engine/Source/Runtime/OpenGLDrv/Private/OpenGLShaders.cpp b/Engine/Source/Runtime/OpenGLDrv/Private/OpenGLShaders.cpp
--- a/Engine/Source/Runtime/OpenGLDrv/Private/OpenGLShaders.cpp
+++ b/Engine/Source/Runtime/OpenGLDrv/Private/OpenGLShaders.cpp
@@ -336,7 +336,7 @@ ShaderType* CompileOpenGLShader(const TArray<uint8>& Code)
 		ANSICHAR* VersionString = new ANSICHAR[18];
 		memset((void*)VersionString, '\0', 18 * sizeof(ANSICHAR));
 
-		if (IsES2Platform(GRHIShaderPlatform_DEPRECATED))
+		if (IsES2Platform(GMaxRHIShaderPlatform))
 		{
 			//#version 100 has to be the first line in the file, so it has to be added before anything else.
 			if (FOpenGL::UseES30ShadingLanguage())
@@ -365,7 +365,7 @@ ShaderType* CompileOpenGLShader(const TArray<uint8>& Code)
 #endif
 
 #if PLATFORM_ANDROID 
-		if (IsES2Platform(GRHIShaderPlatform_DEPRECATED))
+		if (IsES2Platform(GMaxRHIShaderPlatform))
 		{
 			//Here is some code to add in a #define for textureCubeLodEXT
 			const ANSICHAR* ExtensionString = "";
@@ -1259,7 +1259,7 @@ static FOpenGLLinkedProgram* LinkProgram( const FOpenGLLinkedProgramConfiguratio
 	}
 	
 	// E.g. GLSL_430 uses layout(location=xx) instead of having to call glBindAttribLocation and glBindFragDataLocation
-	if (OpenGLShaderPlatformNeedsBindLocation(GRHIShaderPlatform_DEPRECATED))
+	if (OpenGLShaderPlatformNeedsBindLocation(GMaxRHIShaderPlatform))
 	{
 		// Bind attribute indices.
 		if (Config.Shaders[CrossCompiler::SHADER_STAGE_VERTEX].Resource)
