Fix bSupportsColorBufferHelfFloat incorrectly being set when glReadPixels would fail. HTML5
Fixes pink glow issue in VehicleAdv template, caused by invalid data from failed glReadPixels being passed to the shader (UE-3762 & UE-4808).

diff --git a/Engine/Source/Runtime/OpenGLDrv/Private/HTML5/HTML5OpenGL.cpp b/Engine/Source/Runtime/OpenGLDrv/Private/HTML5/HTML5OpenGL.cpp
--- a/Engine/Source/Runtime/OpenGLDrv/Private/HTML5/HTML5OpenGL.cpp
+++ b/Engine/Source/Runtime/OpenGLDrv/Private/HTML5/HTML5OpenGL.cpp
@@ -99,7 +99,11 @@ void FHTML5OpenGL::ProcessExtensions( const FString& ExtensionsString )
         glFramebufferTexture2D(GL_FRAMEBUFFER, GL_COLOR_ATTACHMENT0, GL_TEXTURE_2D, tex, 0);
 
         GLenum fbstatus = glCheckFramebufferStatus(GL_FRAMEBUFFER);
-        err = glGetError();
+
+		// bSupportsColorBufferHalfFloat is used in such a way that reading the frame buffer is expected to work. Test that now.
+		uint16 outpixels[32*32];
+		glReadPixels(0, 0, 32, 32, GL_RGBA, GL_HALF_FLOAT_OES, outpixels);
+		err = glGetError();
 
         bSupportsColorBufferHalfFloat = fbstatus == GL_FRAMEBUFFER_COMPLETE && err == GL_NO_ERROR;
 
