Don't warn about unsupported OpenGL debug output extensions on Mac

diff --git a/Engine/Source/Runtime/OpenGLDrv/Private/OpenGLDevice.cpp b/Engine/Source/Runtime/OpenGLDrv/Private/OpenGLDevice.cpp
--- a/Engine/Source/Runtime/OpenGLDrv/Private/OpenGLDevice.cpp
+++ b/Engine/Source/Runtime/OpenGLDrv/Private/OpenGLDevice.cpp
@@ -429,7 +429,7 @@ void InitDebugContext()
 		bDebugOutputInitialized = (glGetError() == GL_NO_ERROR);
 	}
 #endif // GL_AMD_debug_output
-	if (!bDebugOutputInitialized)
+	if (!bDebugOutputInitialized && !PLATFORM_MAC)
 	{
 		UE_LOG(LogRHI,Warning,TEXT("OpenGL debug output extension not supported!"));
 	}
