Fix UE-5462 - an erroneous call setting GL_UNPACK_CLIENT_STORAGE_APPLE to GL_TRUE when it needed to be GL_FALSE which was causing errors creating textures. This prevented the atmospheric fog precomputation from being copied into a bulkdata store.

diff --git a/Engine/Source/Runtime/OpenGLDrv/Private/OpenGLTexture.cpp b/Engine/Source/Runtime/OpenGLDrv/Private/OpenGLTexture.cpp
--- a/Engine/Source/Runtime/OpenGLDrv/Private/OpenGLTexture.cpp
+++ b/Engine/Source/Runtime/OpenGLDrv/Private/OpenGLTexture.cpp
@@ -1129,7 +1129,7 @@ void TOpenGLTexture<RHIResourceType>::Unlock(uint32 MipIndex,uint32 ArrayIndex)
 		}
 		if(bUseClientStorage)
 		{
-			glPixelStorei(GL_UNPACK_CLIENT_STORAGE_APPLE, GL_TRUE);
+			glPixelStorei(GL_UNPACK_CLIENT_STORAGE_APPLE, GL_FALSE);
 		}
 		else
 		{
