Fix incorrectly commented out code #html5

diff --git a/Engine/Source/Runtime/OpenGLDrv/Private/HTML5/HTML5OpenGL.cpp b/Engine/Source/Runtime/OpenGLDrv/Private/HTML5/HTML5OpenGL.cpp
--- a/Engine/Source/Runtime/OpenGLDrv/Private/HTML5/HTML5OpenGL.cpp
+++ b/Engine/Source/Runtime/OpenGLDrv/Private/HTML5/HTML5OpenGL.cpp
@@ -100,11 +100,13 @@ void FHTML5OpenGL::ProcessExtensions( const FString& ExtensionsString )
 
         GLenum fbstatus = glCheckFramebufferStatus(GL_FRAMEBUFFER);
 
+// this code causes a stack corruption on win32. 
+#if PLATFORM_HTML5_BROWSER
 		// bSupportsColorBufferHalfFloat is used in such a way that reading the frame buffer is expected to work. Test that now.
-		//uint16 outpixels[32*32];
-	//	glReadPixels(0, 0, 32, 32, GL_RGBA, GL_HALF_FLOAT_OES, outpixels);
-	//	err = glGetError();
-
+		uint16 outpixels[32*32];
+		glReadPixels(0, 0, 32, 32, GL_RGBA, GL_HALF_FLOAT_OES, outpixels);
+		err = glGetError();
+#endif 
         bSupportsColorBufferHalfFloat = fbstatus == GL_FRAMEBUFFER_COMPLETE && err == GL_NO_ERROR;
 
         if (bSupportsColorBufferHalfFloat) {
