Apply "workaround" to Intel as well.

- And enable it on both Windows and Linux.

diff --git a/Engine/Source/Runtime/OpenGLDrv/Private/OpenGLDrv.cpp b/Engine/Source/Runtime/OpenGLDrv/Private/OpenGLDrv.cpp
--- a/Engine/Source/Runtime/OpenGLDrv/Private/OpenGLDrv.cpp
+++ b/Engine/Source/Runtime/OpenGLDrv/Private/OpenGLDrv.cpp
@@ -503,13 +503,13 @@ void FOpenGLBase::ProcessExtensions( const FString& ExtensionsString )
 
 	bSupportsDrawBuffersBlend = ExtensionsString.Contains(TEXT("GL_ARB_draw_buffers_blend"));
 
-	#if PLATFORM_WINDOWS
-		FString VendorName( ANSI_TO_TCHAR((const ANSICHAR*)glGetString(GL_VENDOR) ) );
-		if ( VendorName.Contains(TEXT("ATI ")) )
-		{
-			bAmdWorkaround = true;
-		}
-	#endif
+#if PLATFORM_WINDOWS || PLATFORM_LINUX
+	FString VendorName( ANSI_TO_TCHAR((const ANSICHAR*)glGetString(GL_VENDOR) ) );
+	if (VendorName.Contains(TEXT("ATI ")) || VendorName.Contains(TEXT("Intel ")))
+	{
+		bAmdWorkaround = true;
+	}
+#endif
 }
 
 void GetExtensionsString( FString& ExtensionsString)
