OS X does support GL_ARB_texture_compression_rgtc and we are using it for BC5, so we can use BC4 after all if we hook up the correct definition.

diff --git a/Engine/Source/Developer/Mac/MacTargetPlatform/Private/GenericMacTargetPlatform.h b/Engine/Source/Developer/Mac/MacTargetPlatform/Private/GenericMacTargetPlatform.h
--- a/Engine/Source/Developer/Mac/MacTargetPlatform/Private/GenericMacTargetPlatform.h
+++ b/Engine/Source/Developer/Mac/MacTargetPlatform/Private/GenericMacTargetPlatform.h
@@ -122,12 +122,6 @@ return TSuper::SupportsFeature(Feature);
 			// no textures needed for dedicated server target
 			OutFormats.Add(NAME_None);
 		}
-		else if ( InTexture->CompressionSettings == TC_Alpha)
-		{
-			// no BC4 texture compression on OS X
-			static FName NameG8(TEXT("G8"));
-			OutFormats.Add(NameG8);
-		}
 		else
 		{
 			// just use the standard texture format name for this texture
diff --git a/Engine/Source/Runtime/OpenGLDrv/Private/OpenGLDevice.cpp b/Engine/Source/Runtime/OpenGLDrv/Private/OpenGLDevice.cpp
--- a/Engine/Source/Runtime/OpenGLDrv/Private/OpenGLDevice.cpp
+++ b/Engine/Source/Runtime/OpenGLDrv/Private/OpenGLDevice.cpp
@@ -730,6 +730,7 @@ static void InitRHICapabilitiesForGL()
 		SetupTextureFormat( PF_V8U8,			FOpenGLTextureFormat( GL_RG8_SNORM,				GL_NONE,				GL_RG,			GL_BYTE,							false,	false));
 		SetupTextureFormat( PF_R8G8,			FOpenGLTextureFormat( GL_RG8,					GL_NONE,				GL_RG,			GL_UNSIGNED_BYTE,					false,	false));
 		SetupTextureFormat( PF_BC5,				FOpenGLTextureFormat( GL_COMPRESSED_RG_RGTC2,	GL_COMPRESSED_RG_RGTC2,	GL_RG,			GL_UNSIGNED_BYTE,					true,	false));
+		SetupTextureFormat( PF_BC4,				FOpenGLTextureFormat( GL_COMPRESSED_RED_RGTC1,	GL_COMPRESSED_RED_RGTC1,	GL_RED,			GL_UNSIGNED_BYTE,					true,	false));
 		SetupTextureFormat( PF_A8,				FOpenGLTextureFormat( GL_R8,					GL_NONE,				GL_RED,			GL_UNSIGNED_BYTE,					false,	false));
 		SetupTextureFormat( PF_R32_UINT,		FOpenGLTextureFormat( GL_R32UI,					GL_NONE,				GL_RED_INTEGER,	GL_UNSIGNED_INT,					false,	false));
 		SetupTextureFormat( PF_R32_SINT,		FOpenGLTextureFormat( GL_R32I,					GL_NONE,				GL_RED_INTEGER,	GL_INT,								false,	false));
