On 10.10.1 we don't need to flush after creating every texture, but we do need to flush before resuming the render thread after suspending it, so that any resource changes we have just made are guaranteed to be available to the render thread. Without the flush Apple's OpenGL *can* ignore resources created during the suspension on the shared context, because the commands to allocate their backing store has yet to be processed by the time the render context tries to use them.

diff --git a/Engine/Source/Runtime/Core/Private/Mac/MacPlatformMisc.cpp b/Engine/Source/Runtime/Core/Private/Mac/MacPlatformMisc.cpp
--- a/Engine/Source/Runtime/Core/Private/Mac/MacPlatformMisc.cpp
+++ b/Engine/Source/Runtime/Core/Private/Mac/MacPlatformMisc.cpp
@@ -1065,6 +1065,34 @@ bool FMacPlatformMisc::IsRunningOnMavericks()
 	return GMacAppInfo.RunningOnMavericks;
 }
 
+int32 FMacPlatformMisc::MacOSXVersionCompare(uint8 Major, uint8 Minor, uint8 Revision)
+{
+	TArray<FString> Components;
+	GMacAppInfo.OSVersion.ParseIntoArray(&Components, TEXT("."), true);
+	
+	uint8 TargetValues[3] = {Major, Minor, Revision};
+	uint8 ComponentValues[3] = {0, 0, 0};
+	
+	for(uint32 i = 0; i < Components.Num(); i++)
+	{
+		TTypeFromString<uint8>::FromString(ComponentValues[i], *Components[i]);
+	}
+	
+	for(uint32 i = 0; i < 3; i++)
+	{
+		if(ComponentValues[i] < TargetValues[i])
+		{
+			return -1;
+		}
+		else if(ComponentValues[i] > TargetValues[i])
+		{
+			return 1;
+		}
+	}
+	
+	return 0;
+}
+
 /** Global pointer to crash handler */
 void (* GCrashHandlerPointer)(const FGenericCrashContext& Context) = NULL;
 FMacMallocCrashHandler* GCrashMalloc = nullptr;
diff --git a/Engine/Source/Runtime/Core/Public/Mac/MacPlatformMisc.h b/Engine/Source/Runtime/Core/Public/Mac/MacPlatformMisc.h
--- a/Engine/Source/Runtime/Core/Public/Mac/MacPlatformMisc.h
+++ b/Engine/Source/Runtime/Core/Public/Mac/MacPlatformMisc.h
@@ -152,6 +152,11 @@ struct CORE_API FMacPlatformMisc : public FGenericPlatformMisc
 	 * Returns whether the Mac OS X version is 10.9.x or not.
 	 */
 	static bool IsRunningOnMavericks();
+	
+	/**
+	 * Returns if current < target returns 1, if current > target returns 1, else current == target and returns 0.
+	 */
+	static int32 MacOSXVersionCompare(uint8 Major, uint8 Minor, uint8 Revision);
 
 	static bool bChachedMacMenuStateNeedsUpdate;
 };
diff --git a/Engine/Source/Runtime/OpenGLDrv/Private/Mac/OpenGLMac.cpp b/Engine/Source/Runtime/OpenGLDrv/Private/Mac/OpenGLMac.cpp
--- a/Engine/Source/Runtime/OpenGLDrv/Private/Mac/OpenGLMac.cpp
+++ b/Engine/Source/Runtime/OpenGLDrv/Private/Mac/OpenGLMac.cpp
@@ -27,6 +27,7 @@ static TAutoConsoleVariable<int32> CVarMacUseFrameBufferSRGB(
 
 // @todo: remove once Apple fixes radr://15553950, TTP# 315197
 static int32 GMacFlushTexStorage = true;
+static int32 GMacMustFlushTexStorage = false;
 static FAutoConsoleVariableRef CVarMacFlushTexStorage(
 	TEXT("r.Mac.FlushTexStorage"),
 	GMacFlushTexStorage,
@@ -680,6 +681,10 @@ EOpenGLCurrentContext PlatformOpenGLCurrentContext(FPlatformOpenGLDevice* Device
 
 void PlatformFlushIfNeeded()
 {
+	if([NSOpenGLContext currentContext])
+	{
+		FMacOpenGL::Flush();
+	}
 }
 
 void PlatformRebindResources(FPlatformOpenGLDevice* Device)
@@ -1207,6 +1212,11 @@ void FMacOpenGL::ProcessExtensions(const FString& ExtensionsString)
 		glDrawArraysIndirect = (PFNGLDRAWARRAYSINDIRECTPROC)dlsym(RTLD_SELF, "glDrawArraysIndirect");
 		glDrawElementsIndirect = (PFNGLDRAWELEMENTSINDIRECTPROC)dlsym(RTLD_SELF, "glDrawElementsIndirect");
 	}
+	
+	if(FMacPlatformMisc::MacOSXVersionCompare(10,10,1) < 0)
+	{
+		GMacMustFlushTexStorage = ((FPlatformMisc::IsRunningOnMavericks() && IsRHIDeviceNVIDIA()) || GMacUseMTGL);
+	}
 }
 
 void FMacOpenGL::MacQueryTimestampCounter(GLuint QueryID)
@@ -1325,8 +1335,9 @@ bool FMacOpenGL::MustFlushTexStorage(void)
 {
 	// @todo There is a bug in Apple's GL with TexStorage calls when using MTGL that can see the texture never be created, which then subsequently causes crashes
 	// @todo This bug also affects Nvidia cards under Mavericks without MTGL.
+	// @todo Fixed in 10.10.1.
 	FPlatformOpenGLContext::VerifyCurrentContext();
-	return GMacFlushTexStorage && ((FPlatformMisc::IsRunningOnMavericks() && IsRHIDeviceNVIDIA()) || GMacUseMTGL);
+	return GMacFlushTexStorage || GMacMustFlushTexStorage;
 }
 
 void FMacOpenGL::DeleteTextures(GLsizei Number, const GLuint* Textures)
diff --git a/Engine/Source/Runtime/RenderCore/Private/RenderingThread.cpp b/Engine/Source/Runtime/RenderCore/Private/RenderingThread.cpp
--- a/Engine/Source/Runtime/RenderCore/Private/RenderingThread.cpp
+++ b/Engine/Source/Runtime/RenderCore/Private/RenderingThread.cpp
@@ -137,6 +137,13 @@ FSuspendRenderingThread::FSuspendRenderingThread( bool bInRecreateThread )
 /** Destructor that starts the renderthread again */
 FSuspendRenderingThread::~FSuspendRenderingThread()
 {
+#if PLATFORM_MAC	// On OS X Apple's context sharing is a strict interpretation of the spec. so a resource is only properly visible to other contexts
+					// in the share group after a flush. Thus we call RHIFlushResources which will flush the current context's commands to GL (but not wait for them).
+	ENQUEUE_UNIQUE_RENDER_COMMAND(FlushCommand,
+		RHIFlushResources();
+	);
+#endif
+	
 	if ( bRecreateThread )
 	{
 		GUseThreadedRendering = bUseRenderingThread;
