More fixes for doxygen warnings.

diff --git a/Engine/Source/Editor/PropertyEditor/Private/DetailLayoutBuilderImpl.cpp b/Engine/Source/Editor/PropertyEditor/Private/DetailLayoutBuilderImpl.cpp
--- a/Engine/Source/Editor/PropertyEditor/Private/DetailLayoutBuilderImpl.cpp
+++ b/Engine/Source/Editor/PropertyEditor/Private/DetailLayoutBuilderImpl.cpp
@@ -424,7 +424,7 @@ bool FDetailLayoutBuilderImpl::IsPropertyVisible( TSharedRef<IPropertyHandle> Pr
 	return false;
 }
 
-bool FDetailLayoutBuilderImpl::IsPropertyVisible( const FPropertyAndParent& PropertyAndParent ) const
+bool FDetailLayoutBuilderImpl::IsPropertyVisible( const struct FPropertyAndParent& PropertyAndParent ) const
 {
 	return DetailsView.IsPropertyVisible( PropertyAndParent );
 }
diff --git a/Engine/Source/Editor/UnrealEd/Public/DragAndDrop/DecoratedDragDropOp.h b/Engine/Source/Editor/UnrealEd/Public/DragAndDrop/DecoratedDragDropOp.h
--- a/Engine/Source/Editor/UnrealEd/Public/DragAndDrop/DecoratedDragDropOp.h
+++ b/Engine/Source/Editor/UnrealEd/Public/DragAndDrop/DecoratedDragDropOp.h
@@ -2,6 +2,8 @@
 
 #pragma once
 
+#include "DragAndDrop.h"
+
 #define LOCTEXT_NAMESPACE "DecoratedDragDrop"
 
 class FDecoratedDragDropOp : public FDragDropOperation
diff --git a/Engine/Source/Programs/UnrealDocTool/APIDocTool/APIDocTool/Program.cs b/Engine/Source/Programs/UnrealDocTool/APIDocTool/APIDocTool/Program.cs
--- a/Engine/Source/Programs/UnrealDocTool/APIDocTool/APIDocTool/Program.cs
+++ b/Engine/Source/Programs/UnrealDocTool/APIDocTool/APIDocTool/Program.cs
@@ -94,6 +94,7 @@ namespace APIDocTool
 			"VARARGS=",
 			"VARARG_DECL(FuncRet,StaticFuncRet,Return,FuncName,Pure,FmtType,ExtraDecl,ExtraCall)=FuncRet FuncName(ExtraDecl FmtType Fmt, ...)",
 			"VARARG_BODY(FuncRet,FuncName,FmtType,ExtraDecl)=FuncRet FuncName(ExtraDecl FmtType Fmt, ...)",
+			"VARARG_EXTRA(Arg)=Arg,",
 			"PRAGMA_DISABLE_OPTIMIZATION=",
 			"PRAGMA_ENABLE_OPTIMIZATION=",
 			"NO_API= ",
@@ -135,6 +136,7 @@ namespace APIDocTool
 
 			// Slate declarative syntax
 			"SLATE_BEGIN_ARGS(WidgetType)=public: struct FArguments : public TSlateBaseNamedArgs<WidgetType> { typedef FArguments WidgetArgsType; FArguments()",
+			"HACK_SLATE_SLOT_ARGS(WidgetType)=public: struct FArguments : public TSlateBaseNamedArgs<WidgetType>{ typedef FArguments WidgetArgsType; FArguments()",
 			"SLATE_ATTRIBUTE(AttrType, AttrName)=WidgetArgsType &AttrName( const TAttribute<AttrType>& InAttribute );",
 			"SLATE_TEXT_ATTRIBUTE(AttrName)=WidgetArgsType &AttrName( const TAttribute<FText>& InAttribute ); WidgetArgsType &AttrName( const TAttribute<FString>& InAttribute );",
 			"SLATE_ARGUMENT(ArgType, ArgName)=WidgetArgsType &ArgName(ArgType InArg);",
diff --git a/Engine/Source/Runtime/AIModule/Private/BehaviorTree/BTCompositeNode.cpp b/Engine/Source/Runtime/AIModule/Private/BehaviorTree/BTCompositeNode.cpp
--- a/Engine/Source/Runtime/AIModule/Private/BehaviorTree/BTCompositeNode.cpp
+++ b/Engine/Source/Runtime/AIModule/Private/BehaviorTree/BTCompositeNode.cpp
@@ -45,7 +45,7 @@ int32 UBTCompositeNode::FindChildToExecute(FBehaviorTreeSearchData& SearchData,
 	return RetIdx;
 }
 
-int32 UBTCompositeNode::GetChildIndex(FBehaviorTreeSearchData& SearchData, const UBTNode* ChildNode) const
+int32 UBTCompositeNode::GetChildIndex(struct FBehaviorTreeSearchData& SearchData, const class UBTNode* ChildNode) const
 {
 	if (ChildNode->GetParentNode() != this)
 	{
@@ -56,7 +56,7 @@ int32 UBTCompositeNode::GetChildIndex(FBehaviorTreeSearchData& SearchData, const
 	return GetChildIndex(ChildNode);
 }
 
-int32 UBTCompositeNode::GetChildIndex(const UBTNode* ChildNode) const
+int32 UBTCompositeNode::GetChildIndex(const class UBTNode* ChildNode) const
 {
 	for (int32 ChildIndex = 0; ChildIndex < Children.Num(); ChildIndex++)
 	{
@@ -70,7 +70,7 @@ int32 UBTCompositeNode::GetChildIndex(const UBTNode* ChildNode) const
 	return BTSpecialChild::ReturnToParent;
 }
 
-void UBTCompositeNode::OnChildActivation(FBehaviorTreeSearchData& SearchData, const UBTNode* ChildNode) const
+void UBTCompositeNode::OnChildActivation(struct FBehaviorTreeSearchData& SearchData, const class UBTNode* ChildNode) const
 {
 	OnChildActivation(SearchData, GetChildIndex(SearchData, ChildNode));
 }
diff --git a/Engine/Source/Runtime/Core/Public/Delegates/IntegerSequence.h b/Engine/Source/Runtime/Core/Public/Delegates/IntegerSequence.h
--- a/Engine/Source/Runtime/Core/Public/Delegates/IntegerSequence.h
+++ b/Engine/Source/Runtime/Core/Public/Delegates/IntegerSequence.h
@@ -5,6 +5,9 @@
 template <uint32... Indices>
 struct TIntegerSequence;
 
+// Doxygen can't parse recursive template definitions; just skip it.
+#if !UE_BUILD_DOCS
+
 template <uint32 ToAdd, uint32... Values>
 struct TMakeIntegerSequenceImpl : TMakeIntegerSequenceImpl<ToAdd - 1, Values..., sizeof...(Values)>
 {
@@ -16,5 +19,7 @@ struct TMakeIntegerSequenceImpl<0, Values...>
 	typedef TIntegerSequence<Values...> Type;
 };
 
+#endif
+
 template <uint32 ToAdd>
 using TMakeIntegerSequence = typename TMakeIntegerSequenceImpl<ToAdd>::Type;
diff --git a/Engine/Source/Runtime/Engine/Private/AI/Navigation/NavFilters/NavigationQueryFilter.cpp b/Engine/Source/Runtime/Engine/Private/AI/Navigation/NavFilters/NavigationQueryFilter.cpp
--- a/Engine/Source/Runtime/Engine/Private/AI/Navigation/NavFilters/NavigationQueryFilter.cpp
+++ b/Engine/Source/Runtime/Engine/Private/AI/Navigation/NavFilters/NavigationQueryFilter.cpp
@@ -116,7 +116,7 @@ UNavigationQueryFilter::UNavigationQueryFilter(const FObjectInitializer& ObjectI
 	ExcludeFlags.Packed = 0;
 }
 
-TSharedPtr<const FNavigationQueryFilter> UNavigationQueryFilter::GetQueryFilter(const ANavigationData* NavData) const
+TSharedPtr<const FNavigationQueryFilter> UNavigationQueryFilter::GetQueryFilter(const class ANavigationData* NavData) const
 {
 	TSharedPtr<const FNavigationQueryFilter> SharedFilter = NavData->GetQueryFilter(GetClass());
 	if (!SharedFilter.IsValid())
@@ -133,7 +133,7 @@ TSharedPtr<const FNavigationQueryFilter> UNavigationQueryFilter::GetQueryFilter(
 	return SharedFilter;
 }
 
-void UNavigationQueryFilter::InitializeFilter(const ANavigationData* NavData, FNavigationQueryFilter* Filter) const
+void UNavigationQueryFilter::InitializeFilter(const class ANavigationData* NavData, FNavigationQueryFilter* Filter) const
 {
 	// apply overrides
 	for (int32 i = 0; i < Areas.Num(); i++)
diff --git a/Engine/Source/Runtime/OpenGLDrv/Private/OpenGLDrv.cpp b/Engine/Source/Runtime/OpenGLDrv/Private/OpenGLDrv.cpp
--- a/Engine/Source/Runtime/OpenGLDrv/Private/OpenGLDrv.cpp
+++ b/Engine/Source/Runtime/OpenGLDrv/Private/OpenGLDrv.cpp
@@ -19,6 +19,9 @@ IMPLEMENT_MODULE(FOpenGLDynamicRHIModule, OpenGLDrv);
 /** OpenGL Logging. */
 DEFINE_LOG_CATEGORY(LogOpenGL);
 
+// Ignore functions from RHIMethods.h when parsing documentation; Doxygen's preprocessor can't parse the declaration, so spews warnings for the definitions.
+#if !UE_BUILD_DOCS
+
 void FOpenGLDynamicRHI::RHIPushEvent(const TCHAR* Name)
 {
 #if ENABLE_OPENGL_DEBUG_GROUPS
@@ -47,6 +50,9 @@ void FOpenGLGPUProfiler::PopEvent()
 	FGPUProfiler::PopEvent();
 
 }
+
+#endif
+
 void FOpenGLGPUProfiler::BeginFrame(FOpenGLDynamicRHI* InRHI)
 {
 	CurrentEventNode = NULL;
diff --git a/Engine/Source/Runtime/OpenGLDrv/Private/OpenGLQuery.cpp b/Engine/Source/Runtime/OpenGLDrv/Private/OpenGLQuery.cpp
--- a/Engine/Source/Runtime/OpenGLDrv/Private/OpenGLQuery.cpp
+++ b/Engine/Source/Runtime/OpenGLDrv/Private/OpenGLQuery.cpp
@@ -7,6 +7,9 @@
 #include "OpenGLDrvPrivate.h"
 #include "RenderCore.h"
 
+// Ignore functions from RHIMethods.h when parsing documentation; Doxygen's preprocessor can't parse the declaration, so spews warnings for the definitions.
+#if !UE_BUILD_DOCS
+
 FRenderQueryRHIRef FOpenGLDynamicRHI::RHICreateRenderQuery(ERenderQueryType QueryType)
 {
 	VERIFY_GL_SCOPE();
@@ -202,6 +205,8 @@ bool FOpenGLDynamicRHI::RHIGetRenderQueryResult(FRenderQueryRHIParamRef QueryRHI
 	return bSuccess;
 }
 
+#endif
+
 extern void OnQueryCreation( FOpenGLRenderQuery* Query );
 extern void OnQueryDeletion( FOpenGLRenderQuery* Query );
 
diff --git a/Engine/Source/Runtime/OpenGLDrv/Private/OpenGLRenderTarget.cpp b/Engine/Source/Runtime/OpenGLDrv/Private/OpenGLRenderTarget.cpp
--- a/Engine/Source/Runtime/OpenGLDrv/Private/OpenGLRenderTarget.cpp
+++ b/Engine/Source/Runtime/OpenGLDrv/Private/OpenGLRenderTarget.cpp
@@ -297,6 +297,9 @@ void FOpenGLDynamicRHI::PurgeFramebufferFromCaches( GLuint Framebuffer )
 	}
 }
 
+// Ignore functions from RHIMethods.h when parsing documentation; Doxygen's preprocessor can't parse the declaration, so spews warnings for the definitions.
+#if !UE_BUILD_DOCS
+
 void FOpenGLDynamicRHI::RHICopyToResolveTarget(FTextureRHIParamRef SourceTextureRHI, FTextureRHIParamRef DestTextureRHI, bool bKeepOriginalSurface, const FResolveParams& ResolveParams)
 {
 	if(!SourceTextureRHI || !DestTextureRHI)
@@ -432,6 +435,8 @@ void FOpenGLDynamicRHI::RHICopyToResolveTarget(FTextureRHIParamRef SourceTexture
 	}
 }
 
+#endif
+
 void FOpenGLDynamicRHI::ReadSurfaceDataRaw(FOpenGLContextState& ContextState, FTextureRHIParamRef TextureRHI,FIntRect Rect,TArray<uint8>& OutData, FReadSurfaceDataFlags InFlags)
 {
 	VERIFY_GL_SCOPE();
@@ -696,6 +701,9 @@ void FOpenGLDynamicRHI::ReadSurfaceDataRaw(FOpenGLContextState& ContextState, FT
 	ContextState.Framebuffer = (GLuint)-1;
 }
 
+// Ignore functions from RHIMethods.h when parsing documentation; Doxygen's preprocessor can't parse the declaration, so spews warnings for the definitions.
+#if !UE_BUILD_DOCS
+
 void FOpenGLDynamicRHI::RHIReadSurfaceData(FTextureRHIParamRef TextureRHI,FIntRect Rect,TArray<FColor>& OutData, FReadSurfaceDataFlags InFlags)
 {
 	TArray<uint8> Temp;
@@ -862,6 +870,8 @@ void FOpenGLDynamicRHI::RHIRead3DSurfaceFloatData(FTextureRHIParamRef TextureRHI
 	ContextState.Framebuffer = (GLuint)-1;
 }
 
+#endif
+
 void FOpenGLDynamicRHI::BindPendingFramebuffer( FOpenGLContextState& ContextState )
 {
 	VERIFY_GL_SCOPE();
diff --git a/Engine/Source/Runtime/OpenGLDrv/Private/OpenGLTexture.cpp b/Engine/Source/Runtime/OpenGLDrv/Private/OpenGLTexture.cpp
--- a/Engine/Source/Runtime/OpenGLDrv/Private/OpenGLTexture.cpp
+++ b/Engine/Source/Runtime/OpenGLDrv/Private/OpenGLTexture.cpp
@@ -208,6 +208,9 @@ void FOpenGLDynamicRHI::RHIGetTextureMemoryStats(FTextureMemoryStats& OutStats)
 	OutStats.PendingMemoryAdjustment = 0;
 }
 
+// Ignore functions from RHIMethods.h when parsing documentation; Doxygen's preprocessor can't parse the declaration, so spews warnings for the definitions.
+#if !UE_BUILD_DOCS
+
 /**
  * Fills a texture with to visualize the texture pool memory.
  *
@@ -224,6 +227,8 @@ bool FOpenGLDynamicRHI::RHIGetTextureMemoryVisualizeData( FColor* /*TextureData*
 	return false;
 }
 
+#endif
+
 FRHITexture* FOpenGLDynamicRHI::CreateOpenGLTexture(uint32 SizeX,uint32 SizeY,bool bCubeTexture, bool bArrayTexture, uint8 Format,uint32 NumMips,uint32 NumSamples, uint32 ArraySize, uint32 Flags, FResourceBulkDataInterface* BulkData)
 {
 	VERIFY_GL_SCOPE();
@@ -1737,6 +1742,9 @@ FShaderResourceViewRHIRef FOpenGLDynamicRHI::RHICreateShaderResourceView(FTextur
 	return View;
 }
 
+// Ignore functions from RHIMethods.h when parsing documentation; Doxygen's preprocessor can't parse the declaration, so spews warnings for the definitions.
+#if !UE_BUILD_DOCS
+
 /** Generates mip maps for the surface. */
 void FOpenGLDynamicRHI::RHIGenerateMips(FTextureRHIParamRef SurfaceRHI)
 {
@@ -1761,6 +1769,8 @@ void FOpenGLDynamicRHI::RHIGenerateMips(FTextureRHIParamRef SurfaceRHI)
 	}
 }
 
+#endif
+
 /**
  * Computes the size in memory required by a given texture.
  *
diff --git a/Engine/Source/Runtime/Renderer/Private/ShaderBaseClasses.cpp b/Engine/Source/Runtime/Renderer/Private/ShaderBaseClasses.cpp
--- a/Engine/Source/Runtime/Renderer/Private/ShaderBaseClasses.cpp
+++ b/Engine/Source/Runtime/Renderer/Private/ShaderBaseClasses.cpp
@@ -204,6 +204,9 @@ void FMaterialShader::SetParameters(
 	}
 }
 
+// Doxygen struggles to parse these explicit specializations. Just ignore them for now.
+#if !UE_BUILD_DOCS
+
 #define IMPLEMENT_MATERIAL_SHADER_SetParameters( ShaderRHIParamRef ) \
 	template RENDERER_API void FMaterialShader::SetParameters< ShaderRHIParamRef >( \
 		FRHICommandList& RHICmdList,					\
@@ -222,6 +225,8 @@ IMPLEMENT_MATERIAL_SHADER_SetParameters( FGeometryShaderRHIParamRef );
 IMPLEMENT_MATERIAL_SHADER_SetParameters( FPixelShaderRHIParamRef );
 IMPLEMENT_MATERIAL_SHADER_SetParameters( FComputeShaderRHIParamRef );
 
+#endif
+
 bool FMaterialShader::Serialize(FArchive& Ar)
 {
 	const bool bShaderHasOutdatedParameters = FShader::Serialize(Ar);
