More fixes for doxygen warnings.

diff --git a/Engine/Plugins/ScriptPlugin/Source/ScriptEditorPlugin/Private/ScriptFactory.cpp b/Engine/Plugins/ScriptPlugin/Source/ScriptEditorPlugin/Private/ScriptFactory.cpp
--- a/Engine/Plugins/ScriptPlugin/Source/ScriptEditorPlugin/Private/ScriptFactory.cpp
+++ b/Engine/Plugins/ScriptPlugin/Source/ScriptEditorPlugin/Private/ScriptFactory.cpp
@@ -69,7 +69,7 @@ UObject* UScriptFactory::FactoryCreateText(UClass* InClass, UObject* InParent, F
 
 /** UReimportScriptFactory */
 
-UReimportScriptFactory::UReimportScriptFactory(const class FObjectInitializer& ObjectInitializer)
+UReimportScriptFactory::UReimportScriptFactory(const FObjectInitializer& ObjectInitializer)
 : Super(ObjectInitializer)
 {
 }
diff --git a/Engine/Source/Programs/UnrealDocTool/APIDocTool/APIDocTool/Program.cs b/Engine/Source/Programs/UnrealDocTool/APIDocTool/APIDocTool/Program.cs
--- a/Engine/Source/Programs/UnrealDocTool/APIDocTool/APIDocTool/Program.cs
+++ b/Engine/Source/Programs/UnrealDocTool/APIDocTool/APIDocTool/Program.cs
@@ -148,6 +148,7 @@ namespace APIDocTool
 			"SLATE_NAMED_SLOT(DeclarationType, SlotName)=NamedSlotProperty<DeclarationType> SlotName();",
 			"SLATE_EVENT(DelegateName,EventName)=WidgetArgsType& EventName( const DelegateName& InDelegate );",
 			"SLATE_END_ARGS()=};",
+			"DRAG_DROP_OPERATOR_TYPE(TYPE, BASE)=static const FString& GetTypeId(); virtual bool IsOfTypeImpl(const FString& Type) const;",
 
 			// Rendering macros
 			"IMPLEMENT_SHADER_TYPE(TemplatePrefix,ShaderClass,SourceFilename,FunctionName,Frequency)= ",
diff --git a/Engine/Source/Programs/UnrealDocTool/APIDocTool/APIDocTool/UdnWriter.cs b/Engine/Source/Programs/UnrealDocTool/APIDocTool/APIDocTool/UdnWriter.cs
--- a/Engine/Source/Programs/UnrealDocTool/APIDocTool/APIDocTool/UdnWriter.cs
+++ b/Engine/Source/Programs/UnrealDocTool/APIDocTool/APIDocTool/UdnWriter.cs
@@ -154,7 +154,7 @@ namespace APIDocTool
 
 		public void WritePageHeader(string Title, string Crumbs, string Description)
 		{
-			WriteLine("Title: {0}", Title);
+			WriteLine("Title: {0}", Title.Replace("<", "&lt;").Replace(">", "&gt;"));
 			WriteLine("Crumbs: {0}", Crumbs);
 			WriteLine("Description: {0}", Description.Replace('\n', ' ').Replace('\r', ' '));
 			WriteLine("Availability: Public");
diff --git a/Engine/Source/Runtime/AppFramework/Private/Widgets/Testing/STestSuite.cpp b/Engine/Source/Runtime/AppFramework/Private/Widgets/Testing/STestSuite.cpp
--- a/Engine/Source/Runtime/AppFramework/Private/Widgets/Testing/STestSuite.cpp
+++ b/Engine/Source/Runtime/AppFramework/Private/Widgets/Testing/STestSuite.cpp
@@ -4378,17 +4378,14 @@ public:
 
 };
 
-namespace
+class SResponsiveGridPanelTestWidgetImpl : public SResponsiveGridPanelTestWidget
 {
-	class SResponsiveGridPanelTestWidgetImpl : public SResponsiveGridPanelTestWidget
-	{
-	public:
-		virtual void Construct(const FArguments& InArgs) override;
+public:
+	virtual void Construct(const FArguments& InArgs) override;
 
-	private:
-		TSharedRef<SWidget> ConstructBox(const FString& Text) const;
-	};
-}
+private:
+	TSharedRef<SWidget> ConstructBox(const FString& Text) const;
+};
 
 TSharedRef<SWidget> SResponsiveGridPanelTestWidgetImpl::ConstructBox(const FString& Text) const
 {
diff --git a/Engine/Source/Runtime/Engine/Private/GPUSkinCache.h b/Engine/Source/Runtime/Engine/Private/GPUSkinCache.h
--- a/Engine/Source/Runtime/Engine/Private/GPUSkinCache.h
+++ b/Engine/Source/Runtime/Engine/Private/GPUSkinCache.h
@@ -10,6 +10,7 @@
 #include "Array.h"
 #include "RHIResources.h"
 #include "RenderResource.h"
+#include "UniformBuffer.h"
 
 class FSkeletalMeshObjectGPUSkin;
 
diff --git a/Engine/Source/Runtime/OpenGLDrv/Private/OpenGLUAV.cpp b/Engine/Source/Runtime/OpenGLDrv/Private/OpenGLUAV.cpp
--- a/Engine/Source/Runtime/OpenGLDrv/Private/OpenGLUAV.cpp
+++ b/Engine/Source/Runtime/OpenGLDrv/Private/OpenGLUAV.cpp
@@ -119,8 +119,13 @@ FShaderResourceViewRHIRef FOpenGLDynamicRHI::RHICreateShaderResourceView(FStruct
 	return new FOpenGLShaderResourceView(this,0,GL_TEXTURE_BUFFER);
 }
 
+// Ignore functions from RHIMethods.h when parsing documentation; Doxygen's preprocessor can't parse the declaration, so spews warnings for the definitions.
+#if !UE_BUILD_DOCS
+
 void FOpenGLDynamicRHI::RHIClearUAV(FUnorderedAccessViewRHIParamRef UnorderedAccessViewRHI, const uint32* Values)
 {
 	UE_LOG(LogRHI, Fatal,TEXT("OpenGL RHI doesn't support RHIClearUAV."));
 	GPUProfilingData.RegisterGPUWork(1);
 }
+
+#endif
diff --git a/Engine/Source/Runtime/OpenGLDrv/Private/OpenGLVertexBuffer.cpp b/Engine/Source/Runtime/OpenGLDrv/Private/OpenGLVertexBuffer.cpp
--- a/Engine/Source/Runtime/OpenGLDrv/Private/OpenGLVertexBuffer.cpp
+++ b/Engine/Source/Runtime/OpenGLDrv/Private/OpenGLVertexBuffer.cpp
@@ -202,6 +202,9 @@ void FOpenGLDynamicRHI::RHIUnlockVertexBuffer(FVertexBufferRHIParamRef VertexBuf
 	}
 }
 
+// Ignore functions from RHIMethods.h when parsing documentation; Doxygen's preprocessor can't parse the declaration, so spews warnings for the definitions.
+#if !UE_BUILD_DOCS
+
 void FOpenGLDynamicRHI::RHICopyVertexBuffer(FVertexBufferRHIParamRef SourceBufferRHI,FVertexBufferRHIParamRef DestBufferRHI)
 {
 	VERIFY_GL_SCOPE();
@@ -216,3 +219,5 @@ void FOpenGLDynamicRHI::RHICopyVertexBuffer(FVertexBufferRHIParamRef SourceBuffe
 	glBindBuffer(GL_COPY_READ_BUFFER,0);
 	glBindBuffer(GL_COPY_WRITE_BUFFER,0);
 }
+
+#endif
diff --git a/Engine/Source/Runtime/OpenGLDrv/Private/OpenGLViewport.cpp b/Engine/Source/Runtime/OpenGLDrv/Private/OpenGLViewport.cpp
--- a/Engine/Source/Runtime/OpenGLDrv/Private/OpenGLViewport.cpp
+++ b/Engine/Source/Runtime/OpenGLDrv/Private/OpenGLViewport.cpp
@@ -67,6 +67,9 @@ void FOpenGLDynamicRHI::RHITick( float DeltaTime )
  *	Viewport functions.
  *=============================================================================*/
 
+// Ignore functions from RHIMethods.h when parsing documentation; Doxygen's preprocessor can't parse the declaration, so spews warnings for the definitions.
+#if !UE_BUILD_DOCS
+
 void FOpenGLDynamicRHI::RHIBeginDrawingViewport(FViewportRHIParamRef ViewportRHI, FTextureRHIParamRef RenderTarget)
 {
 	VERIFY_GL_SCOPE();
@@ -180,6 +183,8 @@ bool FOpenGLDynamicRHI::RHIIsDrawingViewport()
 	return DrawingViewport != NULL;
 }
 
+#endif
+
 FTexture2DRHIRef FOpenGLDynamicRHI::RHIGetViewportBackBuffer(FViewportRHIParamRef ViewportRHI)
 {
 	DYNAMIC_CAST_OPENGLRESOURCE(Viewport,Viewport);
diff --git a/Engine/Source/Runtime/Renderer/Private/LightPropagationVolume.h b/Engine/Source/Runtime/Renderer/Private/LightPropagationVolume.h
--- a/Engine/Source/Runtime/Renderer/Private/LightPropagationVolume.h
+++ b/Engine/Source/Runtime/Renderer/Private/LightPropagationVolume.h
@@ -12,6 +12,8 @@
 
 #pragma once
 
+#include "UniformBuffer.h"
+
 #define LPV_MULTIPLE_BOUNCES  1
 #define LPV_VOLUME_TEXTURE    1
 #define LPV_GV_VOLUME_TEXTURE 1
