Make fp rendering work on stable browsers again.
 - existing skylight code won't work - we can't render to nonzero mip levels on webgl 1.0
 - keep exsiting code path, just make sure we don't return uninitialized values for SHcoefficients. things should work automatiically when we move to webgl 2.0
 - remove glReadpixels test in for checking implicit framebuffer completeness check - it won't work consistently across browsers.

diff --git a/Engine/Source/Runtime/OpenGLDrv/Private/HTML5/HTML5OpenGL.cpp b/Engine/Source/Runtime/OpenGLDrv/Private/HTML5/HTML5OpenGL.cpp
--- a/Engine/Source/Runtime/OpenGLDrv/Private/HTML5/HTML5OpenGL.cpp
+++ b/Engine/Source/Runtime/OpenGLDrv/Private/HTML5/HTML5OpenGL.cpp
@@ -87,6 +87,10 @@ void FHTML5OpenGL::ProcessExtensions( const FString& ExtensionsString )
         }
 
         GLuint tex, fb;
+#if PLATFORM_HTML5_WIN32
+		glClearColor( 1.0, 0.5, 0.0,1.0);
+		glClear( GL_COLOR_BUFFER_BIT ); 
+#endif 
         glGenTextures(1, &tex);
         glBindTexture(GL_TEXTURE_2D, tex);
 		glTexParameteri(GL_TEXTURE_2D, GL_TEXTURE_WRAP_S, GL_REPEAT);
@@ -100,12 +104,17 @@ void FHTML5OpenGL::ProcessExtensions( const FString& ExtensionsString )
 
         GLenum fbstatus = glCheckFramebufferStatus(GL_FRAMEBUFFER);
 
-// this code causes a stack corruption on win32. 
-#if PLATFORM_HTML5_BROWSER
-		// bSupportsColorBufferHalfFloat is used in such a way that reading the frame buffer is expected to work. Test that now.
-		uint16 outpixels[32*32];
-		glReadPixels(0, 0, 32, 32, GL_RGBA, GL_HALF_FLOAT_OES, outpixels);
+#if PLATFORM_HTML5_WIN32
+		// keep glReadPixel out of floating point tests for HTML5 Browser builds, glReadPixels doesn't work consistently across browser and is 
+		// hidden behind inconsistent webgl extentions. 
+		TArray<FLinearColor> Data;
+		data.AddUninitialized(32*32); 
+		glViewport(0, 0, 32, 32);
+		glClear(GL_COLOR_BUFFER_BIT);
+		FMemory::Memzero(data.GetData(),32*32*sizeof(FLinearColor));
+		glReadPixels(0, 0, 32, 32, GL_RGBA, GL_FLOAT, Data.GetData());
 		err = glGetError();
+		UE_LOG(LogRHI, Log, TEXT(" %f %f %f %f"), Data[0].R,Data[0].G,Data[0].B,Data[0].A);
 #endif 
         bSupportsColorBufferHalfFloat = fbstatus == GL_FRAMEBUFFER_COMPLETE && err == GL_NO_ERROR;
 
diff --git a/Engine/Source/Runtime/OpenGLDrv/Private/HTML5/HTML5OpenGL.h b/Engine/Source/Runtime/OpenGLDrv/Private/HTML5/HTML5OpenGL.h
--- a/Engine/Source/Runtime/OpenGLDrv/Private/HTML5/HTML5OpenGL.h
+++ b/Engine/Source/Runtime/OpenGLDrv/Private/HTML5/HTML5OpenGL.h
@@ -108,6 +108,7 @@ struct FHTML5OpenGL : public FOpenGLES2
 		return 0;
 	}
 
+	static FORCEINLINE int32 GetReadHalfFloatPixelsEnum()				{ return GL_FLOAT; }
 protected:
 	/** http://www.khronos.org/registry/webgl/extensions/WEBGL_depth_texture/ */
 	static bool bCombinedDepthStencilAttachment;
diff --git a/Engine/Source/Runtime/OpenGLDrv/Private/OpenGLRenderTarget.cpp b/Engine/Source/Runtime/OpenGLDrv/Private/OpenGLRenderTarget.cpp
--- a/Engine/Source/Runtime/OpenGLDrv/Private/OpenGLRenderTarget.cpp
+++ b/Engine/Source/Runtime/OpenGLDrv/Private/OpenGLRenderTarget.cpp
@@ -792,6 +792,7 @@ void FOpenGLDynamicRHI::RHIReadSurfaceFloatData(FTextureRHIParamRef TextureRHI,F
 		TArray<FLinearColor> FloatData;
 		// 4 float components per texel (RGBA)
 		FloatData.AddUninitialized(SizeX * SizeY);
+		FMemory::Memzero(FloatData.GetData(),SizeX * SizeY*sizeof(FLinearColor));
 		glReadPixels(Rect.Min.X, Rect.Min.Y, SizeX, SizeY, GL_RGBA, GL_FLOAT, FloatData.GetData());
 		FLinearColor* FloatDataPtr = FloatData.GetData();
 		for (uint32 Index = 0; Index < SizeX * SizeY; ++Index, ++FloatDataPtr)
