Fix explicit renderer selection when using the Mac Graphics Switching plugin settings on a dual-GPU laptop.

diff --git a/Engine/Source/Runtime/OpenGLDrv/Private/Mac/MacOpenGLContext.cpp b/Engine/Source/Runtime/OpenGLDrv/Private/Mac/MacOpenGLContext.cpp
--- a/Engine/Source/Runtime/OpenGLDrv/Private/Mac/MacOpenGLContext.cpp
+++ b/Engine/Source/Runtime/OpenGLDrv/Private/Mac/MacOpenGLContext.cpp
@@ -119,21 +119,17 @@ static NSOpenGLContext* CreateContext( NSOpenGLContext* SharedContext )
 			bExplicitRendererSetup = true;
 		}
 		
-		// Specify a single explicit renderer ID
-		if (GMacExplicitRendererID && (GNumActiveGPUsForRendering == 1) && !GMacUseAutomaticGraphicsSwitching)
-		{
-			Attributes.Add(NSOpenGLPFARendererID);
-			Attributes.Add(GMacExplicitRendererID);
-		}
-		
-		if(GMacUseAutomaticGraphicsSwitching)
+		if(DisplayMask || GMacExplicitRendererID || GNumActiveGPUsForRendering > 1 || GMacUseAutomaticGraphicsSwitching)
 		{
 			Attributes.Add(kCGLPFASupportsAutomaticGraphicsSwitching);
+			Attributes.Add(NSOpenGLPFAAllowOfflineRenderers);
 		}
 		
-		if(DisplayMask || GMacExplicitRendererID || GNumActiveGPUsForRendering > 1)
+		// Specify a single explicit renderer ID
+		if (GMacExplicitRendererID && (GNumActiveGPUsForRendering == 1) && !GMacUseAutomaticGraphicsSwitching)
 		{
-			Attributes.Add(NSOpenGLPFAAllowOfflineRenderers);
+			Attributes.Add(NSOpenGLPFARendererID);
+			Attributes.Add(GMacExplicitRendererID);
 		}
 		
 		Attributes.Add(0);
