fix html5-win32 build.

diff --git a/Engine/Source/Runtime/OpenGLDrv/Private/HTML5/HTML5OpenGL.cpp b/Engine/Source/Runtime/OpenGLDrv/Private/HTML5/HTML5OpenGL.cpp
--- a/Engine/Source/Runtime/OpenGLDrv/Private/HTML5/HTML5OpenGL.cpp
+++ b/Engine/Source/Runtime/OpenGLDrv/Private/HTML5/HTML5OpenGL.cpp
@@ -108,10 +108,10 @@ void FHTML5OpenGL::ProcessExtensions( const FString& ExtensionsString )
 		// keep glReadPixel out of floating point tests for HTML5 Browser builds, glReadPixels doesn't work consistently across browser and is 
 		// hidden behind inconsistent webgl extentions. 
 		TArray<FLinearColor> Data;
-		data.AddUninitialized(32*32); 
+		Data.AddUninitialized(32*32); 
 		glViewport(0, 0, 32, 32);
 		glClear(GL_COLOR_BUFFER_BIT);
-		FMemory::Memzero(data.GetData(),32*32*sizeof(FLinearColor));
+		FMemory::Memzero(Data.GetData(),32*32*sizeof(FLinearColor));
 		glReadPixels(0, 0, 32, 32, GL_RGBA, GL_FLOAT, Data.GetData());
 		err = glGetError();
 		UE_LOG(LogRHI, Log, TEXT(" %f %f %f %f"), Data[0].R,Data[0].G,Data[0].B,Data[0].A);
