One cause of an increased performance difference between Mac and PC recently was the change to avoid glBufferSubData when using MTGL. This avoided the crash bug [UE-4772] but it is no longer occuring for me on 10.10.1, so use glBufferSubData from that revision onward. On a Mac Pro 4,1/NV 680 GTX MTGL in ElementalDemo is now ~10% faster again.

diff --git a/Engine/Source/Runtime/OpenGLDrv/Private/Mac/OpenGLMac.cpp b/Engine/Source/Runtime/OpenGLDrv/Private/Mac/OpenGLMac.cpp
--- a/Engine/Source/Runtime/OpenGLDrv/Private/Mac/OpenGLMac.cpp
+++ b/Engine/Source/Runtime/OpenGLDrv/Private/Mac/OpenGLMac.cpp
@@ -1348,7 +1348,9 @@ void FMacOpenGL::DeleteTextures(GLsizei Number, const GLuint* Textures)
 
 void FMacOpenGL::BufferSubData(GLenum Target, GLintptr Offset, GLsizeiptr Size, const GLvoid* Data)
 {
-	if(GMacUseMTGL)
+	// @todo The crash that caused UE-4772 is no longer occuring for me on 10.10.1, so use glBufferSubData from that revision onward.
+	static bool const bNeedsMTLGMapBuffer = (FMacPlatformMisc::MacOSXVersionCompare(10,10,1) < 0);
+	if(GMacUseMTGL && bNeedsMTLGMapBuffer)
 	{
 		void* Dest = MapBufferRange(Target, Offset, Size, FOpenGLBase::RLM_WriteOnly);
 		FMemory::Memcpy(Dest, Data, Size);
