[UE-2320] Refactor CompileOpenGLShader string handling.

diff --git a/Engine/Source/Runtime/OpenGLDrv/Private/OpenGLShaders.cpp b/Engine/Source/Runtime/OpenGLDrv/Private/OpenGLShaders.cpp
--- a/Engine/Source/Runtime/OpenGLDrv/Private/OpenGLShaders.cpp
+++ b/Engine/Source/Runtime/OpenGLDrv/Private/OpenGLShaders.cpp
@@ -190,63 +190,32 @@ static const TCHAR* ShaderNameFromShaderType(GLenum ShaderType)
 
 // ============================================================================================================================
 
-
-static ANSICHAR* SkipShaderExtensionText(ANSICHAR* ShaderCode)
+namespace
 {
-	char* ExtensionEnd = ShaderCode;
-	for (;;)
+	inline void AppendCString(TArray<ANSICHAR> & Dest, const ANSICHAR * Source)
 	{
-		// See if we have an extension
-		char* Find = (char*)strstr(ExtensionEnd, "#extension");
-		if (Find)
+		if (Dest.Num() > 0)
 		{
-			// Skip to the end of the line
-			Find = (char*)strstr(Find, "\n");
-			if (Find == NULL)
-			{
-				// Abort if there is no end of line
-				return ShaderCode;
-			}
-			ExtensionEnd = Find;
+			Dest.Insert(Source, FCStringAnsi::Strlen(Source), Dest.Num() - 1);;
 		}
 		else
 		{
-			break;
+			Dest.Append(Source, FCStringAnsi::Strlen(Source) + 1);
 		}
 	}
 
-	if (ExtensionEnd != ShaderCode)
-	{
-		// terminate the extension part of the shader code
-		*ExtensionEnd = '\0';
-		return ExtensionEnd + 1;
-	}
-	else
-	{
-		return NULL;
-	}
-}
-
-// ============================================================================================================================
-
-
-static void ReplaceShaderSubstring(ANSICHAR* ShaderCode, const ANSICHAR* OldString, const ANSICHAR* NewString)
-{
-	char* Curr = ShaderCode;
-	while (Curr != NULL)
+	inline void ReplaceCString(TArray<ANSICHAR> & Dest, const ANSICHAR * Source, const ANSICHAR * Replacement)
 	{
-		char* Substr = (char*)strstr(Curr, OldString);
-		Curr = Substr;
-		if (Substr)
+		int32 SourceLen = FCStringAnsi::Strlen(Source);
+		int32 ReplacementLen = FCStringAnsi::Strlen(Replacement);
+		int32 FoundIndex = 0;
+		for (const ANSICHAR * FoundPointer = FCStringAnsi::Strstr(Dest.GetData(), Source);
+			nullptr != FoundPointer;
+			FoundPointer = FCStringAnsi::Strstr(Dest.GetData()+FoundIndex, Source))
 		{
-			size_t NewStringLength = strlen(NewString);
-			FCStringAnsi::Strcpy(Substr, strlen(Substr), NewString);
-			Curr += NewStringLength;
-			for (size_t index = 0; index < strlen(OldString) - NewStringLength; ++index)
-			{
-				*Curr = ' ';
-				++Curr;
-			}
+			FoundIndex = FoundPointer - Dest.GetData();
+			Dest.RemoveAt(FoundIndex, SourceLen);
+			Dest.Insert(Replacement, ReplacementLen, FoundIndex);
 		}
 	}
 }
@@ -258,10 +227,12 @@ static void ReplaceShaderSubstring(ANSICHAR* ShaderCode, const ANSICHAR* OldStri
 template <typename ShaderType>
 ShaderType* CompileOpenGLShader(const TArray<uint8>& Code)
 {
+	typedef TArray<ANSICHAR> FAnsiCharArray;
+
 	SCOPE_CYCLE_COUNTER(STAT_OpenGLShaderCompileTime);
 	VERIFY_GL_SCOPE();
 
-	ShaderType* Shader = NULL;
+	ShaderType* Shader = nullptr;
 	const GLenum TypeEnum = ShaderType::TypeEnum;
 	FMemoryReader Ar(Code, true);
 	FOpenGLCodeHeader Header = { 0 };
@@ -283,29 +254,28 @@ ShaderType* CompileOpenGLShader(const TArray<uint8>& Code)
 			Header.GlslMarker,
 			Header.FrequencyMarker
 			);
-		return NULL;
+		return nullptr;
 	}
 
 	int32 CodeOffset = Ar.Tell();
-	const ANSICHAR* GlslCode = (ANSICHAR*)Code.GetData() + CodeOffset;
-	GLint GlslCodeLength = Code.Num() - CodeOffset - 1;
-	uint32 GlslCodeCRC = FCrc::MemCrc_DEPRECATED(GlslCode,GlslCodeLength);
 
-#if CHECK_FOR_GL_SHADERS_TO_REPLACE
-	ANSICHAR* NewCode = 0;
-#endif
+	// The code as given to us.
+	FAnsiCharArray GlslCodeOriginal;
+	AppendCString(GlslCodeOriginal, (ANSICHAR*)Code.GetData() + CodeOffset);
+	uint32 GlslCodeOriginalCRC = FCrc::MemCrc_DEPRECATED(GlslCodeOriginal.GetData(), GlslCodeOriginal.Num() + 1);
 
-	FOpenGLCompiledShaderKey Key(TypeEnum,GlslCodeLength,GlslCodeCRC);
+	// The amended code we actually compile.
+	FAnsiCharArray GlslCode;
 
 	// Find the existing compiled shader in the cache.
+	FOpenGLCompiledShaderKey Key(TypeEnum, GlslCodeOriginal.Num(), GlslCodeOriginalCRC);
 	GLuint Resource = GetOpenGLCompiledShaderCache().FindRef(Key);
 	if (!Resource)
 	{
 #if CHECK_FOR_GL_SHADERS_TO_REPLACE
 		{
 			// 1. Check for specific file
-			uint32 CRC = FCrc::MemCrc_DEPRECATED( GlslCode, GlslCodeLength+1 );
-			FString PotentialShaderFileName = FString::Printf( TEXT("%s-%d-0x%x.txt"), ShaderNameFromShaderType(TypeEnum), GlslCodeLength, CRC );
+			FString PotentialShaderFileName = FString::Printf(TEXT("%s-%d-0x%x.txt"), ShaderNameFromShaderType(TypeEnum), GlslCodeOriginal.Num(), GlslCodeOriginalCRC);
 			FString PotentialShaderFile = FPaths::ProfilingDir();
 			PotentialShaderFile *= PotentialShaderFileName;
 
@@ -317,100 +287,73 @@ ShaderType* CompileOpenGLShader(const TArray<uint8>& Code)
 				FArchive* Ar = IFileManager::Get().CreateFileReader(*PotentialShaderFile);
 				if( Ar != NULL )
 				{
+					UE_LOG(LogRHI, Log, TEXT("Replacing %s shader with length %d and CRC 0x%x with the one from a file."), (TypeEnum == GL_VERTEX_SHADER) ? TEXT("vertex") : ((TypeEnum == GL_FRAGMENT_SHADER) ? TEXT("fragment") : TEXT("geometry")), GlslCodeOriginal.Num(), GlslCodeOriginalCRC);
+
 					// read in the file
-					ANSICHAR* NewCode = (ANSICHAR*)FMemory::Malloc( FileSize+1 );
-					Ar->Serialize(NewCode, FileSize);
+					GlslCodeOriginal.Empty();
+					GlslCodeOriginal.AddUninitialized(FileSize + 1);
+					Ar->Serialize(GlslCodeOriginal.GetData(), FileSize);
 					delete Ar;
-					NewCode[FileSize] = 0;
-
-					GlslCode = NewCode;
-					GlslCodeLength = (GLint)(int32)FileSize+1;
-					UE_LOG( LogRHI, Log, TEXT("Replacing %s shader with length %d and CRC 0x%x with the one from a file."), ( TypeEnum == GL_VERTEX_SHADER ) ? TEXT("vertex") : ( ( TypeEnum == GL_FRAGMENT_SHADER ) ? TEXT("fragment") : TEXT("geometry")), GlslCodeLength-1, CRC );
+					GlslCodeOriginal[FileSize] = 0;
 				}
 			}
 		}
 #endif
 
 		Resource = glCreateShader(TypeEnum);
-#if (PLATFORM_ANDROID || PLATFORM_HTML5)
-		ANSICHAR* VersionString = new ANSICHAR[18];
-		memset((void*)VersionString, '\0', 18 * sizeof(ANSICHAR));
 
+#if (PLATFORM_ANDROID || PLATFORM_HTML5)
 		if (IsES2Platform(GMaxRHIShaderPlatform))
 		{
-			//#version 100 has to be the first line in the file, so it has to be added before anything else.
+			// #version NNN has to be the first line in the file, so it has to be added before anything else.
 			if (FOpenGL::UseES30ShadingLanguage())
 			{
-				strcpy(VersionString, "#version 300 es");
-				strcat(VersionString, "\n");
+				AppendCString(GlslCode, "#version 300 es\n");
 			}
 			else 
 			{
-				strcpy(VersionString, "#version 100");
-				strcat(VersionString, "\n");
-			}
-		
-			//remove "#version 100" line,  if found in existing GlslCode. 
-			static const ANSICHAR *Version100String = "#version 100";
-			ANSICHAR* VersionCodePointer = const_cast<ANSICHAR*>(GlslCode);
-			VersionCodePointer = strstr(VersionCodePointer, Version100String);
-			if(VersionCodePointer != NULL)
-			{
-				for(int index=0; index < strlen( Version100String ); index++)
-				{
-					VersionCodePointer[index] = ' ';
-				}
+				AppendCString(GlslCode, "#version 100\n");
 			}
+			ReplaceCString(GlslCodeOriginal, "#version 100", "");
 		}
 #endif
 
 #if PLATFORM_ANDROID 
 		if (IsES2Platform(GMaxRHIShaderPlatform))
 		{
-			//Here is some code to add in a #define for textureCubeLodEXT
-			const ANSICHAR* ExtensionString = "";
-		
-			//This #define fixes compiler errors on Android (which doesnt seem to support textureCubeLodEXT)
-			const ANSICHAR* Prologue = NULL; 
+			// This #define fixes compiler errors on Android (which doesn't seem to support textureCubeLodEXT)
 			if (FOpenGL::UseES30ShadingLanguage())
 			{
 				if (TypeEnum == GL_VERTEX_SHADER)
 				{
-					Prologue = "#define texture2D texture				\n"
-						"#define texture2DProj textureProj		\n"
-						"#define texture2DLod textureLod			\n"
+					AppendCString(GlslCode,
+						"#define texture2D texture \n"
+						"#define texture2DProj textureProj \n"
+						"#define texture2DLod textureLod \n"
 						"#define texture2DProjLod textureProjLod \n"
-						"#define textureCube texture				\n"
-						"#define textureCubeLod textureLod		\n"
-						"#define textureCubeLodEXT textureLod	\n";
+						"#define textureCube texture \n"
+						"#define textureCubeLod textureLod \n"
+						"#define textureCubeLodEXT textureLod \n");
 
-					ReplaceShaderSubstring(const_cast<ANSICHAR*>(GlslCode), "attribute", "in");
-					ReplaceShaderSubstring(const_cast<ANSICHAR*>(GlslCode), "varying", "out");
+					ReplaceCString(GlslCodeOriginal, "attribute", "in");
+					ReplaceCString(GlslCodeOriginal, "varying", "out");
 				} 
 				else if (TypeEnum == GL_FRAGMENT_SHADER)
 				{
-					Prologue = "\n"
-						"#define texture2D texture				\n"
-						"#define texture2DProj textureProj		\n"
-						"#define texture2DLod textureLod			\n"
+					AppendCString(GlslCode,
+						"#define texture2D texture \n"
+						"#define texture2DProj textureProj \n"
+						"#define texture2DLod textureLod \n"
 						"#define texture2DLodEXT textureLod \n"
 						"#define texture2DProjLod textureProjLod \n"
-						"#define textureCube texture				\n"
-						"#define textureCubeLod textureLod		\n"
-						"#define textureCubeLodEXT textureLod	\n"
+						"#define textureCube texture \n"
+						"#define textureCubeLod textureLod \n"
+						"#define textureCubeLodEXT textureLod \n"
 						"\n"
-						"#define gl_FragColor out_FragColor	\n"
-						"out mediump vec4 out_FragColor;\n";
+						"#define gl_FragColor out_FragColor \n"
+						"out mediump vec4 out_FragColor; \n");
 
-					// See if we need to skip any #extension string
-					ANSICHAR* Temp = SkipShaderExtensionText(const_cast<ANSICHAR*>(GlslCode));
-					if (Temp)
-					{
-						ExtensionString = GlslCode;
-						GlslCode = Temp;
-					}
-
-					ReplaceShaderSubstring(const_cast<ANSICHAR*>(GlslCode), "varying", "in");
+					ReplaceCString(GlslCodeOriginal, "varying", "in");
 				}
 			}
 			else if ( (TypeEnum == GL_FRAGMENT_SHADER) &&
@@ -418,65 +361,50 @@ ShaderType* CompileOpenGLShader(const TArray<uint8>& Code)
 			{
 				// Daniel: This device has some shader compiler compatibility issues force them to be disabled
 				//			The cross compiler will put the DONTEMITEXTENSIONSHADERTEXTURELODENABLE define around incompatible sections of code
-				Prologue = "#define DONTEMITEXTENSIONSHADERTEXTURELODENABLE\n"
-					"#define DONTEMITSAMPLERDEFAULTPRECISION\n"
+				AppendCString(GlslCode,
+					"#define DONTEMITEXTENSIONSHADERTEXTURELODENABLE \n"
+					"#define DONTEMITSAMPLERDEFAULTPRECISION \n"
 					"#define texture2DLodEXT(a, b, c) texture2D(a, b) \n"
-				"#define textureCubeLodEXT(a, b, c) textureCube(a, b) \n";
+					"#define textureCubeLodEXT(a, b, c) textureCube(a, b) \n");
 			}
 			else if ( ( TypeEnum == GL_FRAGMENT_SHADER) && 
 				FOpenGL::RequiresTextureCubeLodEXTToTextureCubeLodDefine() )
 			{
-				Prologue = "#define textureCubeLodEXT textureCubeLod \n";
+				AppendCString(GlslCode,
+					"#define textureCubeLodEXT textureCubeLod \n");
 			}
 			else if(!FOpenGL::SupportsShaderTextureLod() || !FOpenGL::SupportsShaderTextureCubeLod())
 			{
-				Prologue = 	"#define texture2DLodEXT(a, b, c) texture2D(a, b) \n"
-							"#define textureCubeLodEXT(a, b, c) textureCube(a, b) \n";
-			}
-			else 
-			{
-				Prologue = "";
+				AppendCString(GlslCode,
+					"#define texture2DLodEXT(a, b, c) texture2D(a, b) \n"
+					"#define textureCubeLodEXT(a, b, c) textureCube(a, b) \n");
 			}
-
-			//UE_LOG(LogRHI, Display,TEXT("Prologue is %s"), ANSI_TO_TCHAR( Prologue ));
-
-			//Assemble the source strings into an array to pass into the compiler.
-			const GLchar* ShaderSourceStrings[4] = { VersionString, ExtensionString, Prologue, GlslCode };
-			const GLint ShaderSourceLen[4] = { (GLint)(strlen(VersionString)), (GLint)(strlen(ExtensionString)), (GLint)(strlen(Prologue)), (GLint)(strlen(GlslCode)) };
-			glShaderSource(Resource, 4, (const GLchar**)&ShaderSourceStrings, ShaderSourceLen);		
-			glCompileShader(Resource);
-		}
-		else
-		{
-			glShaderSource(Resource, 1, (const GLchar**)&GlslCode, &GlslCodeLength);
-			glCompileShader(Resource);
 		}
 
 #elif PLATFORM_HTML5
 
 		// HTML5 use case is much simpler, use a separate chunk of code from android. 
-		ANSICHAR Prologue[1024] = ""; 
-
-		if (!FOpenGL::SupportsShaderTextureLod()) 
+		if (!FOpenGL::SupportsShaderTextureLod())
 		{
-			strcat(Prologue, "#define  DONTEMITEXTENSIONSHADERTEXTURELODENABLE \n");
-			strcat(Prologue, "#define texture2DLodEXT(a, b, c) texture2D(a, b) \n");
-			strcat(Prologue, "#define textureCubeLodEXT(a, b, c) textureCube(a, b)\n");
+			AppendCString(GlslCode,
+				"#define DONTEMITEXTENSIONSHADERTEXTURELODENABLE \n"
+				"#define texture2DLodEXT(a, b, c) texture2D(a, b) \n"
+				"#define textureCubeLodEXT(a, b, c) textureCube(a, b) \n");
 		}
 
-		const GLchar* ShaderSourceStrings[3] = { VersionString, Prologue, GlslCode };
-		const GLint ShaderSourceLen[3] = { (GLint)(strlen(VersionString)), (GLint)(strlen(Prologue)), (GLint)(strlen(GlslCode)) };
-
-		glShaderSource(Resource, 3, (const GLchar**)&ShaderSourceStrings, ShaderSourceLen);		
-		glCompileShader(Resource);
+#endif
 
-#else
+		// Append the possibly edited shader to the one we will compile.
+		// This is to make it easier to debug as we can see the whole
+		// shader source.
+		AppendCString(GlslCode, "\n\n");
+		AppendCString(GlslCode, GlslCodeOriginal.GetData());
 
-		glShaderSource(Resource, 1, (const GLchar**)&GlslCode, &GlslCodeLength);
+		const ANSICHAR * GlslCodeString = GlslCode.GetData();
+		int32 GlslCodeLength = GlslCode.Num();
+		glShaderSource(Resource, 1, (const GLchar**)&GlslCodeString, &GlslCodeLength);
 		glCompileShader(Resource);
 
-#endif
-
 		GetOpenGLCompiledShaderCache().Add(Key,Resource);
 	}
 
@@ -486,19 +414,10 @@ ShaderType* CompileOpenGLShader(const TArray<uint8>& Code)
 	Shader->UniformBuffersCopyInfo = Header.UniformBuffersCopyInfo;
 
 #if DEBUG_GL_SHADERS
-	Shader->GlslCode.Empty(GlslCodeLength + 1);
-	Shader->GlslCode.AddUninitialized(GlslCodeLength + 1);
-	FMemory::Memcpy(Shader->GlslCode.GetData(), GlslCode, GlslCodeLength + 1);
+	Shader->GlslCode = GlslCode;
 	Shader->GlslCodeString = (ANSICHAR*)Shader->GlslCode.GetData();
 #endif
 
-#if CHECK_FOR_GL_SHADERS_TO_REPLACE
-	if( NewCode )
-	{
-		FMemory::Free( NewCode );
-	}
-#endif
-
 	return Shader;
 }
 
