Fix rendering on the AMD 4xx0 series GPUs which apparently can't support depth fetch during depth test or combined clears using glClearBuffer.

diff --git a/Engine/Source/Runtime/OpenGLDrv/Private/Mac/MacOpenGLContext.cpp b/Engine/Source/Runtime/OpenGLDrv/Private/Mac/MacOpenGLContext.cpp
--- a/Engine/Source/Runtime/OpenGLDrv/Private/Mac/MacOpenGLContext.cpp
+++ b/Engine/Source/Runtime/OpenGLDrv/Private/Mac/MacOpenGLContext.cpp
@@ -429,22 +429,24 @@ void FPlatformOpenGLContext::VerifyCurrentContext()
 				PlatformContext->VendorID = 0x1002;
 			}
 			
-			if(FPlatformMisc::IsRunningOnMavericks())
+			// Renderer IDs matchup to driver kexts, so switching based on them will allow us to target workarouds to many GPUs
+			// which exhibit the same unfortunate driver bugs without having to parse their individual ID strings.
+			switch((PlatformContext->RendererID & kCGLRendererIDMatchingMask))
 			{
-				// Renderer IDs matchup to driver kexts, so switching based on them will allow us to target workarouds to many GPUs
-				// which exhibit the same unfortunate driver bugs without having to parse their individual ID strings.
-				switch((PlatformContext->RendererID & kCGLRendererIDMatchingMask))
+				case kCGLRendererATIRadeonX2000ID:
 				{
-					case kCGLRendererATIRadeonX4000ID:
-					{
-						// @todo: remove once AMD fix the AMDX4000 driver for GCN cards so that it is possible to sample the depth while also stencil testing to the same DEPTH_STENCIL texture - it works on all other cards we have.
-						PlatformContext->SupportsDepthFetchDuringDepthTest = false;
-						break;
-					}
-					default:
-					{
-						PlatformContext->SupportsDepthFetchDuringDepthTest = true;
-					}
+					PlatformContext->SupportsDepthFetchDuringDepthTest = false;
+					break;
+				}
+				case kCGLRendererATIRadeonX4000ID:
+				{
+					// @todo: remove once AMD fix the AMDX4000 driver for GCN cards so that it is possible to sample the depth while also stencil testing to the same DEPTH_STENCIL texture - it works on all other cards we have.
+					PlatformContext->SupportsDepthFetchDuringDepthTest = FPlatformMisc::IsRunningOnMavericks() ? false : true;
+					break;
+				}
+				default:
+				{
+					PlatformContext->SupportsDepthFetchDuringDepthTest = true;
 				}
 			}
 		}
diff --git a/Engine/Source/Runtime/OpenGLDrv/Private/Mac/OpenGLMac.h b/Engine/Source/Runtime/OpenGLDrv/Private/Mac/OpenGLMac.h
--- a/Engine/Source/Runtime/OpenGLDrv/Private/Mac/OpenGLMac.h
+++ b/Engine/Source/Runtime/OpenGLDrv/Private/Mac/OpenGLMac.h
@@ -67,7 +67,7 @@ public:
 	
 	static FORCEINLINE void ClearBufferfi(GLenum Buffer, GLint DrawBufferIndex, GLfloat Depth, GLint Stencil)
 	{
-		if(FPlatformMisc::IsRunningOnMavericks())
+		if(FPlatformMisc::IsRunningOnMavericks() || !GSupportsDepthFetchDuringDepthTest)
 		{
 			switch (Buffer)
 			{
@@ -206,7 +206,7 @@ public:
 		}
 		else
 		{
-			UE_LOG(LogRHI, Fatal, TEXT("OpenGL state on draw requires setting different blend operation or factors to different render targets. This is not supported on Mac OS X GL 3.2 contexts!"));
+			UE_LOG(LogRHI, Fatal, TEXT("OpenGL state on draw requires setting different blend operation or factors to different render targets. This is not supported on Mac OS X GL 3 contexts!"));
 		}
 	}
 	
@@ -218,7 +218,7 @@ public:
 		}
 		else
 		{
-			UE_LOG(LogRHI, Fatal, TEXT("OpenGL state on draw requires setting different blend operation or factors to different render targets. This is not supported on Mac OS X GL 3.2 contexts!"));
+			UE_LOG(LogRHI, Fatal, TEXT("OpenGL state on draw requires setting different blend operation or factors to different render targets. This is not supported on Mac OS X GL 3 contexts!"));
 		}
 	}
 	
@@ -230,7 +230,7 @@ public:
 		}
 		else
 		{
-			UE_LOG(LogRHI, Fatal, TEXT("OpenGL state on draw requires setting different blend operation or factors to different render targets. This is not supported on Mac OS X GL 3.2 contexts!"));
+			UE_LOG(LogRHI, Fatal, TEXT("OpenGL state on draw requires setting different blend operation or factors to different render targets. This is not supported on Mac OS X GL 3 contexts!"));
 		}
 	}
 	
@@ -242,7 +242,7 @@ public:
 		}
 		else
 		{
-			UE_LOG(LogRHI, Fatal, TEXT("OpenGL state on draw requires setting different blend operation or factors to different render targets. This is not supported on Mac OS X GL 3.2 contexts!"));
+			UE_LOG(LogRHI, Fatal, TEXT("OpenGL state on draw requires setting different blend operation or factors to different render targets. This is not supported on Mac OS X GL 3 contexts!"));
 		}
 	}
 	
@@ -254,7 +254,7 @@ public:
 		}
 		else
 		{
-			UE_LOG(LogRHI, Fatal, TEXT("Tessellation not supported on Mac OS X GL 3.2 contexts!"));
+			UE_LOG(LogRHI, Fatal, TEXT("Tessellation not supported on Mac OS X GL 3 contexts!"));
 		}
 	}
 	
@@ -263,6 +263,30 @@ public:
 		glDepthBoundsEXT( Min, Max);
 	}
 	
+	static FORCEINLINE void DrawArraysIndirect (GLenum Mode, const void *Offset)
+	{
+		if(glDrawArraysIndirect)
+		{
+			glDrawArraysIndirect( Mode, Offset);
+		}
+		else
+		{
+			UE_LOG(LogRHI, Fatal, TEXT("Draw Indirect not supported on Mac OS X GL 3 contexts!"));
+		}
+	}
+	
+	static FORCEINLINE void DrawElementsIndirect (GLenum Mode, GLenum Type, const void *Offset)
+	{
+		if(glDrawElementsIndirect)
+		{
+			glDrawElementsIndirect( Mode, Type, Offset);
+		}
+		else
+		{
+			UE_LOG(LogRHI, Fatal, TEXT("Draw Indirect not supported on Mac OS X GL 3 contexts!"));
+		}
+	}
+	
 private:
 	/** Detects apitrace interception library & equivalents */
 	static bool bUsingApitrace;
