- Added missing api function (GetReadHalfFloatPixelsEnum()) to FOpenGLES31

diff --git a/Engine/Source/Runtime/OpenGLDrv/Public/OpenGLES31.h b/Engine/Source/Runtime/OpenGLDrv/Public/OpenGLES31.h
--- a/Engine/Source/Runtime/OpenGLDrv/Public/OpenGLES31.h
+++ b/Engine/Source/Runtime/OpenGLDrv/Public/OpenGLES31.h
@@ -804,6 +804,8 @@ struct FOpenGLES31 : public FOpenGLBase
 	static void ProcessQueryGLInt();
 	static void ProcessExtensions(const FString& ExtensionsString);
 
+	static FORCEINLINE int32 GetReadHalfFloatPixelsEnum() { return GL_HALF_FLOAT; }
+
 protected:
 	static GLsizei NextTextureName;
 	static GLuint TextureNamesCache[OPENGL_NAME_CACHE_SIZE];
