Added checks for multisampled render to texture support in ES2.

diff --git a/Engine/Source/Runtime/OpenGLDrv/Private/Android/AndroidOpenGL.cpp b/Engine/Source/Runtime/OpenGLDrv/Private/Android/AndroidOpenGL.cpp
--- a/Engine/Source/Runtime/OpenGLDrv/Private/Android/AndroidOpenGL.cpp
+++ b/Engine/Source/Runtime/OpenGLDrv/Private/Android/AndroidOpenGL.cpp
@@ -28,6 +28,8 @@ PFNGLGETQUERYOBJECTUI64VEXTPROC			glGetQueryObjectui64vEXT = NULL;
 // Offscreen MSAA rendering
 PFNBLITFRAMEBUFFERNVPROC				glBlitFramebufferNV = NULL;
 PFNGLDISCARDFRAMEBUFFEREXTPROC			glDiscardFramebufferEXT = NULL;
+PFNGLFRAMEBUFFERTEXTURE2DMULTISAMPLEEXT	glFramebufferTexture2DMultisampleEXT = NULL;
+PFNGLRENDERBUFFERSTORAGEMULTISAMPLEEXT	glRenderbufferStorageMultisampleEXT = NULL;
 
 PFNGLPUSHGROUPMARKEREXTPROC				glPushGroupMarkerEXT = NULL;
 PFNGLPOPGROUPMARKEREXTPROC				glPopGroupMarkerEXT = NULL;
@@ -340,6 +342,8 @@ void FAndroidOpenGL::ProcessExtensions(const FString& ExtensionsString)
 	}
 
 	glDiscardFramebufferEXT = (PFNGLDISCARDFRAMEBUFFEREXTPROC)((void*)eglGetProcAddress("glDiscardFramebufferEXT"));
+	glFramebufferTexture2DMultisampleEXT = (PFNGLFRAMEBUFFERTEXTURE2DMULTISAMPLEEXT)((void*)eglGetProcAddress("glFramebufferTexture2DMultisampleEXT"));
+	glRenderbufferStorageMultisampleEXT = (PFNGLRENDERBUFFERSTORAGEMULTISAMPLEEXT)((void*)eglGetProcAddress("glRenderbufferStorageMultisampleEXT"));
 	glPushGroupMarkerEXT = (PFNGLPUSHGROUPMARKEREXTPROC)((void*)eglGetProcAddress("glPushGroupMarkerEXT"));
 	glPopGroupMarkerEXT = (PFNGLPOPGROUPMARKEREXTPROC)((void*)eglGetProcAddress("glPopGroupMarkerEXT"));
 	glLabelObjectEXT = (PFNGLLABELOBJECTEXTPROC)((void*)eglGetProcAddress("glLabelObjectEXT"));
diff --git a/Engine/Source/Runtime/OpenGLDrv/Private/Android/AndroidOpenGL.h b/Engine/Source/Runtime/OpenGLDrv/Private/Android/AndroidOpenGL.h
--- a/Engine/Source/Runtime/OpenGLDrv/Private/Android/AndroidOpenGL.h
+++ b/Engine/Source/Runtime/OpenGLDrv/Private/Android/AndroidOpenGL.h
@@ -76,6 +76,8 @@ typedef void (GL_APIENTRYP PFNGLPUSHGROUPMARKEREXTPROC) (GLsizei length, const G
 typedef void (GL_APIENTRYP PFNGLLABELOBJECTEXTPROC) (GLenum type, GLuint object, GLsizei length, const GLchar *label);
 typedef void (GL_APIENTRYP PFNGLGETOBJECTLABELEXTPROC) (GLenum type, GLuint object, GLsizei bufSize, GLsizei *length, GLchar *label);
 typedef void (GL_APIENTRYP PFNGLPOPGROUPMARKEREXTPROC) (void);
+typedef void (GL_APIENTRYP PFNGLFRAMEBUFFERTEXTURE2DMULTISAMPLEEXT) (GLenum target, GLenum attachment, GLenum textarget, GLuint texture, GLint level, GLsizei samples);
+typedef void (GL_APIENTRYP PFNGLRENDERBUFFERSTORAGEMULTISAMPLEEXT) (GLenum target, GLsizei samples, GLenum internalformat, GLsizei width, GLsizei height);
 /** from ES 3.0 but can be called on certain Adreno devices */
 typedef void (GL_APIENTRYP PFNGLTEXSTORAGE2DPROC) (GLenum target, GLsizei levels, GLenum internalformat, GLsizei width, GLsizei height);
 
@@ -92,6 +94,8 @@ extern PFNGLGETQUERYOBJECTUI64VEXTPROC	glGetQueryObjectui64vEXT;
 extern PFNGLMAPBUFFEROESPROC			glMapBufferOES;
 extern PFNGLUNMAPBUFFEROESPROC			glUnmapBufferOES;
 extern PFNGLDISCARDFRAMEBUFFEREXTPROC 	glDiscardFramebufferEXT ;
+extern PFNGLFRAMEBUFFERTEXTURE2DMULTISAMPLEEXT	glFramebufferTexture2DMultisampleEXT;
+extern PFNGLRENDERBUFFERSTORAGEMULTISAMPLEEXT	glRenderbufferStorageMultisampleEXT;
 extern PFNGLPUSHGROUPMARKEREXTPROC		glPushGroupMarkerEXT;
 extern PFNGLLABELOBJECTEXTPROC			glLabelObjectEXT;
 extern PFNGLGETOBJECTLABELEXTPROC		glGetObjectLabelEXT;
diff --git a/Engine/Source/Runtime/OpenGLDrv/Private/OpenGLES2.cpp b/Engine/Source/Runtime/OpenGLDrv/Private/OpenGLES2.cpp
--- a/Engine/Source/Runtime/OpenGLDrv/Private/OpenGLES2.cpp
+++ b/Engine/Source/Runtime/OpenGLDrv/Private/OpenGLES2.cpp
@@ -58,6 +58,9 @@ bool FOpenGLES2::bSupportsColorBufferHalfFloat = false;
 /** GL_EXT_shader_framebuffer_fetch */
 bool FOpenGLES2::bSupportsShaderFramebufferFetch = false;
 
+/** GL_EXT_multisampled_render_to_texture */
+bool FOpenGLES2::bSupportsMultisampledRenderToTexture = false;
+
 /** GL_EXT_sRGB */
 bool FOpenGLES2::bSupportsSGRB = false;
 
@@ -164,6 +167,7 @@ void FOpenGLES2::ProcessExtensions( const FString& ExtensionsString )
 	bSupportsSGRB = ExtensionsString.Contains(TEXT("GL_EXT_sRGB"));
 	bSupportsColorBufferHalfFloat = ExtensionsString.Contains(TEXT("GL_EXT_color_buffer_half_float"));
 	bSupportsShaderFramebufferFetch = ExtensionsString.Contains(TEXT("GL_EXT_shader_framebuffer_fetch")) || ExtensionsString.Contains(TEXT("GL_NV_shader_framebuffer_fetch"));
+	bSupportsMultisampledRenderToTexture = ExtensionsString.Contains(TEXT("GL_EXT_multisampled_render_to_texture"));
 	// @todo ios7: SRGB support does not work with our texture format setup (ES2 docs indicate that internalFormat and format must match, but they don't at all with sRGB enabled)
 	//             One possible solution us to use GLFormat.InternalFormat[bSRGB] instead of GLFormat.Format
 	bSupportsSGRB = false;//ExtensionsString.Contains(TEXT("GL_EXT_sRGB"));
diff --git a/Engine/Source/Runtime/OpenGLDrv/Public/OpenGLES2.h b/Engine/Source/Runtime/OpenGLDrv/Public/OpenGLES2.h
--- a/Engine/Source/Runtime/OpenGLDrv/Public/OpenGLES2.h
+++ b/Engine/Source/Runtime/OpenGLDrv/Public/OpenGLES2.h
@@ -92,6 +92,7 @@ struct FOpenGLES2 : public FOpenGLBase
 	static FORCEINLINE bool SupportsTextureHalfFloat()					{ return bSupportsTextureHalfFloat; }
 	static FORCEINLINE bool SupportsColorBufferHalfFloat()				{ return bSupportsColorBufferHalfFloat; }
 	static FORCEINLINE bool SupportsShaderFramebufferFetch()			{ return bSupportsShaderFramebufferFetch; }
+	static FORCEINLINE bool SupportsMultisampledRenderToTexture()		{ return bSupportsMultisampledRenderToTexture; }
 	static FORCEINLINE bool SupportsVertexArrayBGRA()					{ return false; }
 	static FORCEINLINE bool SupportsBGRA8888()							{ return bSupportsBGRA8888; }
 	static FORCEINLINE bool SupportsSRGB()								{ return bSupportsSGRB; }
@@ -416,6 +417,9 @@ protected:
 	/** GL_EXT_shader_framebuffer_fetch */
 	static bool bSupportsShaderFramebufferFetch;
 
+	/** GL_EXT_MULTISAMPLED_RENDER_TO_TEXTURE */
+	static bool bSupportsMultisampledRenderToTexture;
+
 	/** GL_FRAGMENT_SHADER, GL_LOW_FLOAT */
 	static int ShaderLowPrecision;
 
