Fix for off-by-one error in OpenGL shader string handling.
WebGL doesn't like null characters in string source.

diff --git a/Engine/Source/Runtime/OpenGLDrv/Private/OpenGLShaders.cpp b/Engine/Source/Runtime/OpenGLDrv/Private/OpenGLShaders.cpp
--- a/Engine/Source/Runtime/OpenGLDrv/Private/OpenGLShaders.cpp
+++ b/Engine/Source/Runtime/OpenGLDrv/Private/OpenGLShaders.cpp
@@ -401,7 +401,7 @@ ShaderType* CompileOpenGLShader(const TArray<uint8>& Code)
 		AppendCString(GlslCode, GlslCodeOriginal.GetData());
 
 		const ANSICHAR * GlslCodeString = GlslCode.GetData();
-		int32 GlslCodeLength = GlslCode.Num();
+		int32 GlslCodeLength = GlslCode.Num() - 1;
 		glShaderSource(Resource, 1, (const GLchar**)&GlslCodeString, &GlslCodeLength);
 		glCompileShader(Resource);
 
