UE4 - UE-4944 - Temp workaround for GL crash on packaged builds with nvidia driver

diff --git a/Engine/Source/Runtime/OpenGLDrv/Private/OpenGLCommands.cpp b/Engine/Source/Runtime/OpenGLDrv/Private/OpenGLCommands.cpp
--- a/Engine/Source/Runtime/OpenGLDrv/Private/OpenGLCommands.cpp
+++ b/Engine/Source/Runtime/OpenGLDrv/Private/OpenGLCommands.cpp
@@ -46,7 +46,8 @@ namespace OpenGLConsoleVariables
 		TEXT("If true, don't issue dispatch work.")
 		);
 
-	int32 bUseVAB = PLATFORM_LINUX ? 0 : 1;	// [RCL] @FIXME - currently (2014-09-19) causes crashes
+	//@todo-rco: Workaround Nvidia driver crash
+	int32 bUseVAB = (PLATFORM_LINUX || (PLATFORM_DESKTOP && IsRHIDeviceNVIDIA())) ? 0 : 1;	// [RCL] @FIXME - currently (2014-09-19) causes crashes
 	static FAutoConsoleVariableRef CVarUseVAB(
 		TEXT("OpenGL.UseVAB"),
 		bUseVAB,
