[UE-6680] Fix #entension directives getting composed after non-preprocessor GLSL directives.

diff --git a/Engine/Source/Runtime/OpenGLDrv/Private/OpenGLShaders.cpp b/Engine/Source/Runtime/OpenGLDrv/Private/OpenGLShaders.cpp
--- a/Engine/Source/Runtime/OpenGLDrv/Private/OpenGLShaders.cpp
+++ b/Engine/Source/Runtime/OpenGLDrv/Private/OpenGLShaders.cpp
@@ -218,6 +218,59 @@ namespace
 			Dest.Insert(Replacement, ReplacementLen, FoundIndex);
 		}
 	}
+
+	inline const ANSICHAR * CStringEndOfLine(const ANSICHAR * Text)
+	{
+		const ANSICHAR * LineEnd = FCStringAnsi::Strchr(Text, '\n');
+		if (nullptr == LineEnd)
+		{
+			LineEnd = Text + FCStringAnsi::Strlen(Text);
+		}
+		return LineEnd;
+	}
+
+	inline bool MoveHashLines(TArray<ANSICHAR> & Dest, TArray<ANSICHAR> & Source)
+	{
+		// Walk through the lines to find the first non-# line...
+		const ANSICHAR * LineStart = Source.GetData();
+		for (bool FoundNonHashLine = false; !FoundNonHashLine;)
+		{
+			const ANSICHAR * LineEnd = CStringEndOfLine(LineStart);
+			if (LineStart[0] != '#' && LineStart[0] != '\n')
+			{
+				FoundNonHashLine = true;
+			}
+			else if (LineEnd[0] == '\n')
+			{
+				LineStart = LineEnd + 1;
+			}
+			else
+			{
+				LineStart = LineEnd;
+			}
+		}
+		// Copy the hash lines over, if we found any. And delete from
+		// the source.
+		if (LineStart > Source.GetData())
+		{
+			int32 LineLength = LineStart - Source.GetData();
+			if (Dest.Num() > 0)
+			{
+				Dest.Insert(Source.GetData(), LineLength, Dest.Num() - 1);
+			}
+			else
+			{
+				Dest.Append(Source.GetData(), LineLength);
+			}
+			if (Dest.Last() != '\n')
+			{
+				Dest.Append("\n", 1);
+			}
+			Source.RemoveAt(0, LineStart - Source.GetData());
+			return true;
+		}
+		return false;
+	}
 }
 
 /**
@@ -340,6 +393,12 @@ ShaderType* CompileOpenGLShader(const TArray<uint8>& Code)
 				} 
 				else if (TypeEnum == GL_FRAGMENT_SHADER)
 				{
+					// #extension directives have to come before any non-# directives. Because
+					// we add non-# stuff below and the #extension directives
+					// get added to the incoming shader source we move any # directives
+					// to be right after the #version to ensure they are always correct.
+					MoveHashLines(GlslCode, GlslCodeOriginal);
+
 					AppendCString(GlslCode,
 						"#define texture2D texture \n"
 						"#define texture2DProj textureProj \n"
