[UE-6680] Fix generation of ES30 shaders that have blank lines or extra EOL whitespace.

diff --git a/Engine/Source/Runtime/OpenGLDrv/Private/OpenGLShaders.cpp b/Engine/Source/Runtime/OpenGLDrv/Private/OpenGLShaders.cpp
--- a/Engine/Source/Runtime/OpenGLDrv/Private/OpenGLShaders.cpp
+++ b/Engine/Source/Runtime/OpenGLDrv/Private/OpenGLShaders.cpp
@@ -229,6 +229,19 @@ namespace
 		return LineEnd;
 	}
 
+	inline bool CStringIsBlankLine(const ANSICHAR * Text)
+	{
+		while (!FCharAnsi::IsLinebreak(*Text))
+		{
+			if (!FCharAnsi::IsWhitespace(*Text))
+			{
+				return false;
+			}
+			++Text;
+		}
+		return true;
+	}
+
 	inline bool MoveHashLines(TArray<ANSICHAR> & Dest, TArray<ANSICHAR> & Source)
 	{
 		// Walk through the lines to find the first non-# line...
@@ -236,7 +249,7 @@ namespace
 		for (bool FoundNonHashLine = false; !FoundNonHashLine;)
 		{
 			const ANSICHAR * LineEnd = CStringEndOfLine(LineStart);
-			if (LineStart[0] != '#' && LineStart[0] != '\n')
+			if (LineStart[0] != '#' && !CStringIsBlankLine(LineStart))
 			{
 				FoundNonHashLine = true;
 			}
@@ -261,10 +274,11 @@ namespace
 			else
 			{
 				Dest.Append(Source.GetData(), LineLength);
+				Dest.Append("", 1);
 			}
-			if (Dest.Last() != '\n')
+			if (Dest.Last(1) != '\n')
 			{
-				Dest.Append("\n", 1);
+				Dest.Insert("\n", 1, Dest.Num() - 1);
 			}
 			Source.RemoveAt(0, LineStart - Source.GetData());
 			return true;
