Integration of CL 2385626 from NVIDIA which should fix some bugs related to vertex attributes.

diff --git a/Engine/Source/Runtime/OpenGLDrv/Private/OpenGLCommands.cpp b/Engine/Source/Runtime/OpenGLDrv/Private/OpenGLCommands.cpp
--- a/Engine/Source/Runtime/OpenGLDrv/Private/OpenGLCommands.cpp
+++ b/Engine/Source/Runtime/OpenGLDrv/Private/OpenGLCommands.cpp
@@ -1947,37 +1947,41 @@ void FOpenGLDynamicRHI::SetupVertexArraysVAB(FOpenGLContextState& ContextState,
 	uint32 Divisor[NUM_OPENGL_VERTEX_STREAMS] = { 0 };
 	uint32 LastMaxAttrib = ContextState.MaxActiveAttrib;
 	bool UpdateDivisors = false;
+	uint32 StreamMask = ContextState.ActiveStreamMask;
 
 	check(IsValidRef(PendingState.BoundShaderState));
 	check(IsValidRef(PendingState.BoundShaderState->VertexShader));
 	FOpenGLVertexDeclaration* VertexDeclaration = PendingState.BoundShaderState->VertexDeclaration;
 	uint32 AttributeMask = PendingState.BoundShaderState->VertexShader->Bindings.InOutMask;
-	if ( ContextState.VertexDecl != VertexDeclaration )
+
+	if (ContextState.VertexDecl != VertexDeclaration || AttributeMask != ContextState.ActiveAttribMask)
 	{
-		ContextState.MaxActiveStream = 0;
 		ContextState.MaxActiveAttrib = 0;
+		StreamMask = 0;
 		UpdateDivisors = true;
 
+		check(VertexDeclaration->VertexElements.Num() <= 32);
+
 		for (int32 ElementIndex = 0; ElementIndex < VertexDeclaration->VertexElements.Num(); ElementIndex++)
 		{
 			FOpenGLVertexElement& VertexElement = VertexDeclaration->VertexElements[ElementIndex];
 			const uint32 AttributeIndex = VertexElement.AttributeIndex;
 			const uint32 StreamIndex = VertexElement.StreamIndex;
 
-			ContextState.MaxActiveStream = FMath::Max( ContextState.MaxActiveStream, StreamIndex);
 			ContextState.MaxActiveAttrib = FMath::Max( ContextState.MaxActiveAttrib, AttributeIndex);
 
 			if ( VertexElement.StreamIndex < NumStreams)
 			{
-			
 				FOpenGLCachedAttr &Attr = ContextState.VertexAttrs[AttributeIndex];
 
+				// Track the actively used streams, to limit the updates to those in use
+				StreamMask |= 0x1 << VertexElement.StreamIndex;
+
 				// Verify that the Divisor is consistent across the stream
 				check( !KnowsDivisor[StreamIndex] || Divisor[StreamIndex] == VertexElement.Divisor);
 				KnowsDivisor[StreamIndex] = true;
 				Divisor[StreamIndex] = VertexElement.Divisor;
 
-#if 1
 				if (
 					(Attr.StreamOffset != VertexElement.Offset) ||
 					(Attr.Size != VertexElement.Size) ||
@@ -2004,18 +2008,6 @@ void FOpenGLDynamicRHI::SetupVertexArraysVAB(FOpenGLContextState& ContextState,
 					FOpenGL::VertexAttribBinding( AttributeIndex, VertexElement.StreamIndex);
 					Attr.StreamIndex = StreamIndex;
 				}
-#else
-				if (!VertexElement.bShouldConvertToFloat)
-				{
-					FOpenGL::VertexAttribIFormat( AttributeIndex, VertexElement.Size, VertexElement.Type, VertexElement.Offset);
-				}
-				else
-				{
-					FOpenGL::VertexAttribFormat( AttributeIndex, VertexElement.Size, VertexElement.Type, VertexElement.bNormalized, VertexElement.Offset);
-				}
-
-				FOpenGL::VertexAttribBinding( AttributeIndex, VertexElement.StreamIndex);
-#endif
 			}
 			else
 			{
@@ -2023,21 +2015,28 @@ void FOpenGLDynamicRHI::SetupVertexArraysVAB(FOpenGLContextState& ContextState,
 				float data[4] = { 0.0f};
 
 				glVertexAttrib4fv( AttributeIndex, data );
+
+				//Kill this attribute to make sure it isn't enabled
+				AttributeMask &= ~(1 << AttributeIndex);
 			}
 		}
 		ContextState.VertexDecl = VertexDeclaration;
+
+		//Update the stream mask
+		ContextState.ActiveStreamMask = StreamMask;
 	}
 
 	//setup streams
-	for (uint32 StreamIndex = 0; StreamIndex < NumStreams && StreamIndex <= ContextState.MaxActiveStream; StreamIndex++)
+	for (uint32 StreamIndex = 0; StreamIndex < NumStreams && StreamMask != 0; StreamIndex++, StreamMask >>= 1)
 	{
 		FOpenGLStream &CachedStream = ContextState.VertexStreams[StreamIndex];
 		FOpenGLStream &Stream = Streams[StreamIndex];
-		if (Stream.VertexBuffer)
+		if ((StreamMask & 0x1) != 0 && Stream.VertexBuffer)
 		{
 			uint32 Offset = BaseVertexIndex * Stream.Stride + Stream.Offset;
 			if ( CachedStream.VertexBuffer != Stream.VertexBuffer || CachedStream.Offset != Offset || CachedStream.Stride != Stream.Stride)
 			{
+				check(Stream.VertexBuffer->Resource != 0);
 				FOpenGL::BindVertexBuffer( StreamIndex, Stream.VertexBuffer->Resource, Offset, Stream.Stride);
 				CachedStream.VertexBuffer = Stream.VertexBuffer;
 				CachedStream.Offset = Offset;
@@ -2051,6 +2050,9 @@ void FOpenGLDynamicRHI::SetupVertexArraysVAB(FOpenGLContextState& ContextState,
 		}
 	}
 
+	//Ensure that all requested streams were set
+	check(StreamMask == 0);
+
 	// Set the enable/disable state on the arrays
 	uint32 MaskDif = ContextState.ActiveAttribMask ^ AttributeMask;
 	if (MaskDif)
diff --git a/Engine/Source/Runtime/OpenGLDrv/Public/OpenGLState.h b/Engine/Source/Runtime/OpenGLDrv/Public/OpenGLState.h
--- a/Engine/Source/Runtime/OpenGLDrv/Public/OpenGLState.h
+++ b/Engine/Source/Runtime/OpenGLDrv/Public/OpenGLState.h
@@ -227,7 +227,7 @@ struct FOpenGLCachedAttr
 
 	bool bEnabled;
 
-	FOpenGLCachedAttr() : Pointer(FOpenGLCachedAttr_Invalid), bEnabled(false) {}
+	FOpenGLCachedAttr() : Pointer(FOpenGLCachedAttr_Invalid), Stride(-1), Divisor(0xFFFFFFFF), Type(0), StreamIndex(0xFFFFFFFF), bEnabled(false) {}
 };
 
 struct FOpenGLStream
@@ -322,7 +322,7 @@ struct FOpenGLContextState : public FOpenGLCommonState
 
 	FOpenGLVertexDeclaration* VertexDecl;
 	uint32 ActiveAttribMask;
-	uint32 MaxActiveStream;
+	uint32 ActiveStreamMask;
 	uint32 MaxActiveAttrib;
 
 	FOpenGLContextState()
@@ -350,7 +350,7 @@ struct FOpenGLContextState : public FOpenGLCommonState
 #endif
 	, VertexDecl(0)
 	, ActiveAttribMask(0)
-	, MaxActiveStream(0)
+	, ActiveStreamMask(0)
 	, MaxActiveAttrib(0)
 	{
 		Scissor.Min.X = Scissor.Min.Y = Scissor.Max.X = Scissor.Max.Y = 0;
