Integrated NVIDIA's 2386212 changelist: Additional clean-up on VAB state tracking.

diff --git a/Engine/Source/Runtime/OpenGLDrv/Private/OpenGLCommands.cpp b/Engine/Source/Runtime/OpenGLDrv/Private/OpenGLCommands.cpp
--- a/Engine/Source/Runtime/OpenGLDrv/Private/OpenGLCommands.cpp
+++ b/Engine/Source/Runtime/OpenGLDrv/Private/OpenGLCommands.cpp
@@ -1970,54 +1970,58 @@ void FOpenGLDynamicRHI::SetupVertexArraysVAB(FOpenGLContextState& ContextState,
 
 			ContextState.MaxActiveAttrib = FMath::Max( ContextState.MaxActiveAttrib, AttributeIndex);
 
-			if ( VertexElement.StreamIndex < NumStreams)
+			//only setup/track attributes actually in use
+			if (AttributeMask & (0x1 << AttributeIndex))
 			{
-				FOpenGLCachedAttr &Attr = ContextState.VertexAttrs[AttributeIndex];
+				if (VertexElement.StreamIndex < NumStreams)
+				{
+					FOpenGLCachedAttr &Attr = ContextState.VertexAttrs[AttributeIndex];
 
-				// Track the actively used streams, to limit the updates to those in use
-				StreamMask |= 0x1 << VertexElement.StreamIndex;
+					// Track the actively used streams, to limit the updates to those in use
+					StreamMask |= 0x1 << VertexElement.StreamIndex;
 
-				// Verify that the Divisor is consistent across the stream
-				check( !KnowsDivisor[StreamIndex] || Divisor[StreamIndex] == VertexElement.Divisor);
-				KnowsDivisor[StreamIndex] = true;
-				Divisor[StreamIndex] = VertexElement.Divisor;
+					// Verify that the Divisor is consistent across the stream
+					check(!KnowsDivisor[StreamIndex] || Divisor[StreamIndex] == VertexElement.Divisor);
+					KnowsDivisor[StreamIndex] = true;
+					Divisor[StreamIndex] = VertexElement.Divisor;
 
-				if (
-					(Attr.StreamOffset != VertexElement.Offset) ||
-					(Attr.Size != VertexElement.Size) ||
-					(Attr.Type != VertexElement.Type) ||
-					(Attr.bNormalized != VertexElement.bNormalized))
-				{
-					if (!VertexElement.bShouldConvertToFloat)
+					if (
+						(Attr.StreamOffset != VertexElement.Offset) ||
+						(Attr.Size != VertexElement.Size) ||
+						(Attr.Type != VertexElement.Type) ||
+						(Attr.bNormalized != VertexElement.bNormalized))
 					{
-						FOpenGL::VertexAttribIFormat( AttributeIndex, VertexElement.Size, VertexElement.Type, VertexElement.Offset);
+						if (!VertexElement.bShouldConvertToFloat)
+						{
+							FOpenGL::VertexAttribIFormat(AttributeIndex, VertexElement.Size, VertexElement.Type, VertexElement.Offset);
+						}
+						else
+						{
+							FOpenGL::VertexAttribFormat(AttributeIndex, VertexElement.Size, VertexElement.Type, VertexElement.bNormalized, VertexElement.Offset);
+						}
+
+						Attr.StreamOffset = VertexElement.Offset;
+						Attr.Size = VertexElement.Size;
+						Attr.Type = VertexElement.Type;
+						Attr.bNormalized = VertexElement.bNormalized;
 					}
-					else
+
+					if (Attr.StreamIndex != StreamIndex)
 					{
-						FOpenGL::VertexAttribFormat( AttributeIndex, VertexElement.Size, VertexElement.Type, VertexElement.bNormalized, VertexElement.Offset);
+						FOpenGL::VertexAttribBinding(AttributeIndex, VertexElement.StreamIndex);
+						Attr.StreamIndex = StreamIndex;
 					}
-
-					Attr.StreamOffset = VertexElement.Offset;
-					Attr.Size = VertexElement.Size;
-					Attr.Type = VertexElement.Type;
-					Attr.bNormalized = VertexElement.bNormalized;
 				}
-
-				if (Attr.StreamIndex != StreamIndex)
+				else
 				{
-					FOpenGL::VertexAttribBinding( AttributeIndex, VertexElement.StreamIndex);
-					Attr.StreamIndex = StreamIndex;
-				}
-			}
-			else
-			{
-				// bogus stream, make sure current value is zero to match D3D
-				float data[4] = { 0.0f};
+					// bogus stream, make sure current value is zero to match D3D
+					float data[4] = { 0.0f };
 
-				glVertexAttrib4fv( AttributeIndex, data );
+					glVertexAttrib4fv(AttributeIndex, data);
 
-				//Kill this attribute to make sure it isn't enabled
-				AttributeMask &= ~(1 << AttributeIndex);
+					//Kill this attribute to make sure it isn't enabled
+					AttributeMask &= ~(1 << AttributeIndex);
+				}
 			}
 		}
 		ContextState.VertexDecl = VertexDeclaration;
@@ -2048,6 +2052,13 @@ void FOpenGLDynamicRHI::SetupVertexArraysVAB(FOpenGLContextState& ContextState,
 				CachedStream.Divisor = Divisor[StreamIndex];
 			}
 		}
+		else
+		{
+			if (((StreamMask & 0x1) != 0) && (Stream.VertexBuffer == nullptr))
+			{
+				UE_LOG(LogRHI, Error, TEXT("Stream %d marked as in use, but vertex buffer provided is NULL (Mask = %x)"), StreamIndex, StreamMask);
+			}
+		}
 	}
 
 	//Ensure that all requested streams were set
