Linux: make errors while making GL context current non-fatal.

- Slate can destroy native windows prior to releasing RHI resources associated with them.
- This code sweeps the problem under the rug, possibly resulting in leaks, but fixes UE-7140 (crashes on switching windows).
- Proper resolution is tracked as UE-7388.

diff --git a/Engine/Source/Runtime/OpenGLDrv/Private/Linux/OpenGLLinux.cpp b/Engine/Source/Runtime/OpenGLDrv/Private/Linux/OpenGLLinux.cpp
--- a/Engine/Source/Runtime/OpenGLDrv/Private/Linux/OpenGLLinux.cpp
+++ b/Engine/Source/Runtime/OpenGLDrv/Private/Linux/OpenGLLinux.cpp
@@ -25,7 +25,13 @@ typedef SDL_GLContext	SDL_HGLContext;
 static void _ContextMakeCurrent( SDL_HWindow hWnd, SDL_HGLContext hGLDC )
 {
 	GLint Result = SDL_GL_MakeCurrent( hWnd, hGLDC );
-	check( !Result );
+	if (Result != 0)
+	{
+		// this is a warning and not error, since Slate sometimes destroys windows before
+		// releasing RHI resources associated with them. This code can result in leaks - proper resolution is tracked as UE-7388
+		FString SdlError(ANSI_TO_TCHAR(SDL_GetError()));
+		UE_LOG(LogLinux, Warning, TEXT("SDL_GL_MakeCurrent() failed, SDL error: '%s'"), *SdlError );
+	}
 }
 
 static SDL_HGLContext _GetCurrentContext()
