Detect GL_OES_standard_derivatives extension for ES2 and work around missing dFdx, dFdy, fwidth functions by patching shader for Android
#android
#ue4

diff --git a/Engine/Source/Runtime/OpenGLDrv/Private/OpenGLES2.cpp b/Engine/Source/Runtime/OpenGLDrv/Private/OpenGLES2.cpp
--- a/Engine/Source/Runtime/OpenGLDrv/Private/OpenGLES2.cpp
+++ b/Engine/Source/Runtime/OpenGLDrv/Private/OpenGLES2.cpp
@@ -118,6 +118,8 @@ bool FOpenGLES2::bRequiresDontEmitPrecisionForTextureSamplers = false;
 /* Some android platforms require textureCubeLod to be used some require textureCubeLodEXT however they either inconsistently or don't use the GL_TextureCubeLodEXT extension definition */
 bool FOpenGLES2::bRequiresTextureCubeLodEXTToTextureCubeLodDefine = false;
 
+/* Some android platforms do not support the GL_OES_standard_derivatives extension */
+bool FOpenGLES2::bSupportsStandardDerivativesExtension = false;
 
 bool FOpenGLES2::SupportsDisjointTimeQueries()
 {
@@ -183,6 +185,11 @@ void FOpenGLES2::ProcessExtensions( const FString& ExtensionsString )
 	bSupportsTextureStorageEXT = ExtensionsString.Contains(TEXT("GL_EXT_texture_storage"));
 	bSupportsCopyTextureLevels = bSupportsTextureStorageEXT && ExtensionsString.Contains(TEXT("GL_APPLE_copy_texture_levels"));
 	bSupportsTextureNPOT = ExtensionsString.Contains(TEXT("GL_OES_texture_npot")) || ExtensionsString.Contains(TEXT("GL_ARB_texture_non_power_of_two"));
+	bSupportsStandardDerivativesExtension = ExtensionsString.Contains(TEXT("GL_OES_standard_derivatives"));
+	if (!bSupportsStandardDerivativesExtension)
+	{
+		UE_LOG(LogRHI, Warning, TEXT("GL_OES_standard_derivatives not supported. There may be rendering errors if materials depend on dFdx, dFdy, or fwidth."));
+	}
 
 	// Report shader precision
 	int Range[2];
diff --git a/Engine/Source/Runtime/OpenGLDrv/Private/OpenGLShaders.cpp b/Engine/Source/Runtime/OpenGLDrv/Private/OpenGLShaders.cpp
--- a/Engine/Source/Runtime/OpenGLDrv/Private/OpenGLShaders.cpp
+++ b/Engine/Source/Runtime/OpenGLDrv/Private/OpenGLShaders.cpp
@@ -386,6 +386,23 @@ ShaderType* CompileOpenGLShader(const TArray<uint8>& Code)
 #endif
 
 #if PLATFORM_ANDROID 
+		// Temporary patch to remove #extension GL_OES_standard_derivaties if not supported
+		if (!FOpenGL::SupportsStandardDerivativesExtension())
+		{
+			const ANSICHAR * FoundPointer = FCStringAnsi::Strstr(GlslCodeOriginal.GetData(), "#extension GL_OES_standard_derivatives");
+			if (FoundPointer != nullptr)
+			{
+				// Replace the extension enable with dFdx, dFdy, and fwidth definitions so shader will compile.
+				// Currently SimpleElementPixelShader.usf is the most likely place this will come from for mobile
+				// as it is used for distance field text rendering (GammaDistanceFieldMain) so use a constant
+				// for the texture step rate of 1/512.  This will not work for other use cases.
+				ReplaceCString(GlslCodeOriginal, "#extension GL_OES_standard_derivatives : enable",
+					"#define dFdx(a) (0.001953125)\n"
+					"#define dFdy(a) (0.001953125)\n"
+					"#define fwidth(a) (0.00390625)\n");
+			}
+		}
+
 		if (IsES2Platform(GMaxRHIShaderPlatform))
 		{
 			// This #define fixes compiler errors on Android (which doesn't seem to support textureCubeLodEXT)
diff --git a/Engine/Source/Runtime/OpenGLDrv/Public/OpenGLES2.h b/Engine/Source/Runtime/OpenGLDrv/Public/OpenGLES2.h
--- a/Engine/Source/Runtime/OpenGLDrv/Public/OpenGLES2.h
+++ b/Engine/Source/Runtime/OpenGLDrv/Public/OpenGLES2.h
@@ -113,6 +113,7 @@ struct FOpenGLES2 : public FOpenGLBase
 
 	static FORCEINLINE bool RequiresDontEmitPrecisionForTextureSamplers() { return bRequiresDontEmitPrecisionForTextureSamplers; }
 	static FORCEINLINE bool RequiresTextureCubeLodEXTToTextureCubeLodDefine() { return bRequiresTextureCubeLodEXTToTextureCubeLodDefine; }
+	static FORCEINLINE bool SupportsStandardDerivativesExtension()		{ return bSupportsStandardDerivativesExtension; }
 
 	static FORCEINLINE int32 GetReadHalfFloatPixelsEnum()				{ return GL_HALF_FLOAT_OES; }
 
@@ -453,6 +454,9 @@ protected:
 	/** GL_EXT_texture_storage */
 	static bool bSupportsTextureStorageEXT;
 
+	/** GL_OES_standard_derivations */
+	static bool bSupportsStandardDerivativesExtension;
+
 public:
 	/* This is a hack to remove the calls to "precision sampler" defaults which are produced by the cross compiler however don't compile on some android platforms */
 	static bool bRequiresDontEmitPrecisionForTextureSamplers;
diff --git a/Engine/Source/Runtime/OpenGLDrv/Public/OpenGLES31.h b/Engine/Source/Runtime/OpenGLDrv/Public/OpenGLES31.h
--- a/Engine/Source/Runtime/OpenGLDrv/Public/OpenGLES31.h
+++ b/Engine/Source/Runtime/OpenGLDrv/Public/OpenGLES31.h
@@ -83,6 +83,7 @@ struct FOpenGLES31 : public FOpenGLBase
 
 	static FORCEINLINE bool RequiresDontEmitPrecisionForTextureSamplers() { return bRequiresDontEmitPrecisionForTextureSamplers; }
 	static FORCEINLINE bool RequiresTextureCubeLodEXTToTextureCubeLodDefine() { return bRequiresTextureCubeLodEXTToTextureCubeLodDefine; }
+	static FORCEINLINE bool SupportsStandardDerivativesExtension()		{ return true; }
 	static FORCEINLINE GLenum GetVertexHalfFloatFormat()				{ return bES2Fallback ? GL_HALF_FLOAT_OES : GL_HALF_FLOAT; }
 
 
