Cleaned up Android shader overrides and don't cache failed shader compile so overrides are actually tested
#android
#ue4

diff --git a/Engine/Source/Runtime/OpenGLDrv/Private/OpenGLDevice.cpp b/Engine/Source/Runtime/OpenGLDrv/Private/OpenGLDevice.cpp
--- a/Engine/Source/Runtime/OpenGLDrv/Private/OpenGLDevice.cpp
+++ b/Engine/Source/Runtime/OpenGLDrv/Private/OpenGLDevice.cpp
@@ -1020,7 +1020,20 @@ static void CheckTextureCubeLodSupport()
 			// we are done
 			return;
 		}
-		
+
+		FOpenGL::bRequiresDontEmitPrecisionForTextureSamplers = true;
+		FOpenGL::bRequiresTextureCubeLodEXTToTextureCubeLodDefine = true;
+
+		// try both hacks
+		PixelShader = (FOpenGLPixelShader*)(RHICreatePixelShader(Code).GetReference());
+
+		if (VerifyCompiledShader(PixelShader->Resource, TestFragmentProgram, false))
+		{
+			UE_LOG(LogRHI, Warning, TEXT("Enabling shader compiler hack to redefine textureCubeLodEXT to textureCubeLod and remove precision modifiers"));
+			// we are done
+			return;
+		}
+
 		UE_LOG(LogRHI, Warning, TEXT("Unable to find a test shader that compiles try running anyway"));
 	}
 #endif
diff --git a/Engine/Source/Runtime/OpenGLDrv/Private/OpenGLShaders.cpp b/Engine/Source/Runtime/OpenGLDrv/Private/OpenGLShaders.cpp
--- a/Engine/Source/Runtime/OpenGLDrv/Private/OpenGLShaders.cpp
+++ b/Engine/Source/Runtime/OpenGLDrv/Private/OpenGLShaders.cpp
@@ -446,28 +446,31 @@ ShaderType* CompileOpenGLShader(const TArray<uint8>& Code)
 					ReplaceCString(GlslCodeOriginal, "varying", "in");
 				}
 			}
-			else if ( (TypeEnum == GL_FRAGMENT_SHADER) &&
-				FOpenGL::RequiresDontEmitPrecisionForTextureSamplers() )
-			{
-				// Daniel: This device has some shader compiler compatibility issues force them to be disabled
-				//			The cross compiler will put the DONTEMITEXTENSIONSHADERTEXTURELODENABLE define around incompatible sections of code
-				AppendCString(GlslCode,
-					"#define DONTEMITEXTENSIONSHADERTEXTURELODENABLE \n"
-					"#define DONTEMITSAMPLERDEFAULTPRECISION \n"
-					"#define texture2DLodEXT(a, b, c) texture2D(a, b) \n"
-					"#define textureCubeLodEXT(a, b, c) textureCube(a, b) \n");
-			}
-			else if ( ( TypeEnum == GL_FRAGMENT_SHADER) && 
-				FOpenGL::RequiresTextureCubeLodEXTToTextureCubeLodDefine() )
-			{
-				AppendCString(GlslCode,
-					"#define textureCubeLodEXT textureCubeLod \n");
-			}
-			else if(!FOpenGL::SupportsShaderTextureLod() || !FOpenGL::SupportsShaderTextureCubeLod())
+			else 
 			{
-				AppendCString(GlslCode,
-					"#define texture2DLodEXT(a, b, c) texture2D(a, b) \n"
-					"#define textureCubeLodEXT(a, b, c) textureCube(a, b) \n");
+				if ((TypeEnum == GL_FRAGMENT_SHADER))
+				{
+					// Apply #defines to deal with incompatible sections of code
+
+					if (FOpenGL::RequiresDontEmitPrecisionForTextureSamplers())
+					{
+						AppendCString(GlslCode,
+							"#define DONTEMITSAMPLERDEFAULTPRECISION \n");
+					}
+
+					if (!FOpenGL::SupportsShaderTextureLod() || !FOpenGL::SupportsShaderTextureCubeLod())
+					{
+						AppendCString(GlslCode,
+							"#define DONTEMITEXTENSIONSHADERTEXTURELODENABLE \n"
+							"#define texture2DLodEXT(a, b, c) texture2D(a, b) \n"
+							"#define textureCubeLodEXT(a, b, c) textureCube(a, b) \n");
+					}
+					else if (FOpenGL::RequiresTextureCubeLodEXTToTextureCubeLodDefine())
+					{
+						AppendCString(GlslCode,
+							"#define textureCubeLodEXT textureCubeLod \n");
+					}
+				}
 			}
 		}
 
@@ -495,7 +498,18 @@ ShaderType* CompileOpenGLShader(const TArray<uint8>& Code)
 		glShaderSource(Resource, 1, (const GLchar**)&GlslCodeString, &GlslCodeLength);
 		glCompileShader(Resource);
 
-		GetOpenGLCompiledShaderCache().Add(Key,Resource);
+#if PLATFORM_ANDROID 
+		// On Android the same shader is compiled with different hacks to find the right one(s) to apply so don't cache unless successful
+		GLint CompileStatus;
+		glGetShaderiv(Resource, GL_COMPILE_STATUS, &CompileStatus);
+		if (CompileStatus == GL_TRUE)
+		{
+			GetOpenGLCompiledShaderCache().Add(Key, Resource);
+		}
+#else
+		// Cache it; compile status will be checked later on link (always caching will prevent multiple attempts to compile a failed shader)
+		GetOpenGLCompiledShaderCache().Add(Key, Resource);
+#endif
 	}
 
 	Shader = new ShaderType();
