Added sorting of display modes to FOpenGLDynamicRHI::RHIGetAvailableResolutions()

diff --git a/Engine/Source/Runtime/OpenGLDrv/Private/OpenGLViewport.cpp b/Engine/Source/Runtime/OpenGLDrv/Private/OpenGLViewport.cpp
--- a/Engine/Source/Runtime/OpenGLDrv/Private/OpenGLViewport.cpp
+++ b/Engine/Source/Runtime/OpenGLDrv/Private/OpenGLViewport.cpp
@@ -30,7 +30,26 @@ void FOpenGLDynamicRHI::RHIGetSupportedResolution(uint32 &Width, uint32 &Height)
 
 bool FOpenGLDynamicRHI::RHIGetAvailableResolutions(FScreenResolutionArray& Resolutions, bool bIgnoreRefreshRate)
 {
-	return PlatformGetAvailableResolutions(Resolutions, bIgnoreRefreshRate);
+	const bool Result = PlatformGetAvailableResolutions(Resolutions, bIgnoreRefreshRate);
+	if (Result)
+	{
+		Resolutions.Sort([](const FScreenResolutionRHI& L, const FScreenResolutionRHI& R)
+							{
+								if (L.Width != R.Width)
+								{
+									return L.Width < R.Width;
+								}
+								else if (L.Height != R.Height)
+								{
+									return L.Height < R.Height;
+								}
+								else
+								{
+									return L.RefreshRate < R.RefreshRate;
+								}
+							});
+	}
+	return Result;
 }
 
 /*=============================================================================
