Workaround ElementalDemo flickering on OS X (UE-7915).
- We can invalidate the whole buffer on MapBufferRange WriteOnly, as that is precisely what Windows does.
- On OS X invalidating the range overwritten when using unsynchronised MapBufferRange is also required to avoid the flickering.
- Added TexBuffer(0) call to the OS X GL_TEXTURE_BUFFER binding workaround to ensure that the buffer is revalidated and not reduced to a no-op by the driver.

diff --git a/Engine/Source/Runtime/OpenGLDrv/Private/OpenGLCommands.cpp b/Engine/Source/Runtime/OpenGLDrv/Private/OpenGLCommands.cpp
--- a/Engine/Source/Runtime/OpenGLDrv/Private/OpenGLCommands.cpp
+++ b/Engine/Source/Runtime/OpenGLDrv/Private/OpenGLCommands.cpp
@@ -466,6 +466,7 @@ void FOpenGLDynamicRHI::InternalUpdateTextureBuffer( FOpenGLContextState& Contex
 				ContextState.ActiveTexture = TextureIndex;
 			}
 			
+			FOpenGL::TexBuffer(GL_TEXTURE_BUFFER, GLFormat.InternalFormat[0], 0);
 			FOpenGL::TexBuffer(GL_TEXTURE_BUFFER, GLFormat.InternalFormat[0], VB->Resource);
 			
 			SRV->ModificationVersion = VB->ModificationCount;
diff --git a/Engine/Source/Runtime/OpenGLDrv/Public/OpenGL3.h b/Engine/Source/Runtime/OpenGLDrv/Public/OpenGL3.h
--- a/Engine/Source/Runtime/OpenGLDrv/Public/OpenGL3.h
+++ b/Engine/Source/Runtime/OpenGLDrv/Public/OpenGL3.h
@@ -128,7 +128,7 @@ struct FOpenGL3 : public FOpenGLBase
 				Access = GL_MAP_READ_BIT;
 				break;
 			case RLM_WriteOnly:
-				Access = (GL_MAP_INVALIDATE_RANGE_BIT | GL_MAP_WRITE_BIT);
+				Access = (GL_MAP_INVALIDATE_BUFFER_BIT | GL_MAP_WRITE_BIT);
 #if !PLATFORM_MAC // On OS X using the UNSYNCHRONIZED_BIT here is unsafe & out-of-spec. which will lead to corrupt rendering when using MTGL.
 				// Temp workaround for synchrnoization when a UBO is discarded while being referenced
 				Access |= GL_MAP_UNSYNCHRONIZED_BIT;
@@ -136,6 +136,9 @@ struct FOpenGL3 : public FOpenGLBase
 				break;
 			case RLM_WriteOnlyUnsynchronized:
 				Access = (GL_MAP_WRITE_BIT | GL_MAP_UNSYNCHRONIZED_BIT);
+#if PLATFORM_MAC
+				Access |= GL_MAP_INVALIDATE_RANGE_BIT;
+#endif
 				break;
 			case RLM_WriteOnlyPersistent:
 				Access = (GL_MAP_WRITE_BIT | GL_MAP_PERSISTENT_BIT | GL_MAP_COHERENT_BIT);
