Fix [UE-8343] - need to set the unpack alignment for Apple GL client-storage as it relies on this value for the cases where an additional internal copy may still be required.

diff --git a/Engine/Source/Runtime/OpenGLDrv/Private/OpenGLTexture.cpp b/Engine/Source/Runtime/OpenGLDrv/Private/OpenGLTexture.cpp
--- a/Engine/Source/Runtime/OpenGLDrv/Private/OpenGLTexture.cpp
+++ b/Engine/Source/Runtime/OpenGLDrv/Private/OpenGLTexture.cpp
@@ -371,6 +371,7 @@ FRHITexture* FOpenGLDynamicRHI::CreateOpenGLTexture(uint32 SizeX,uint32 SizeY,bo
 				}
 				
 				glPixelStorei(GL_UNPACK_CLIENT_STORAGE_APPLE, GL_TRUE);
+				glPixelStorei(GL_UNPACK_ALIGNMENT, 1);
 				
 				uint8* MipPointer = TextureRange;
 				for(uint32 MipIndex = 0; MipIndex < uint32(NumMips); MipIndex++)
@@ -393,6 +394,7 @@ FRHITexture* FOpenGLDynamicRHI::CreateOpenGLTexture(uint32 SizeX,uint32 SizeY,bo
 					}
 				}
 				
+				glPixelStorei(GL_UNPACK_ALIGNMENT, 4);
 				glPixelStorei(GL_UNPACK_CLIENT_STORAGE_APPLE, GL_FALSE);
 				
 				if(FOpenGL::SupportsTextureRange())
@@ -1062,6 +1064,7 @@ void TOpenGLTexture<RHIResourceType>::Unlock(uint32 MipIndex,uint32 ArrayIndex)
 		
 		if(bUseClientStorage)
 		{
+			glPixelStorei(GL_UNPACK_ALIGNMENT, 1);
 			glPixelStorei(GL_UNPACK_CLIENT_STORAGE_APPLE, GL_TRUE);
 			LockedSize = ClientStorageBuffers[BufferIndex].Size;
 			LockedBuffer = ClientStorageBuffers[BufferIndex].Data;
@@ -1134,6 +1137,7 @@ void TOpenGLTexture<RHIResourceType>::Unlock(uint32 MipIndex,uint32 ArrayIndex)
 		}
 		if(bUseClientStorage)
 		{
+			glPixelStorei(GL_UNPACK_ALIGNMENT, 4);
 			glPixelStorei(GL_UNPACK_CLIENT_STORAGE_APPLE, GL_FALSE);
 		}
 		else
