UE4 - Fix for OpenGL assert, fixes UE-8624 and UE-8631

diff --git a/Engine/Source/Runtime/OpenGLDrv/Private/OpenGLRenderTarget.cpp b/Engine/Source/Runtime/OpenGLDrv/Private/OpenGLRenderTarget.cpp
--- a/Engine/Source/Runtime/OpenGLDrv/Private/OpenGLRenderTarget.cpp
+++ b/Engine/Source/Runtime/OpenGLDrv/Private/OpenGLRenderTarget.cpp
@@ -302,49 +302,49 @@ void FOpenGLDynamicRHI::PurgeFramebufferFromCaches( GLuint Framebuffer )
 
 void FOpenGLDynamicRHI::RHICopyToResolveTarget(FTextureRHIParamRef SourceTextureRHI, FTextureRHIParamRef DestTextureRHI, bool bKeepOriginalSurface, const FResolveParams& ResolveParams)
 {
-	if(!SourceTextureRHI || !DestTextureRHI)
+	if (!SourceTextureRHI || !DestTextureRHI)
 	{
-		// no need to do anything (sliently ignored)
+		// no need to do anything (silently ignored)
 		return;
 	}
 
-	VERIFY_GL_SCOPE();
-
-	check(GMaxRHIFeatureLevel >= ERHIFeatureLevel::SM5 || ResolveParams.SourceArrayIndex == 0);
-	check(GMaxRHIFeatureLevel >= ERHIFeatureLevel::SM5 || ResolveParams.DestArrayIndex == 0);
-
 	FOpenGLTextureBase* SourceTexture = GetOpenGLTextureFromRHITexture(SourceTextureRHI);
 	FOpenGLTextureBase* DestTexture = GetOpenGLTextureFromRHITexture(DestTextureRHI);
 
-	const bool bSrcCubemap  = SourceTextureRHI->GetTextureCube() != NULL;
-	const bool bDestCubemap = DestTextureRHI->GetTextureCube() != NULL;
+	if (SourceTexture != DestTexture && FOpenGL::SupportsBlitFramebuffer())
+	{
+		VERIFY_GL_SCOPE();
 
-	uint32 DestIndex = ResolveParams.DestArrayIndex * (bDestCubemap ? 6 : 1) + (bDestCubemap ? uint32(ResolveParams.CubeFace) : 0);
-	uint32 SrcIndex  = ResolveParams.SourceArrayIndex * (bSrcCubemap ? 6 : 1) + (bSrcCubemap ? uint32(ResolveParams.CubeFace) : 0);
+		check(GMaxRHIFeatureLevel >= ERHIFeatureLevel::SM5 || ResolveParams.SourceArrayIndex == 0);
+		check(GMaxRHIFeatureLevel >= ERHIFeatureLevel::SM5 || ResolveParams.DestArrayIndex == 0);
 
-	uint32 BaseX = 0;
-	uint32 BaseY = 0;
-	uint32 SizeX = 0;
-	uint32 SizeY = 0;
-	if (ResolveParams.Rect.IsValid())
-	{
-		BaseX = ResolveParams.Rect.X1;
-		BaseY = ResolveParams.Rect.Y1;
-		SizeX = ResolveParams.Rect.X2 - ResolveParams.Rect.X1;
-		SizeY = ResolveParams.Rect.Y2 - ResolveParams.Rect.Y1;
-	}
-	else
-	{
-		// Invalid rect mans that the entire source is to be copied
-		SizeX = GetOpenGLTextureSizeXFromRHITexture(SourceTextureRHI);
-		SizeY = GetOpenGLTextureSizeYFromRHITexture(SourceTextureRHI);
+		const bool bSrcCubemap  = SourceTextureRHI->GetTextureCube() != NULL;
+		const bool bDestCubemap = DestTextureRHI->GetTextureCube() != NULL;
 
-		SizeX = FMath::Max<uint32>( 1, SizeX >> ResolveParams.MipIndex);
-		SizeY = FMath::Max<uint32>( 1, SizeY >> ResolveParams.MipIndex);
-	}
+		uint32 DestIndex = ResolveParams.DestArrayIndex * (bDestCubemap ? 6 : 1) + (bDestCubemap ? uint32(ResolveParams.CubeFace) : 0);
+		uint32 SrcIndex  = ResolveParams.SourceArrayIndex * (bSrcCubemap ? 6 : 1) + (bSrcCubemap ? uint32(ResolveParams.CubeFace) : 0);
+
+		uint32 BaseX = 0;
+		uint32 BaseY = 0;
+		uint32 SizeX = 0;
+		uint32 SizeY = 0;
+		if (ResolveParams.Rect.IsValid())
+		{
+			BaseX = ResolveParams.Rect.X1;
+			BaseY = ResolveParams.Rect.Y1;
+			SizeX = ResolveParams.Rect.X2 - ResolveParams.Rect.X1;
+			SizeY = ResolveParams.Rect.Y2 - ResolveParams.Rect.Y1;
+		}
+		else
+		{
+			// Invalid rect mans that the entire source is to be copied
+			SizeX = GetOpenGLTextureSizeXFromRHITexture(SourceTextureRHI);
+			SizeY = GetOpenGLTextureSizeYFromRHITexture(SourceTextureRHI);
+
+			SizeX = FMath::Max<uint32>(1, SizeX >> ResolveParams.MipIndex);
+			SizeY = FMath::Max<uint32>(1, SizeY >> ResolveParams.MipIndex);
+		}
 
-	if (SourceTexture != DestTexture && FOpenGL::SupportsBlitFramebuffer())
-	{
 		GPUProfilingData.RegisterGPUWork();
 		uint32 MipmapLevel = ResolveParams.MipIndex;
 
@@ -432,6 +432,7 @@ void FOpenGLDynamicRHI::RHICopyToResolveTarget(FTextureRHIParamRef SourceTexture
 	}
 	else
 	{
+		// no need to do anything (silently ignored)
 	}
 }
 
