Allocate GL queries on heap.

(Merging CL 2419951 from UT branch, fixes issues with r.RHICmdBypass 0).

diff --git a/Engine/Source/Runtime/OpenGLDrv/Private/OpenGLQuery.cpp b/Engine/Source/Runtime/OpenGLDrv/Private/OpenGLQuery.cpp
--- a/Engine/Source/Runtime/OpenGLDrv/Private/OpenGLQuery.cpp
+++ b/Engine/Source/Runtime/OpenGLDrv/Private/OpenGLQuery.cpp
@@ -372,8 +372,8 @@ void FOpenGLBufferedGPUTiming::InitDynamicRHI()
 		
 		for(int32 BufferIndex = 0; BufferIndex < BufferSize; ++BufferIndex)
 		{
-			StartTimestamps.Add(FOpenGLRenderQuery(RQT_AbsoluteTime));
-			EndTimestamps.Add(FOpenGLRenderQuery(RQT_AbsoluteTime));
+			StartTimestamps.Add(new FOpenGLRenderQuery(RQT_AbsoluteTime));
+			EndTimestamps.Add(new FOpenGLRenderQuery(RQT_AbsoluteTime));
 		}
 	}
 }
@@ -385,6 +385,16 @@ void FOpenGLBufferedGPUTiming::ReleaseDynamicRHI()
 {
 	VERIFY_GL_SCOPE();
 
+	for(FOpenGLRenderQuery* Query : StartTimestamps)
+	{
+		delete Query;
+	}
+	
+	for(FOpenGLRenderQuery* Query : EndTimestamps)
+	{
+		delete Query;
+	}
+	
 	StartTimestamps.Reset();
 	EndTimestamps.Reset();
 	
@@ -400,7 +410,7 @@ void FOpenGLBufferedGPUTiming::StartTiming()
 	if ( GIsSupported && !bIsTiming )
 	{
 		int32 NewTimestampIndex = (CurrentTimestamp + 1) % BufferSize;
-		FOpenGLRenderQuery* TimerQuery = &StartTimestamps[NewTimestampIndex];
+		FOpenGLRenderQuery* TimerQuery = StartTimestamps[NewTimestampIndex];
 		{
 			if (!TimerQuery->bInvalidResource && !PlatformContextIsCurrent(TimerQuery->ResourceContext))
 			{
@@ -415,7 +425,7 @@ void FOpenGLBufferedGPUTiming::StartTiming()
 			}
 		}
 
-		FOpenGL::QueryTimestampCounter(StartTimestamps[NewTimestampIndex].Resource);
+		FOpenGL::QueryTimestampCounter(StartTimestamps[NewTimestampIndex]->Resource);
 		CurrentTimestamp = NewTimestampIndex;
 		bIsTiming = true;
 	}
@@ -433,7 +443,7 @@ void FOpenGLBufferedGPUTiming::EndTiming()
 	{
 		checkSlow( CurrentTimestamp >= 0 && CurrentTimestamp < BufferSize );
 		
-		FOpenGLRenderQuery* TimerQuery = &EndTimestamps[CurrentTimestamp];
+		FOpenGLRenderQuery* TimerQuery = EndTimestamps[CurrentTimestamp];
 		{
 			if (!TimerQuery->bInvalidResource && !PlatformContextIsCurrent(TimerQuery->ResourceContext))
 			{
@@ -448,7 +458,7 @@ void FOpenGLBufferedGPUTiming::EndTiming()
 			}
 		} 
 
-		FOpenGL::QueryTimestampCounter(EndTimestamps[CurrentTimestamp].Resource);
+		FOpenGL::QueryTimestampCounter(EndTimestamps[CurrentTimestamp]->Resource);
 		NumIssuedTimestamps = FMath::Min<int32>(NumIssuedTimestamps + 1, BufferSize);
 		bIsTiming = false;
 	}
@@ -473,14 +483,14 @@ uint64 FOpenGLBufferedGPUTiming::GetTiming(bool bGetCurrentResultsAndBlock)
 		int32 TimestampIndex = CurrentTimestamp;
 
 		{
-			FOpenGLRenderQuery* EndStamp = &EndTimestamps[TimestampIndex];
+			FOpenGLRenderQuery* EndStamp = EndTimestamps[TimestampIndex];
 			if (!EndStamp->bInvalidResource && !PlatformContextIsCurrent(EndStamp->ResourceContext))
 			{
 				PlatformReleaseRenderQuery(EndStamp->Resource, EndStamp->ResourceContext);
 				EndStamp->bInvalidResource = true;
 			}
 
-			FOpenGLRenderQuery* StartStamp = &StartTimestamps[TimestampIndex];
+			FOpenGLRenderQuery* StartStamp = StartTimestamps[TimestampIndex];
 			if (!StartStamp->bInvalidResource && !PlatformContextIsCurrent(StartStamp->ResourceContext))
 			{
 				PlatformReleaseRenderQuery(StartStamp->Resource, StartStamp->ResourceContext);
@@ -501,17 +511,17 @@ uint64 FOpenGLBufferedGPUTiming::GetTiming(bool bGetCurrentResultsAndBlock)
 			for ( int32 IssueIndex = 1; IssueIndex < NumIssuedTimestamps; ++IssueIndex )
 			{
 				GLuint EndAvailable = GL_FALSE;
-				FOpenGL::GetQueryObject(EndTimestamps[TimestampIndex].Resource, FOpenGL::QM_ResultAvailable, &EndAvailable);
+				FOpenGL::GetQueryObject(EndTimestamps[TimestampIndex]->Resource, FOpenGL::QM_ResultAvailable, &EndAvailable);
 
 				if ( EndAvailable == GL_TRUE )
 				{
 					GLuint StartAvailable = GL_FALSE;
-					FOpenGL::GetQueryObject(StartTimestamps[TimestampIndex].Resource, FOpenGL::QM_ResultAvailable, &StartAvailable);
+					FOpenGL::GetQueryObject(StartTimestamps[TimestampIndex]->Resource, FOpenGL::QM_ResultAvailable, &StartAvailable);
 
 					if(StartAvailable == GL_TRUE)
 					{
-						FOpenGL::GetQueryObject(EndTimestamps[TimestampIndex].Resource, FOpenGL::QM_Result, &EndTime);
-						FOpenGL::GetQueryObject(StartTimestamps[TimestampIndex].Resource, FOpenGL::QM_Result, &StartTime);
+						FOpenGL::GetQueryObject(EndTimestamps[TimestampIndex]->Resource, FOpenGL::QM_Result, &EndTime);
+						FOpenGL::GetQueryObject(StartTimestamps[TimestampIndex]->Resource, FOpenGL::QM_Result, &StartTime);
 						if (EndTime > StartTime)
 						{
 							return EndTime - StartTime;
@@ -539,7 +549,7 @@ uint64 FOpenGLBufferedGPUTiming::GetTiming(bool bGetCurrentResultsAndBlock)
 			// If we are blocking, retry until the GPU processes the time stamp command
 			do 
 			{
-				FOpenGL::GetQueryObject(EndTimestamps[TimestampIndex].Resource, FOpenGL::QM_ResultAvailable, &EndAvailable);
+				FOpenGL::GetQueryObject(EndTimestamps[TimestampIndex]->Resource, FOpenGL::QM_ResultAvailable, &EndAvailable);
 
 				if ((FPlatformTime::Seconds() - StartTimeoutTime) > 0.5)
 				{
@@ -560,7 +570,7 @@ uint64 FOpenGLBufferedGPUTiming::GetTiming(bool bGetCurrentResultsAndBlock)
 
 				do 
 				{
-					FOpenGL::GetQueryObject(StartTimestamps[TimestampIndex].Resource, FOpenGL::QM_ResultAvailable, &StartAvailable);
+					FOpenGL::GetQueryObject(StartTimestamps[TimestampIndex]->Resource, FOpenGL::QM_ResultAvailable, &StartAvailable);
 
 					if ((FPlatformTime::Seconds() - StartTimeoutTime) > 0.5)
 					{
@@ -573,8 +583,8 @@ uint64 FOpenGLBufferedGPUTiming::GetTiming(bool bGetCurrentResultsAndBlock)
 
 				if(StartAvailable == GL_TRUE)
 				{
-					FOpenGL::GetQueryObject(EndTimestamps[TimestampIndex].Resource, FOpenGL::QM_Result, &EndTime);
-					FOpenGL::GetQueryObject(StartTimestamps[TimestampIndex].Resource, FOpenGL::QM_Result, &StartTime);
+					FOpenGL::GetQueryObject(EndTimestamps[TimestampIndex]->Resource, FOpenGL::QM_Result, &EndTime);
+					FOpenGL::GetQueryObject(StartTimestamps[TimestampIndex]->Resource, FOpenGL::QM_Result, &StartTime);
 					if (EndTime > StartTime)
 					{
 						return EndTime - StartTime;
diff --git a/Engine/Source/Runtime/OpenGLDrv/Public/OpenGLDrv.h b/Engine/Source/Runtime/OpenGLDrv/Public/OpenGLDrv.h
--- a/Engine/Source/Runtime/OpenGLDrv/Public/OpenGLDrv.h
+++ b/Engine/Source/Runtime/OpenGLDrv/Public/OpenGLDrv.h
@@ -107,9 +107,9 @@ private:
 	/** Number of measurements in the buffers (0 - BufferSize). */
 	int32								NumIssuedTimestamps;
 	/** Timestamps for all StartTimings. */
-	TArray<FOpenGLRenderQuery>			StartTimestamps;
+	TArray<FOpenGLRenderQuery *>		StartTimestamps;
 	/** Timestamps for all EndTimings. */
-	TArray<FOpenGLRenderQuery>			EndTimestamps;
+	TArray<FOpenGLRenderQuery *>		EndTimestamps;
 	/** Whether we are currently timing the GPU: between StartTiming() and EndTiming(). */
 	bool								bIsTiming;
 };
