Move Mac-specific workarounds to MapBufferRange to OpenGLMac.h to avoid exposing them to other platforms which don't need to worry about them.

diff --git a/Engine/Source/Runtime/OpenGLDrv/Private/Mac/OpenGLMac.h b/Engine/Source/Runtime/OpenGLDrv/Private/Mac/OpenGLMac.h
--- a/Engine/Source/Runtime/OpenGLDrv/Private/Mac/OpenGLMac.h
+++ b/Engine/Source/Runtime/OpenGLDrv/Private/Mac/OpenGLMac.h
@@ -61,6 +61,31 @@ public:
 		glFlushRenderAPPLE();
 	}
 	
+	// Workaround Mac-specific MapBuffer issues without exposing the changes to other platforms
+	static FORCEINLINE void* MapBufferRange(GLenum Type, uint32 InOffset, uint32 InSize, EResourceLockMode LockMode)
+	{
+		GLenum Access;
+		switch ( LockMode )
+		{
+			case RLM_ReadOnly:
+				Access = GL_MAP_READ_BIT;
+				break;
+			case RLM_WriteOnly:
+				Access = (GL_MAP_INVALIDATE_BUFFER_BIT | GL_MAP_WRITE_BIT);
+				break;
+			case RLM_WriteOnlyUnsynchronized: // Invalidate range is only required to avoid UE-7915 and is a workaround, not a fix.
+				Access = (GL_MAP_WRITE_BIT | GL_MAP_UNSYNCHRONIZED_BIT | GL_MAP_INVALIDATE_RANGE_BIT);
+				break;
+			case RLM_WriteOnlyPersistent:
+				Access = (GL_MAP_WRITE_BIT | GL_MAP_PERSISTENT_BIT | GL_MAP_COHERENT_BIT);
+				break;
+			case RLM_ReadWrite:
+			default:
+				Access = (GL_MAP_READ_BIT | GL_MAP_WRITE_BIT);
+		}
+		return glMapBufferRange(Type, InOffset, InSize, Access);
+	}
+	
 	static void DeleteTextures(GLsizei Number, const GLuint* Textures);
 	
 	static void BufferSubData(GLenum Target, GLintptr Offset, GLsizeiptr Size, const GLvoid* Data);
diff --git a/Engine/Source/Runtime/OpenGLDrv/Public/OpenGL3.h b/Engine/Source/Runtime/OpenGLDrv/Public/OpenGL3.h
--- a/Engine/Source/Runtime/OpenGLDrv/Public/OpenGL3.h
+++ b/Engine/Source/Runtime/OpenGLDrv/Public/OpenGL3.h
@@ -129,16 +129,13 @@ struct FOpenGL3 : public FOpenGLBase
 				break;
 			case RLM_WriteOnly:
 				Access = (GL_MAP_INVALIDATE_BUFFER_BIT | GL_MAP_WRITE_BIT);
-#if !PLATFORM_MAC // On OS X using the UNSYNCHRONIZED_BIT here is unsafe & out-of-spec. which will lead to corrupt rendering when using MTGL.
+#if 1
 				// Temp workaround for synchrnoization when a UBO is discarded while being referenced
 				Access |= GL_MAP_UNSYNCHRONIZED_BIT;
 #endif
 				break;
 			case RLM_WriteOnlyUnsynchronized:
 				Access = (GL_MAP_WRITE_BIT | GL_MAP_UNSYNCHRONIZED_BIT);
-#if PLATFORM_MAC
-				Access |= GL_MAP_INVALIDATE_RANGE_BIT;
-#endif
 				break;
 			case RLM_WriteOnlyPersistent:
 				Access = (GL_MAP_WRITE_BIT | GL_MAP_PERSISTENT_BIT | GL_MAP_COHERENT_BIT);
