Make the Intel HD 3000 report itself as strictly ES2-only on OS X until/unless there is a solution to the deferred renderer running at <1fps.

diff --git a/Engine/Source/Runtime/OpenGLDrv/Private/Mac/OpenGLMac.cpp b/Engine/Source/Runtime/OpenGLDrv/Private/Mac/OpenGLMac.cpp
--- a/Engine/Source/Runtime/OpenGLDrv/Private/Mac/OpenGLMac.cpp
+++ b/Engine/Source/Runtime/OpenGLDrv/Private/Mac/OpenGLMac.cpp
@@ -1343,6 +1343,20 @@ bool FMacOpenGL::MustFlushTexStorage(void)
 	return GMacFlushTexStorage || GMacMustFlushTexStorage;
 }
 
+/** Is the current renderer the Intel HD3000? */
+bool FMacOpenGL::IsIntelHD3000()
+{
+	// Get the current renderer ID
+	GLint RendererID = 0;
+	CGLContextObj Current = CGLGetCurrentContext();
+	if (Current)
+	{
+		CGLError Error = CGLGetParameter(Current, kCGLCPCurrentRendererID, &RendererID);
+		check(Error == kCGLNoError && RendererID != 0);
+	}
+	return (RendererID & kCGLRendererIDMatchingMask) == kCGLRendererIntelHDID;
+}
+
 void FMacOpenGL::DeleteTextures(GLsizei Number, const GLuint* Textures)
 {
 	TArray<GLuint>& TexturesToDelete = GMacTexturesToDelete[GFrameNumberRenderThread % GMacTexturePoolNum];
diff --git a/Engine/Source/Runtime/OpenGLDrv/Private/Mac/OpenGLMac.h b/Engine/Source/Runtime/OpenGLDrv/Private/Mac/OpenGLMac.h
--- a/Engine/Source/Runtime/OpenGLDrv/Private/Mac/OpenGLMac.h
+++ b/Engine/Source/Runtime/OpenGLDrv/Private/Mac/OpenGLMac.h
@@ -24,15 +24,31 @@ private:
 	/** Must we employ a workaround for radr://15553950, TTP# 315197 */
 	static bool MustFlushTexStorage(void);
 	
+	/** Is the current renderer the Intel HD3000? */
+	static bool IsIntelHD3000();
+	
 public:
 	static void ProcessQueryGLInt();
 	static void ProcessExtensions(const FString& ExtensionsString);
 	static uint64 GetVideoMemorySize();
 	
+	static FORCEINLINE ERHIFeatureLevel::Type GetFeatureLevel()
+	{
+		// The Intel HD 3000 on OS X can't cope with our deferred renderer, but ES2 mode works fine :(
+		static bool bForceFeatureLevelES2 = FParse::Param(FCommandLine::Get(), TEXT("FeatureLevelES2"));
+		if (bForceFeatureLevelES2 || IsIntelHD3000())
+		{
+			return ERHIFeatureLevel::ES2;
+		}
+		
+		return FOpenGL3::GetFeatureLevel();
+	}
+	
 	static FORCEINLINE EShaderPlatform GetShaderPlatform()
 	{
+		// The Intel HD 3000 on OS X can't cope with our deferred renderer, but ES2 mode works fine :(
 		static bool bForceFeatureLevelES2 = FParse::Param(FCommandLine::Get(), TEXT("FeatureLevelES2"));
-		if (bForceFeatureLevelES2)
+		if (bForceFeatureLevelES2 || IsIntelHD3000())
 		{
 			return SP_OPENGL_PCES2;
 		}
