Scale all resolutions in FMacOpenGL::PlatformGetAvailableResolutions by the current primary screen's backing scale factor to prevent displaying resolutions that would result in us reporting to Slate that the viewport is larger than the enclosing window.

diff --git a/Engine/Source/Runtime/OpenGLDrv/Private/Mac/OpenGLMac.cpp b/Engine/Source/Runtime/OpenGLDrv/Private/Mac/OpenGLMac.cpp
--- a/Engine/Source/Runtime/OpenGLDrv/Private/Mac/OpenGLMac.cpp
+++ b/Engine/Source/Runtime/OpenGLDrv/Private/Mac/OpenGLMac.cpp
@@ -786,6 +786,13 @@ void PlatformGetSupportedResolution(uint32 &Width, uint32 &Height)
 
 bool PlatformGetAvailableResolutions(FScreenResolutionArray& Resolutions, bool bIgnoreRefreshRate)
 {
+	SCOPED_AUTORELEASE_POOL;
+
+	NSArray* AllScreens = [NSScreen screens];
+	NSScreen* PrimaryScreen = (NSScreen*)[AllScreens objectAtIndex: 0];
+
+	SIZE_T Scale = (SIZE_T)[PrimaryScreen backingScaleFactor];
+
 	int32 MinAllowableResolutionX = 0;
 	int32 MinAllowableResolutionY = 0;
 	int32 MaxAllowableResolutionX = 10480;
@@ -813,8 +820,8 @@ bool PlatformGetAvailableResolutions(FScreenResolutionArray& Resolutions, bool b
 		for (int32 Index = 0; Index < NumModes; Index++)
 		{
 			CGDisplayModeRef Mode = (CGDisplayModeRef)CFArrayGetValueAtIndex(AllModes, Index);
-			SIZE_T Width = CGDisplayModeGetWidth(Mode);
-			SIZE_T Height = CGDisplayModeGetHeight(Mode);
+			SIZE_T Width = CGDisplayModeGetWidth(Mode) / Scale;
+			SIZE_T Height = CGDisplayModeGetHeight(Mode) / Scale;
 			int32 RefreshRate = (int32)CGDisplayModeGetRefreshRate(Mode);
 
 			if (((int32)Width >= MinAllowableResolutionX) &&
