Fix CIS (GL types as keys).

diff --git a/Engine/Source/Runtime/OpenGLDrv/Private/OpenGLShaders.cpp b/Engine/Source/Runtime/OpenGLDrv/Private/OpenGLShaders.cpp
--- a/Engine/Source/Runtime/OpenGLDrv/Private/OpenGLShaders.cpp
+++ b/Engine/Source/Runtime/OpenGLDrv/Private/OpenGLShaders.cpp
@@ -831,9 +831,9 @@ static TMap<GLuint, TMap<FString, GLint64>>& GetOpenGLUniformBlockLocations()
 	return UniformBlockLocations;
 }
 
-static TMap<GLuint, TMap<GLint64, GLint64>>& GetOpenGLUniformBlockBindings()
+static TMap<GLuint, TMap<int64, GLint64>>& GetOpenGLUniformBlockBindings()
 {
-	static TMap<GLuint, TMap<GLint64, GLint64>> UniformBlockBindings;
+	static TMap<GLuint, TMap<int64, GLint64>> UniformBlockBindings;
 	return UniformBlockBindings;
 }
 
@@ -855,8 +855,8 @@ static GLuint GetOpenGLProgramUniformBlockIndex(GLuint Program, const GLchar* Un
 
 static void GetOpenGLProgramUniformBlockBinding(GLuint Program, GLuint UniformBlockIndex, GLuint UniformBlockBinding)
 {
-	TMap<GLint64, GLint64>& Bindings = GetOpenGLUniformBlockBindings().FindOrAdd(Program);
-	GLint64* Bind = Bindings.Find(UniformBlockIndex);
+	TMap<int64, GLint64>& Bindings = GetOpenGLUniformBlockBindings().FindOrAdd(Program);
+	GLint64* Bind = static_cast<GLint64 *>(Bindings.Find(UniformBlockIndex));
 	if(!Bind)
 	{
 		Bind = &(Bindings.Emplace(UniformBlockIndex));
@@ -864,9 +864,9 @@ static void GetOpenGLProgramUniformBlockBinding(GLuint Program, GLuint UniformBl
 		*Bind = -1;
 	}
 	check(Bind);
-	if(*Bind != (GLint64)UniformBlockBinding)
+	if(*Bind != static_cast<int64>(UniformBlockBinding))
 	{
-		*Bind = (GLint64)UniformBlockBinding;
+		*Bind = static_cast<int64>(UniformBlockBinding);
 		FOpenGL::UniformBlockBinding(Program, UniformBlockIndex, UniformBlockBinding);
 	}
 }
