CIS fix for HTML5 GL.

diff --git a/Engine/Source/Runtime/OpenGLDrv/Private/HTML5/HTML5OpenGL.h b/Engine/Source/Runtime/OpenGLDrv/Private/HTML5/HTML5OpenGL.h
--- a/Engine/Source/Runtime/OpenGLDrv/Private/HTML5/HTML5OpenGL.h
+++ b/Engine/Source/Runtime/OpenGLDrv/Private/HTML5/HTML5OpenGL.h
@@ -10,6 +10,7 @@
 #include <SDL_opengles2.h>
 
 typedef char UGLsync; 
+typedef long long  GLint64;
 typedef unsigned long long  GLuint64;
 
 // empty functions. 
