- Changed all GLint64's to int64's (for HTML5 compat) [UE-9432]

diff --git a/Engine/Source/Runtime/OpenGLDrv/Private/OpenGLShaders.cpp b/Engine/Source/Runtime/OpenGLDrv/Private/OpenGLShaders.cpp
--- a/Engine/Source/Runtime/OpenGLDrv/Private/OpenGLShaders.cpp
+++ b/Engine/Source/Runtime/OpenGLDrv/Private/OpenGLShaders.cpp
@@ -825,38 +825,38 @@ void FOpenGLDynamicRHI::BindUniformBufferBase(FOpenGLContextState& ContextState,
 
 // ============================================================================================================================
 
-static TMap<GLuint, TMap<FString, GLint64>>& GetOpenGLUniformBlockLocations()
+static TMap<GLuint, TMap<FString, int64>>& GetOpenGLUniformBlockLocations()
 {
-	static TMap<GLuint, TMap<FString, GLint64>> UniformBlockLocations;
+	static TMap<GLuint, TMap<FString, int64>> UniformBlockLocations;
 	return UniformBlockLocations;
 }
 
-static TMap<GLuint, TMap<int64, GLint64>>& GetOpenGLUniformBlockBindings()
+static TMap<GLuint, TMap<int64, int64>>& GetOpenGLUniformBlockBindings()
 {
-	static TMap<GLuint, TMap<int64, GLint64>> UniformBlockBindings;
+	static TMap<GLuint, TMap<int64, int64>> UniformBlockBindings;
 	return UniformBlockBindings;
 }
 
 static GLuint GetOpenGLProgramUniformBlockIndex(GLuint Program, const GLchar* UniformBlockName)
 {
-	TMap<FString, GLint64>& Locations = GetOpenGLUniformBlockLocations().FindOrAdd(Program);
-	GLint64* Location = Locations.Find(UniformBlockName);
+	TMap<FString, int64>& Locations = GetOpenGLUniformBlockLocations().FindOrAdd(Program);
+	int64* Location = Locations.Find(UniformBlockName);
 	if(Location)
 	{
 		return *Location;
 	}
 	else
 	{
-		GLint64& Loc = Locations.Emplace(UniformBlockName);
-		Loc = (GLint64)FOpenGL::GetUniformBlockIndex(Program, UniformBlockName);
+		int64& Loc = Locations.Emplace(UniformBlockName);
+		Loc = (int64)FOpenGL::GetUniformBlockIndex(Program, UniformBlockName);
 		return Loc;
 	}
 }
 
 static void GetOpenGLProgramUniformBlockBinding(GLuint Program, GLuint UniformBlockIndex, GLuint UniformBlockBinding)
 {
-	TMap<int64, GLint64>& Bindings = GetOpenGLUniformBlockBindings().FindOrAdd(Program);
-	GLint64* Bind = static_cast<GLint64 *>(Bindings.Find(UniformBlockIndex));
+	TMap<int64, int64>& Bindings = GetOpenGLUniformBlockBindings().FindOrAdd(Program);
+	int64* Bind = static_cast<int64 *>(Bindings.Find(UniformBlockIndex));
 	if(!Bind)
 	{
 		Bind = &(Bindings.Emplace(UniformBlockIndex));
