OpenGL separate shader objects require that the shader hash be set even when the shader cache is not in use. The calculated CRC will suffice.

diff --git a/Engine/Source/Runtime/OpenGLDrv/Private/OpenGLShaders.cpp b/Engine/Source/Runtime/OpenGLDrv/Private/OpenGLShaders.cpp
--- a/Engine/Source/Runtime/OpenGLDrv/Private/OpenGLShaders.cpp
+++ b/Engine/Source/Runtime/OpenGLDrv/Private/OpenGLShaders.cpp
@@ -704,6 +704,15 @@ ShaderType* CompileOpenGLShader(const TArray<uint8>& Code)
 	Shader->Resource = Resource;
 	Shader->Bindings = Header.Bindings;
 	Shader->UniformBuffersCopyInfo = Header.UniformBuffersCopyInfo;
+	
+	// If there is no shader cache then we must assign the hash here
+	if (FOpenGL::SupportsSeparateShaderObjects() && !FShaderCache::GetShaderCache())
+	{
+		// Just use the CRC - if it isn't being cached & logged we'll be dependent on the CRC alone anyway
+		FSHAHash Hash;
+		FMemory::Memcpy(Hash.Hash, &GlslCodeOriginalCRC, sizeof(uint32));
+		Shader->SetHash(Hash);
+	}
 
 #if DEBUG_GL_SHADERS
 	Shader->GlslCode = GlslCode;
@@ -1059,7 +1068,7 @@ void FOpenGLLinkedProgram::VerifyUniformBlockBindings( int Stage, uint32 FirstUn
 		NameBuffer[3] = 0;
 		NameBuffer[4] = 0;
 		
-		GLuint StageProgram = FOpenGL::SupportsSeparateShaderObjects() ? Config.Shaders[Stage].Resource : Program;
+		GLuint StageProgram = Config.Shaders[Stage].Resource;
 
 		for (int32 BufferIndex = 0; BufferIndex < Config.Shaders[Stage].Bindings.NumUniformBuffers; ++BufferIndex)
 		{
