Added workarounds for Adreno 2xx GPU on Android (adds checking of maximum varying vectors allowed)
#android
#ue4

diff --git a/Engine/Source/Runtime/Core/Private/Android/AndroidWindow.cpp b/Engine/Source/Runtime/Core/Private/Android/AndroidWindow.cpp
--- a/Engine/Source/Runtime/Core/Private/Android/AndroidWindow.cpp
+++ b/Engine/Source/Runtime/Core/Private/Android/AndroidWindow.cpp
@@ -100,6 +100,8 @@ FPlatformRect FAndroidWindow::GetScreenRect()
 	const bool bMobileHDR32bpp = (MobileHDRCvar && MobileHDRCvar->GetValueOnAnyThread() == 1)
 		&& (FAndroidMisc::SupportsFloatingPointRenderTargets() == false || (MobileHDR32bppCvar && MobileHDR32bppCvar->GetValueOnAnyThread() == 1));
 
+	UE_LOG(LogAndroid, Log, TEXT("Requires Mosaic: %s"), bMobileHDR32bpp ? TEXT("YES") : TEXT("no"));
+
 	if (bMobileHDR32bpp)
 	{
 		const int32 OldMaxWidth = MaxWidth;
diff --git a/Engine/Source/Runtime/OpenGLDrv/Private/Android/AndroidOpenGL.cpp b/Engine/Source/Runtime/OpenGLDrv/Private/Android/AndroidOpenGL.cpp
--- a/Engine/Source/Runtime/OpenGLDrv/Private/Android/AndroidOpenGL.cpp
+++ b/Engine/Source/Runtime/OpenGLDrv/Private/Android/AndroidOpenGL.cpp
@@ -352,9 +352,24 @@ void FAndroidOpenGL::ProcessExtensions(const FString& ExtensionsString)
 	bSupportsETC2 = bES30Support;
 	bUseES30ShadingLanguage = bES30Support;
 
-	// Attempt to find ES 3.0 glTexStorage2D if we're on an Adreno device that supports it.
-	if( FString(ANSI_TO_TCHAR((const ANSICHAR*)glGetString(GL_RENDERER))).Contains(TEXT("Adreno")) )
+	FString RendererString = FString(ANSI_TO_TCHAR((const ANSICHAR*)glGetString(GL_RENDERER)));
+
+	if (RendererString.Contains(TEXT("SGX 540")))
+	{
+		UE_LOG(LogRHI, Warning, TEXT("Disabling support for GL_OES_packed_depth_stencil on SGX 540"));
+		bSupportsPackedDepthStencil = false;
+	}
+
+	if( RendererString.Contains(TEXT("Adreno")) )
 	{
+		// Adreno 2xx doesn't work with packed depth stencil enabled
+		if (RendererString.Contains(TEXT("Adreno (TM) 2")))
+		{
+			UE_LOG(LogRHI, Warning, TEXT("Disabling support for GL_OES_packed_depth_stencil on Adreno 2xx"));
+			bSupportsPackedDepthStencil = false;
+		}
+
+		// Attempt to find ES 3.0 glTexStorage2D if we're on an Adreno device that supports it.
 		glTexStorage2D = (PFNGLTEXSTORAGE2DPROC)((void*)eglGetProcAddress("glTexStorage2D"));
 		if( glTexStorage2D != NULL )
 		{
@@ -382,7 +397,7 @@ void FAndroidOpenGL::ProcessExtensions(const FString& ExtensionsString)
 	bSupportsTextureCubeLodEXT = false;
 
 	// On some Android devices with Mali GPUs textureCubeLod is not available.
-	if( FString(ANSI_TO_TCHAR((const ANSICHAR*)glGetString(GL_RENDERER))).Contains(TEXT("Mali-400")) )
+	if( RendererString.Contains(TEXT("Mali-400")) )
 	{
 		bSupportsShaderTextureCubeLod = false;
 	}
diff --git a/Engine/Source/Runtime/OpenGLDrv/Private/OpenGLDevice.cpp b/Engine/Source/Runtime/OpenGLDrv/Private/OpenGLDevice.cpp
--- a/Engine/Source/Runtime/OpenGLDrv/Private/OpenGLDevice.cpp
+++ b/Engine/Source/Runtime/OpenGLDrv/Private/OpenGLDevice.cpp
@@ -941,6 +941,122 @@ static bool VerifyCompiledShader(GLuint Shader, const ANSICHAR* GlslCode, bool I
 }
 #endif
 
+static void CheckVaryingLimit()
+{
+#if PLATFORM_ANDROID
+	FOpenGL::bRequiresGLFragCoordVaryingLimitHack = false;
+	if (IsES2Platform(GMaxRHIShaderPlatform))
+	{
+		// Some mobile GPUs require an available varying vector to support gl_FragCoord.
+		// If there are only 8 supported, it is possible to run out of varyings on these
+		// GPUs so test to see if need to fake gl_FragCoord with the assumption it is
+		// used for mobile HDR mosaic.
+
+		// Do not need to do this check if more than 8 varyings supported
+		if (FOpenGL::GetMaxVaryingVectors() > 8)
+			return;
+
+		// Make sure MobileHDR is on and device needs mosaic
+		static auto* MobileHDRCvar = IConsoleManager::Get().FindTConsoleVariableDataInt(TEXT("r.MobileHDR"));
+		static auto* MobileHDR32bppCvar = IConsoleManager::Get().FindTConsoleVariableDataInt(TEXT("r.MobileHDR32bpp"));
+		const bool bMobileHDR32bpp = (MobileHDRCvar && MobileHDRCvar->GetValueOnAnyThread() == 1)
+			&& (FAndroidMisc::SupportsFloatingPointRenderTargets() == false || (MobileHDR32bppCvar && MobileHDR32bppCvar->GetValueOnAnyThread() == 1));
+		if (!bMobileHDR32bpp)
+			return;
+
+		UE_LOG(LogRHI, Display, TEXT("Testing for gl_FragCoord requiring a varying since mosaic is enabled"));
+		const ANSICHAR* TestVertexProgram = "\n"
+			"#version 100\n"
+			"attribute vec4 in_ATTRIBUTE0;\n"
+			"attribute vec4 in_ATTRIBUTE1;\n"
+			"varying highp vec4 TexCoord0;\n"
+			"varying highp vec4 TexCoord1;\n"
+			"varying highp vec4 TexCoord2;\n"
+			"varying highp vec4 TexCoord3;\n"
+			"varying highp vec4 TexCoord4;\n"
+			"varying highp vec4 TexCoord5;\n"
+			"varying highp vec4 TexCoord6;\n"
+			"varying highp vec4 TexCoord7;\n"
+			"void main()\n"
+			"{\n"
+			"   TexCoord0 = in_ATTRIBUTE1 * vec4(0.1,0.2,0.3,0.4);\n"
+			"   TexCoord1 = in_ATTRIBUTE1 * vec4(0.5,0.6,0.7,0.8);\n"
+			"   TexCoord2 = in_ATTRIBUTE1 * vec4(0.12,0.22,0.32,0.42);\n"
+			"   TexCoord3 = in_ATTRIBUTE1 * vec4(0.52,0.62,0.72,0.82);\n"
+			"   TexCoord4 = in_ATTRIBUTE1 * vec4(0.14,0.24,0.34,0.44);\n"
+			"   TexCoord5 = in_ATTRIBUTE1 * vec4(0.54,0.64,0.74,0.84);\n"
+			"   TexCoord6 = in_ATTRIBUTE1 * vec4(0.16,0.26,0.36,0.46);\n"
+			"   TexCoord7 = in_ATTRIBUTE1 * vec4(0.56,0.66,0.76,0.86);\n"
+			"	gl_Position.xyzw = in_ATTRIBUTE0;\n"
+			"}\n";
+		const ANSICHAR* TestFragmentProgram = "\n"
+			"#version 100\n"
+			"varying highp vec4 TexCoord0;\n"
+			"varying highp vec4 TexCoord1;\n"
+			"varying highp vec4 TexCoord2;\n"
+			"varying highp vec4 TexCoord3;\n"
+			"varying highp vec4 TexCoord4;\n"
+			"varying highp vec4 TexCoord5;\n"
+			"varying highp vec4 TexCoord6;\n"
+			"varying highp vec4 TexCoord7;\n"
+			"void main()\n"
+			"{\n"
+			"   gl_FragColor = TexCoord0 * TexCoord1 * TexCoord2 * TexCoord3 * TexCoord4 * TexCoord5 * TexCoord6 * TexCoord7 * gl_FragCoord.xyxy;"
+			"}\n";
+
+		FOpenGLCodeHeader Header;
+		Header.FrequencyMarker = 0x5653;
+		Header.GlslMarker = 0x474c534c;
+		TArray<uint8> VertexCode;
+		{
+			FMemoryWriter Writer(VertexCode);
+			Writer << Header;
+			Writer.Serialize((void*)(TestVertexProgram), strlen(TestVertexProgram) + 1);
+			Writer.Close();
+		}
+		Header.FrequencyMarker = 0x5053;
+		Header.GlslMarker = 0x474c534c;
+		TArray<uint8> FragmentCode;
+		{
+			FMemoryWriter Writer(FragmentCode);
+			Writer << Header;
+			Writer.Serialize((void*)(TestFragmentProgram), strlen(TestFragmentProgram) + 1);
+			Writer.Close();
+		}
+
+		// Try to compile test shaders
+		TRefCountPtr<FOpenGLVertexShader> VertexShader = (FOpenGLVertexShader*)(RHICreateVertexShader(VertexCode).GetReference());
+		if (!VerifyCompiledShader(VertexShader->Resource, TestVertexProgram, false))
+		{
+			UE_LOG(LogRHI, Warning, TEXT("Vertex shader for varying test failed to compile. Try running anyway."));
+			return;
+		}
+		TRefCountPtr<FOpenGLPixelShader> PixelShader = (FOpenGLPixelShader*)(RHICreatePixelShader(FragmentCode).GetReference());
+		if (!VerifyCompiledShader(PixelShader->Resource, TestFragmentProgram, false))
+		{
+			UE_LOG(LogRHI, Warning, TEXT("Fragment shader for varying test failed to compile. Try running anyway."));
+			return;
+		}
+
+		// Now try linking them.. this is where gl_FragCoord may cause a failure
+		GLuint Program = glCreateProgram();
+		glAttachShader(Program, VertexShader->Resource);
+		glAttachShader(Program, PixelShader->Resource);
+		glLinkProgram(Program);
+		GLint LinkStatus = 0;
+		glGetProgramiv(Program, GL_LINK_STATUS, &LinkStatus);
+		if (LinkStatus != GL_TRUE)
+		{
+			FOpenGL::bRequiresGLFragCoordVaryingLimitHack = true;
+			UE_LOG(LogRHI, Warning, TEXT("gl_FragCoord uses a varying... enabled hack"));
+			return;
+		}
+
+		UE_LOG(LogRHI, Warning, TEXT("gl_FragCoord does not need a varying"));
+	}
+#endif
+}
+
 static void CheckTextureCubeLodSupport()
 {
 #if PLATFORM_ANDROID
@@ -1105,6 +1221,7 @@ void FOpenGLDynamicRHI::Init()
 	GIsRHIInitialized = true;
 
 	CheckTextureCubeLodSupport();
+	CheckVaryingLimit();
 }
 
 void FOpenGLDynamicRHI::Shutdown()
diff --git a/Engine/Source/Runtime/OpenGLDrv/Private/OpenGLDrv.cpp b/Engine/Source/Runtime/OpenGLDrv/Private/OpenGLDrv.cpp
--- a/Engine/Source/Runtime/OpenGLDrv/Private/OpenGLDrv.cpp
+++ b/Engine/Source/Runtime/OpenGLDrv/Private/OpenGLDrv.cpp
@@ -424,6 +424,7 @@ GLint FOpenGLBase::MaxVertexTextureImageUnits = -1;
 GLint FOpenGLBase::MaxGeometryTextureImageUnits = -1;
 GLint FOpenGLBase::MaxHullTextureImageUnits = -1;
 GLint FOpenGLBase::MaxDomainTextureImageUnits = -1;
+GLint FOpenGLBase::MaxVaryingVectors = -1;
 GLint FOpenGLBase::MaxVertexUniformComponents = -1;
 GLint FOpenGLBase::MaxPixelUniformComponents = -1;
 GLint FOpenGLBase::MaxGeometryUniformComponents = -1;
diff --git a/Engine/Source/Runtime/OpenGLDrv/Private/OpenGLES2.cpp b/Engine/Source/Runtime/OpenGLDrv/Private/OpenGLES2.cpp
--- a/Engine/Source/Runtime/OpenGLDrv/Private/OpenGLES2.cpp
+++ b/Engine/Source/Runtime/OpenGLDrv/Private/OpenGLES2.cpp
@@ -121,6 +121,9 @@ bool FOpenGLES2::bRequiresTextureCubeLodEXTToTextureCubeLodDefine = false;
 /* Some android platforms do not support the GL_OES_standard_derivatives extension */
 bool FOpenGLES2::bSupportsStandardDerivativesExtension = false;
 
+/* This is a hack to remove the gl_FragCoord if shader will fail to link if exceeding the max varying on android platforms */
+bool FOpenGLES2::bRequiresGLFragCoordVaryingLimitHack = false;
+
 bool FOpenGLES2::SupportsDisjointTimeQueries()
 {
 	bool bAllowDisjointTimerQueries = false;
@@ -130,6 +133,7 @@ bool FOpenGLES2::SupportsDisjointTimeQueries()
 
 void FOpenGLES2::ProcessQueryGLInt()
 {
+	LOG_AND_GET_GL_INT(GL_MAX_VARYING_VECTORS, 0, MaxVaryingVectors);
 	LOG_AND_GET_GL_INT(GL_MAX_VERTEX_UNIFORM_VECTORS, 0, MaxVertexUniformComponents);
 	LOG_AND_GET_GL_INT(GL_MAX_FRAGMENT_UNIFORM_VECTORS, 0, MaxPixelUniformComponents);
 
diff --git a/Engine/Source/Runtime/OpenGLDrv/Private/OpenGLES31.cpp b/Engine/Source/Runtime/OpenGLDrv/Private/OpenGLES31.cpp
--- a/Engine/Source/Runtime/OpenGLDrv/Private/OpenGLES31.cpp
+++ b/Engine/Source/Runtime/OpenGLDrv/Private/OpenGLES31.cpp
@@ -121,6 +121,9 @@ bool FOpenGLES31::bRequiresDontEmitPrecisionForTextureSamplers = false;
 /* Some android platforms require textureCubeLod to be used some require textureCubeLodEXT however they either inconsistently or don't use the GL_TextureCubeLodEXT extension definition */
 bool FOpenGLES31::bRequiresTextureCubeLodEXTToTextureCubeLodDefine = false;
 
+/* This is a hack to remove the gl_FragCoord if shader will fail to link if exceeding the max varying on android platforms */
+bool FOpenGLES31::bRequiresGLFragCoordVaryingLimitHack = false;
+
 /** GL_EXT_disjoint_timer_query or GL_NV_timer_query*/
 bool FOpenGLES31::bSupportsDisjointTimeQueries = false;
 
@@ -171,8 +174,10 @@ void FOpenGLES31::ProcessQueryGLInt()
 {
 	if(bES2Fallback)
 	{
+		LOG_AND_GET_GL_INT(GL_MAX_VARYING_VECTORS, 0, MaxVaryingVectors);
 		LOG_AND_GET_GL_INT(GL_MAX_VERTEX_UNIFORM_VECTORS, 0, MaxVertexUniformComponents);
 		LOG_AND_GET_GL_INT(GL_MAX_FRAGMENT_UNIFORM_VECTORS, 0, MaxPixelUniformComponents);
+		MaxVaryingVectors *= 4;
 		MaxVertexUniformComponents *= 4;
 		MaxPixelUniformComponents *= 4;
 		MaxGeometryUniformComponents = 0;
@@ -183,6 +188,7 @@ void FOpenGLES31::ProcessQueryGLInt()
 	}
 	else
 	{
+		GET_GL_INT(GL_MAX_VARYING_VECTORS, 0, MaxVaryingVectors);
 		GET_GL_INT(GL_MAX_VERTEX_UNIFORM_COMPONENTS, 0, MaxVertexUniformComponents);
 		GET_GL_INT(GL_MAX_FRAGMENT_UNIFORM_COMPONENTS, 0, MaxPixelUniformComponents);
 		GET_GL_INT(GL_MAX_GEOMETRY_UNIFORM_COMPONENTS_EXT, 0, MaxGeometryUniformComponents);
diff --git a/Engine/Source/Runtime/OpenGLDrv/Private/OpenGLShaders.cpp b/Engine/Source/Runtime/OpenGLDrv/Private/OpenGLShaders.cpp
--- a/Engine/Source/Runtime/OpenGLDrv/Private/OpenGLShaders.cpp
+++ b/Engine/Source/Runtime/OpenGLDrv/Private/OpenGLShaders.cpp
@@ -82,7 +82,7 @@ static bool VerifyCompiledShader(GLuint Shader, const ANSICHAR* GlslCode )
 {
 	SCOPE_CYCLE_COUNTER(STAT_OpenGLShaderCompileVerifyTime);
 
-#if UE_BUILD_DEBUG
+#if UE_BUILD_DEBUG || DEBUG_GL_SHADERS
 	if (FOpenGL::SupportsSeparateShaderObjects() && glIsProgram(Shader))
 	{
 		bool const bCompiledOK = VerifyLinkedProgram(Shader);
@@ -312,6 +312,22 @@ namespace
 		return true;
 	}
 
+	inline int CStringCountOccurances(TArray<ANSICHAR> & Source, const ANSICHAR * TargetString)
+	{
+		int32 TargetLen = FCStringAnsi::Strlen(TargetString);
+		int Count = 0;
+		int32 FoundIndex = 0;
+		for (const ANSICHAR * FoundPointer = FCStringAnsi::Strstr(Source.GetData(), TargetString);
+			nullptr != FoundPointer;
+			FoundPointer = FCStringAnsi::Strstr(Source.GetData() + FoundIndex, TargetString))
+		{
+			FoundIndex = FoundPointer - Source.GetData();
+			FoundIndex += TargetLen;
+			Count++;
+		}
+		return Count;
+	}
+
 	inline bool MoveHashLines(TArray<ANSICHAR> & Dest, TArray<ANSICHAR> & Source)
 	{
 		// Walk through the lines to find the first non-# line...
@@ -637,6 +653,16 @@ ShaderType* CompileOpenGLShader(const TArray<uint8>& Code)
 						AppendCString(GlslCode,
 							"#define textureCubeLodEXT textureCubeLod \n");
 					}
+
+					// Deal with gl_FragCoord using one of the varying vectors and shader possibly exceeding the limit
+					if (FOpenGL::RequiresGLFragCoordVaryingLimitHack())
+					{
+						if (CStringCountOccurances(GlslCodeOriginal, "vec4 var_TEXCOORD") >= FOpenGL::GetMaxVaryingVectors())
+						{
+							// It is likely gl_FragCoord is used for mosaic color output so use an appropriate constant
+							ReplaceCString(GlslCodeOriginal, "gl_FragCoord.xy", "vec2(400.5,240.5)");
+						}
+					}
 				}
 			}
 		}
diff --git a/Engine/Source/Runtime/OpenGLDrv/Public/OpenGL.h b/Engine/Source/Runtime/OpenGLDrv/Public/OpenGL.h
--- a/Engine/Source/Runtime/OpenGLDrv/Public/OpenGL.h
+++ b/Engine/Source/Runtime/OpenGLDrv/Public/OpenGL.h
@@ -169,6 +169,7 @@ public:
 	static FORCEINLINE GLint GetFirstComputeTextureUnit()		{ return 0; }
 	static FORCEINLINE GLint GetFirstComputeUAVUnit()			{ return 0; }
 
+	static FORCEINLINE GLint GetMaxVaryingVectors()				{ check(MaxVaryingVectors != -1); return MaxVaryingVectors; }
 	static FORCEINLINE GLint GetMaxPixelUniformComponents()		{ check(MaxPixelUniformComponents != -1); return MaxPixelUniformComponents; }
 	static FORCEINLINE GLint GetMaxVertexUniformComponents()	{ check(MaxVertexUniformComponents != -1); return MaxVertexUniformComponents; }
 	static FORCEINLINE GLint GetMaxGeometryUniformComponents()	{ check(MaxGeometryUniformComponents != -1); return MaxGeometryUniformComponents; }
@@ -356,6 +357,7 @@ protected:
 	static GLint MaxGeometryUniformComponents;
 	static GLint MaxHullUniformComponents;
 	static GLint MaxDomainUniformComponents;
+	static GLint MaxVaryingVectors;
 
 	/** GL_KHR_texture_compression_astc_ldr */
 	static bool bSupportsASTC;
diff --git a/Engine/Source/Runtime/OpenGLDrv/Public/OpenGLES2.h b/Engine/Source/Runtime/OpenGLDrv/Public/OpenGLES2.h
--- a/Engine/Source/Runtime/OpenGLDrv/Public/OpenGLES2.h
+++ b/Engine/Source/Runtime/OpenGLDrv/Public/OpenGLES2.h
@@ -128,6 +128,7 @@ struct FOpenGLES2 : public FOpenGLBase
 	static FORCEINLINE bool RequiresDontEmitPrecisionForTextureSamplers() { return bRequiresDontEmitPrecisionForTextureSamplers; }
 	static FORCEINLINE bool RequiresTextureCubeLodEXTToTextureCubeLodDefine() { return bRequiresTextureCubeLodEXTToTextureCubeLodDefine; }
 	static FORCEINLINE bool SupportsStandardDerivativesExtension()		{ return bSupportsStandardDerivativesExtension; }
+	static FORCEINLINE bool RequiresGLFragCoordVaryingLimitHack()		{ return bRequiresGLFragCoordVaryingLimitHack; }
 
 	static FORCEINLINE int32 GetReadHalfFloatPixelsEnum()				{ return GL_HALF_FLOAT_OES; }
 
@@ -478,6 +479,8 @@ public:
 	/* Some android platforms require textureCubeLod to be used some require textureCubeLodEXT however they either inconsistently or don't use the GL_TextureCubeLodEXT extension definition */
 	static bool bRequiresTextureCubeLodEXTToTextureCubeLodDefine;
 
+	/* This is a hack to remove the gl_FragCoord if shader will fail to link if exceeding the max varying on android platforms */
+	static bool bRequiresGLFragCoordVaryingLimitHack;
 };
 
 
diff --git a/Engine/Source/Runtime/OpenGLDrv/Public/OpenGLES31.h b/Engine/Source/Runtime/OpenGLDrv/Public/OpenGLES31.h
--- a/Engine/Source/Runtime/OpenGLDrv/Public/OpenGLES31.h
+++ b/Engine/Source/Runtime/OpenGLDrv/Public/OpenGLES31.h
@@ -84,6 +84,7 @@ struct FOpenGLES31 : public FOpenGLBase
 	static FORCEINLINE bool RequiresDontEmitPrecisionForTextureSamplers() { return bRequiresDontEmitPrecisionForTextureSamplers; }
 	static FORCEINLINE bool RequiresTextureCubeLodEXTToTextureCubeLodDefine() { return bRequiresTextureCubeLodEXTToTextureCubeLodDefine; }
 	static FORCEINLINE bool SupportsStandardDerivativesExtension()		{ return true; }
+	static FORCEINLINE bool RequiresGLFragCoordVaryingLimitHack()		{ return bRequiresGLFragCoordVaryingLimitHack; }
 	static FORCEINLINE GLenum GetVertexHalfFloatFormat()				{ return bES2Fallback ? GL_HALF_FLOAT_OES : GL_HALF_FLOAT; }
 
 
@@ -933,6 +934,9 @@ public:
 
 	/* Some android platforms require textureCubeLod to be used some require textureCubeLodEXT however they either inconsistently or don't use the GL_TextureCubeLodEXT extension definition */
 	static bool bRequiresTextureCubeLodEXTToTextureCubeLodDefine;
+
+	/* This is a hack to remove the gl_FragCoord if shader will fail to link if exceeding the max varying on android platforms */
+	static bool bRequiresGLFragCoordVaryingLimitHack;
 };
 
 // yes they are different between the ES2 extension and ES3.x and GL3.x core
