[INTEGRATE] Change 2445577 by [REDACTED_PERSON][REDACTED_HANDLE] on 2015/02/13 17:52:45

	Do not use ES 3.0 shading language on ES 2.0 Android
	#android
	#ue4

diff --git a/Engine/Source/Runtime/OpenGLDrv/Private/Android/AndroidOpenGL.cpp b/Engine/Source/Runtime/OpenGLDrv/Private/Android/AndroidOpenGL.cpp
--- a/Engine/Source/Runtime/OpenGLDrv/Private/Android/AndroidOpenGL.cpp
+++ b/Engine/Source/Runtime/OpenGLDrv/Private/Android/AndroidOpenGL.cpp
@@ -350,7 +350,7 @@ void FAndroidOpenGL::ProcessExtensions(const FString& ExtensionsString)
 	glGetObjectLabelEXT = (PFNGLGETOBJECTLABELEXTPROC)((void*)eglGetProcAddress("glGetObjectLabelEXT"));
 
 	bSupportsETC2 = bES30Support;
-	bUseES30ShadingLanguage = bES30Support;
+//	bUseES30ShadingLanguage = bES30Support;  // this was causing a problem on some drivers with Android 5.0.1
 
 	FString RendererString = FString(ANSI_TO_TCHAR((const ANSICHAR*)glGetString(GL_RENDERER)));
 
