OpenGL: proper hint for resource usage.

- That variable seems to apply to vertex buffers only.
- Kills a lot of performance warnings issued by NVidia drivers.

Merging CL 2455788 from my branch to main.

diff --git a/Engine/Source/Runtime/OpenGLDrv/Public/OpenGLResources.h b/Engine/Source/Runtime/OpenGLDrv/Public/OpenGLResources.h
--- a/Engine/Source/Runtime/OpenGLDrv/Public/OpenGLResources.h
+++ b/Engine/Source/Runtime/OpenGLDrv/Public/OpenGLResources.h
@@ -81,7 +81,7 @@ class TOpenGLBuffer : public BaseType
 	{
 		GLenum Access = GL_STATIC_DRAW;
 
-		if (!OpenGLConsoleVariables::bUseStagingBuffer)
+		if (Type != GL_ARRAY_BUFFER || !OpenGLConsoleVariables::bUseStagingBuffer)
 		{
 			Access = bStreamDraw ? GL_STREAM_DRAW : (IsDynamic() ? GL_DYNAMIC_DRAW : GL_STATIC_DRAW);
 		}
