Fix CIS build errors from shader predraw (CL #2460109) by removing the TEXTURE_1D symbols that aren't present on some OpenGL platforms & aren't used anyway.

diff --git a/Engine/Source/Runtime/OpenGLDrv/Public/OpenGLResources.h b/Engine/Source/Runtime/OpenGLDrv/Public/OpenGLResources.h
--- a/Engine/Source/Runtime/OpenGLDrv/Public/OpenGLResources.h
+++ b/Engine/Source/Runtime/OpenGLDrv/Public/OpenGLResources.h
@@ -878,11 +878,6 @@ public:
 			Tex.Z = InSizeZ;
 			switch(InTarget)
 			{
-				case GL_TEXTURE_1D:
-				{
-					Tex.Type = SCTT_Texture1D;
-					break;
-				}
 				case GL_TEXTURE_2D:
 				case GL_TEXTURE_2D_MULTISAMPLE:
 				{
@@ -899,11 +894,6 @@ public:
 					Tex.Type = SCTT_TextureCube;
 					break;
 				}
-				case GL_TEXTURE_1D_ARRAY:
-				{
-					Tex.Type = SCTT_Texture1DArray;
-					break;
-				}
 				case GL_TEXTURE_2D_ARRAY:
 				{
 					Tex.Type = SCTT_Texture2DArray;
