ES2: Remap vertex attributes when max vert attribs < 16
Fixes PVR SGX 540 crashes in UE-5155
[Requires shader compile worker rebuild.]

diff --git a/Engine/Shaders/ShaderVersion.usf b/Engine/Shaders/ShaderVersion.usf
--- a/Engine/Shaders/ShaderVersion.usf
+++ b/Engine/Shaders/ShaderVersion.usf
@@ -2,4 +2,4 @@
 // Each time the console command is executed it generates a new GUID. As this file is included
 // in common.usf (which should be included in any shader) it allows to invalidate the shader DDC.
 // 
-// GUID = {2490AEE3-37E5-4684-93D8-F170F910238C}
+// GUID = {2490AEE3-37E5-4684-93D8-F170F910238D}
diff --git a/Engine/Source/Developer/ShaderFormatOpenGL/Private/OpenGLShaderCompiler.cpp b/Engine/Source/Developer/ShaderFormatOpenGL/Private/OpenGLShaderCompiler.cpp
--- a/Engine/Source/Developer/ShaderFormatOpenGL/Private/OpenGLShaderCompiler.cpp
+++ b/Engine/Source/Developer/ShaderFormatOpenGL/Private/OpenGLShaderCompiler.cpp
@@ -666,6 +666,27 @@ static void BuildShaderOutput(
         }
 	}
 
+	// Generate vertex attribute remapping table.
+	// This is used on devices where GL_MAX_VERTEX_ATTRIBS < 16
+	if (Frequency == SF_Vertex)
+	{
+		uint32 AttributeMask = Header.Bindings.InOutMask;
+		int32 NextAttributeSlot = 0;
+		Header.Bindings.VertexRemappedMask = 0;
+		for (int32 AttributeIndex = 0; AttributeIndex<16; AttributeIndex++, AttributeMask >>= 1)
+		{
+			if (AttributeMask & 0x1)
+			{
+				Header.Bindings.VertexRemappedMask |= (1 << NextAttributeSlot);
+				Header.Bindings.VertexAttributeRemap[AttributeIndex] = NextAttributeSlot++;
+			}
+			else
+			{
+				Header.Bindings.VertexAttributeRemap[AttributeIndex] = -1;
+			}
+		}
+	}
+
 	// Then the list of outputs.
 	if (FCStringAnsi::Strncmp(ShaderSource, OutputsPrefix, OutputsPrefixLen) == 0)
 	{
diff --git a/Engine/Source/Developer/ShaderFormatOpenGL/Private/ShaderFormatOpenGL.cpp b/Engine/Source/Developer/ShaderFormatOpenGL/Private/ShaderFormatOpenGL.cpp
--- a/Engine/Source/Developer/ShaderFormatOpenGL/Private/ShaderFormatOpenGL.cpp
+++ b/Engine/Source/Developer/ShaderFormatOpenGL/Private/ShaderFormatOpenGL.cpp
@@ -22,14 +22,14 @@ class FShaderFormatGLSL : public IShaderFormat
 	enum 
 	{
 		/** Version for shader format, this becomes part of the DDC key. */
-		UE_SHADER_GLSL_150_VER = 24,
-		UE_SHADER_GLSL_150_MAC_VER = 4,
-		UE_SHADER_GLSL_430_VER = 54,
-		UE_SHADER_GLSL_ES2_VER = 12,
-		UE_SHADER_GLSL_150ES2_VER = 15,
-		UE_SHADER_GLSL_ES2_VER_WEBGL  = 13,
-		UE_SHADER_GLSL_ES2_IOS_VER  = 3,
-		UE_SHADER_GLSL_310_ES_EXT_VER = 1,
+		UE_SHADER_GLSL_150_VER = 25,
+		UE_SHADER_GLSL_150_MAC_VER = 5,
+		UE_SHADER_GLSL_430_VER = 55,
+		UE_SHADER_GLSL_ES2_VER = 13,
+		UE_SHADER_GLSL_150ES2_VER = 16,
+		UE_SHADER_GLSL_ES2_VER_WEBGL  = 14,
+		UE_SHADER_GLSL_ES2_IOS_VER  = 4,
+		UE_SHADER_GLSL_310_ES_EXT_VER = 2,
 	}; 
 
 	void CheckFormat(FName Format) const
diff --git a/Engine/Source/Runtime/OpenGLDrv/Private/OpenGLCommands.cpp b/Engine/Source/Runtime/OpenGLDrv/Private/OpenGLCommands.cpp
--- a/Engine/Source/Runtime/OpenGLDrv/Private/OpenGLCommands.cpp
+++ b/Engine/Source/Runtime/OpenGLDrv/Private/OpenGLCommands.cpp
@@ -1840,6 +1840,7 @@ void FOpenGLDynamicRHI::EnableVertexElementCached(
 	check( !(FOpenGL::SupportsVertexAttribBinding() && OpenGLConsoleVariables::bUseVAB));
 
 	GLuint AttributeIndex = VertexElement.AttributeIndex;
+	AttributeIndex = RemapVertexAttrib(AttributeIndex);
 	FOpenGLCachedAttr &Attr = ContextState.VertexAttrs[AttributeIndex];
 
 	if (!Attr.bEnabled)
@@ -1896,6 +1897,8 @@ void FOpenGLDynamicRHI::EnableVertexElementCachedZeroStride(FOpenGLContextState&
 	VERIFY_GL_SCOPE();
 
 	GLuint AttributeIndex = VertexElement.AttributeIndex;
+	AttributeIndex = RemapVertexAttrib(AttributeIndex);
+
 	FOpenGLCachedAttr &Attr = ContextState.VertexAttrs[AttributeIndex];
 	uint32 Stride = ZeroStrideVertexBuffer->GetSize();
 
@@ -1923,67 +1926,65 @@ void FOpenGLDynamicRHI::SetupVertexArrays(FOpenGLContextState& ContextState, uin
 	check(IsValidRef(PendingState.BoundShaderState));
 	check(IsValidRef(PendingState.BoundShaderState->VertexShader));
 	FOpenGLVertexDeclaration* VertexDeclaration = PendingState.BoundShaderState->VertexDeclaration;
-	uint32 AttributeMask = PendingState.BoundShaderState->VertexShader->Bindings.InOutMask;
 	for (int32 ElementIndex = 0; ElementIndex < VertexDeclaration->VertexElements.Num(); ElementIndex++)
 	{
 		FOpenGLVertexElement& VertexElement = VertexDeclaration->VertexElements[ElementIndex];
+		uint32 AttributeIndex = VertexElement.AttributeIndex;
+		const bool bAttribInUse = (PendingState.BoundShaderState->VertexShader->Bindings.InOutMask & (0x1 << AttributeIndex)) != 0;
+		if (!bAttribInUse)
+		{
+			continue; // skip unused attributes.
+		}
+
+		AttributeIndex = RemapVertexAttrib(AttributeIndex);
 
 		if ( VertexElement.StreamIndex < NumStreams)
 		{
 			FOpenGLStream* Stream = &Streams[VertexElement.StreamIndex];
 			uint32 Stride = Stream->Stride;
 
-			uint32 AttributeBit = (1 << VertexElement.AttributeIndex);
-			if ((AttributeBit & AttributeMask) == AttributeBit)
+			if( Stream->VertexBuffer->GetUsage() & BUF_ZeroStride )
 			{
-				if( Stream->VertexBuffer->GetUsage() & BUF_ZeroStride )
-				{
-					check(Stride == 0);
-					check(Stream->Offset == 0);
-					check(VertexElement.Offset == 0);
-					check(Stream->VertexBuffer->GetZeroStrideBuffer());
-					EnableVertexElementCachedZeroStride(
-						ContextState,
-						VertexElement,
-						MaxVertices,
-						Stream->VertexBuffer
-						);
-				}
-				else
-				{
-					check( Stride > 0 );
-					EnableVertexElementCached(
-						ContextState,
-						VertexElement,
-						Stride,
-						INDEX_TO_VOID(BaseVertexIndex * Stride + Stream->Offset + VertexElement.Offset),
-						Stream->VertexBuffer->Resource
-						);
-				}
-				UsedAttributes[VertexElement.AttributeIndex] = true;
+				check(Stride == 0);
+				check(Stream->Offset == 0);
+				check(VertexElement.Offset == 0);
+				check(Stream->VertexBuffer->GetZeroStrideBuffer());
+				EnableVertexElementCachedZeroStride(
+					ContextState,
+					VertexElement,
+					MaxVertices,
+					Stream->VertexBuffer
+					);
+			}
+			else
+			{
+				check( Stride > 0 );
+				EnableVertexElementCached(
+					ContextState,
+					VertexElement,
+					Stride,
+					INDEX_TO_VOID(BaseVertexIndex * Stride + Stream->Offset + VertexElement.Offset),
+					Stream->VertexBuffer->Resource
+					);
 			}
+			UsedAttributes[AttributeIndex] = true;
 		}
 		else
 		{
 			//workaround attributes with no streams
-			uint32 AttributeBit = (1 << VertexElement.AttributeIndex);
-			if ((AttributeBit & AttributeMask) == AttributeBit)
-			{
-				VERIFY_GL_SCOPE();
+			VERIFY_GL_SCOPE();
 
-				GLuint AttributeIndex = VertexElement.AttributeIndex;
-				FOpenGLCachedAttr &Attr = ContextState.VertexAttrs[AttributeIndex];
+			FOpenGLCachedAttr &Attr = ContextState.VertexAttrs[AttributeIndex];
 
-				if (Attr.bEnabled)
-				{
-					glDisableVertexAttribArray(AttributeIndex);
-					Attr.bEnabled = false;
-				}
+			if (Attr.bEnabled)
+			{
+				glDisableVertexAttribArray(AttributeIndex);
+				Attr.bEnabled = false;
+			}
 
-				float data[4] = { 0.0f};
+			float data[4] = { 0.0f};
 
-				glVertexAttrib4fv( AttributeIndex, data );
-			}
+			glVertexAttrib4fv(AttributeIndex, data);
 		}
 	}
 
@@ -2011,6 +2012,10 @@ void FOpenGLDynamicRHI::SetupVertexArraysVAB(FOpenGLContextState& ContextState,
 	check(IsValidRef(PendingState.BoundShaderState->VertexShader));
 	FOpenGLVertexDeclaration* VertexDeclaration = PendingState.BoundShaderState->VertexDeclaration;
 	uint32 AttributeMask = PendingState.BoundShaderState->VertexShader->Bindings.InOutMask;
+	if (FOpenGL::NeedsVertexAttribRemapTable())
+	{
+		AttributeMask = PendingState.BoundShaderState->VertexShader->Bindings.VertexRemappedMask;
+	}
 
 	if (ContextState.VertexDecl != VertexDeclaration || AttributeMask != ContextState.ActiveAttribMask)
 	{
@@ -2023,13 +2028,19 @@ void FOpenGLDynamicRHI::SetupVertexArraysVAB(FOpenGLContextState& ContextState,
 		for (int32 ElementIndex = 0; ElementIndex < VertexDeclaration->VertexElements.Num(); ElementIndex++)
 		{
 			FOpenGLVertexElement& VertexElement = VertexDeclaration->VertexElements[ElementIndex];
-			const uint32 AttributeIndex = VertexElement.AttributeIndex;
+			uint32 AttributeIndex = VertexElement.AttributeIndex;
+			const bool bAttribInUse = (PendingState.BoundShaderState->VertexShader->Bindings.InOutMask & (0x1 << AttributeIndex)) != 0;
+			if (bAttribInUse)
+			{
+				AttributeIndex = RemapVertexAttrib(AttributeIndex);
+			}
+
 			const uint32 StreamIndex = VertexElement.StreamIndex;
 
 			ContextState.MaxActiveAttrib = FMath::Max( ContextState.MaxActiveAttrib, AttributeIndex);
 
 			//only setup/track attributes actually in use
-			if (AttributeMask & (0x1 << AttributeIndex))
+			if (bAttribInUse)
 			{
 				if (VertexElement.StreamIndex < NumStreams)
 				{
@@ -2158,16 +2169,18 @@ void FOpenGLDynamicRHI::SetupVertexArraysUP(FOpenGLContextState& ContextState, v
 	check(IsValidRef(PendingState.BoundShaderState));
 	check(IsValidRef(PendingState.BoundShaderState->VertexShader));
 	FOpenGLVertexDeclaration* VertexDeclaration = PendingState.BoundShaderState->VertexDeclaration;
-	uint32 AttributeMask = PendingState.BoundShaderState->VertexShader->Bindings.InOutMask;
+
 	for (int32 ElementIndex = 0; ElementIndex < VertexDeclaration->VertexElements.Num(); ElementIndex++)
 	{
 		FOpenGLVertexElement &VertexElement = VertexDeclaration->VertexElements[ElementIndex];
 		check(VertexElement.StreamIndex < 1);
 
-		uint32 AttributeBit = (1 << VertexElement.AttributeIndex);
-		if ((AttributeBit & AttributeMask) == AttributeBit)
+		uint32 AttributeIndex = VertexElement.AttributeIndex;
+		const bool bAttribInUse = (PendingState.BoundShaderState->VertexShader->Bindings.InOutMask & (0x1 << AttributeIndex)) != 0;
+		if (bAttribInUse)
 		{
-			check( Stride > 0 );
+			AttributeIndex = RemapVertexAttrib(AttributeIndex);
+			check(Stride > 0);
 			EnableVertexElementCached(
 				ContextState,
 				VertexElement,
@@ -2175,7 +2188,7 @@ void FOpenGLDynamicRHI::SetupVertexArraysUP(FOpenGLContextState& ContextState, v
 				(void*)(((char*)Buffer) + VertexElement.Offset),
 				0
 				);
-			UsedAttributes[VertexElement.AttributeIndex] = true;
+			UsedAttributes[AttributeIndex] = true;
 		}
 	}
 
diff --git a/Engine/Source/Runtime/OpenGLDrv/Private/OpenGLES2.cpp b/Engine/Source/Runtime/OpenGLDrv/Private/OpenGLES2.cpp
--- a/Engine/Source/Runtime/OpenGLDrv/Private/OpenGLES2.cpp
+++ b/Engine/Source/Runtime/OpenGLDrv/Private/OpenGLES2.cpp
@@ -124,6 +124,9 @@ bool FOpenGLES2::bSupportsStandardDerivativesExtension = false;
 /* This is a hack to remove the gl_FragCoord if shader will fail to link if exceeding the max varying on android platforms */
 bool FOpenGLES2::bRequiresGLFragCoordVaryingLimitHack = false;
 
+/** Vertex attributes need remapping if GL_MAX_VERTEX_ATTRIBS < 16 */
+bool FOpenGLES2::bNeedsVertexAttribRemap = false;
+
 bool FOpenGLES2::SupportsDisjointTimeQueries()
 {
 	bool bAllowDisjointTimerQueries = false;
@@ -133,6 +136,17 @@ bool FOpenGLES2::SupportsDisjointTimeQueries()
 
 void FOpenGLES2::ProcessQueryGLInt()
 {
+	GLint MaxVertexAttribs;
+	LOG_AND_GET_GL_INT(GL_MAX_VERTEX_ATTRIBS, 0, MaxVertexAttribs);
+	bNeedsVertexAttribRemap = MaxVertexAttribs < 16;
+	if (bNeedsVertexAttribRemap)
+	{
+		UE_LOG(LogRHI, Warning,
+			TEXT("Device reports support for %d vertex attributes, UE4 requires 16. Rendering artifacts may occur."),
+			MaxVertexAttribs
+			);
+	}
+
 	LOG_AND_GET_GL_INT(GL_MAX_VARYING_VECTORS, 0, MaxVaryingVectors);
 	LOG_AND_GET_GL_INT(GL_MAX_VERTEX_UNIFORM_VECTORS, 0, MaxVertexUniformComponents);
 	LOG_AND_GET_GL_INT(GL_MAX_FRAGMENT_UNIFORM_VECTORS, 0, MaxPixelUniformComponents);
diff --git a/Engine/Source/Runtime/OpenGLDrv/Private/OpenGLShaders.cpp b/Engine/Source/Runtime/OpenGLDrv/Private/OpenGLShaders.cpp
--- a/Engine/Source/Runtime/OpenGLDrv/Private/OpenGLShaders.cpp
+++ b/Engine/Source/Runtime/OpenGLDrv/Private/OpenGLShaders.cpp
@@ -359,7 +359,7 @@ inline uint32 GetTypeHash(FAnsiCharArray const& CharArray)
 	return FCrc::MemCrc32(CharArray.GetData(), CharArray.Num() * sizeof(ANSICHAR));
 }
 
-static void BindShaderLocations(GLenum TypeEnum, GLuint Resource, uint16 InOutMask)
+static void BindShaderLocations(GLenum TypeEnum, GLuint Resource, uint16 InOutMask, const uint8 * RemapTable = nullptr)
 {
 	if ( OpenGLShaderPlatformNeedsBindLocation(GMaxRHIShaderPlatform) )
 	{
@@ -386,7 +386,18 @@ static void BindShaderLocations(GLenum TypeEnum, GLuint Resource, uint16 InOutMa
 							Buf[13] = '0' + (Index % 10);
 							Buf[14] = 0;
 						}
-						glBindAttribLocation(Resource, Index, Buf);
+
+						if (FOpenGL::NeedsVertexAttribRemapTable())
+						{
+							check(RemapTable != nullptr);
+							uint32 MappedAttributeIndex = RemapTable[Index];
+							check(MappedAttributeIndex < NUM_OPENGL_VERTEX_STREAMS);
+							glBindAttribLocation(Resource, MappedAttributeIndex, Buf);
+						}
+						else
+						{
+							glBindAttribLocation(Resource, Index, Buf);
+						}
 					}
 					Index++;
 					Mask >>= 1;
@@ -1606,7 +1617,8 @@ static FOpenGLLinkedProgram* LinkProgram( const FOpenGLLinkedProgramConfiguratio
 			// Bind attribute indices.
 			if (Config.Shaders[CrossCompiler::SHADER_STAGE_VERTEX].Resource)
 			{
-				BindShaderLocations(GL_VERTEX_SHADER, Program, Config.Shaders[CrossCompiler::SHADER_STAGE_VERTEX].Bindings.InOutMask);
+				auto& VertexBindings = Config.Shaders[CrossCompiler::SHADER_STAGE_VERTEX].Bindings;
+				BindShaderLocations(GL_VERTEX_SHADER, Program, VertexBindings.InOutMask, VertexBindings.VertexAttributeRemap);
 			}
 
 			// Bind frag data locations.
diff --git a/Engine/Source/Runtime/OpenGLDrv/Public/OpenGL.h b/Engine/Source/Runtime/OpenGLDrv/Public/OpenGL.h
--- a/Engine/Source/Runtime/OpenGLDrv/Public/OpenGL.h
+++ b/Engine/Source/Runtime/OpenGLDrv/Public/OpenGL.h
@@ -147,6 +147,7 @@ public:
 	static FORCEINLINE bool HasHardwareHiddenSurfaceRemoval()			{ return false; }
 	static FORCEINLINE bool AmdWorkaround()								{ return false; }
 	static FORCEINLINE bool SupportsSeparateShaderObjects()				{ return false; }
+	static FORCEINLINE bool NeedsVertexAttribRemapTable()				{ return false; }
 
 	static FORCEINLINE GLenum GetDepthFormat()							{ return GL_DEPTH_COMPONENT16; }
 	static FORCEINLINE GLenum GetShadowDepthFormat()					{ return GL_DEPTH_COMPONENT16; }
diff --git a/Engine/Source/Runtime/OpenGLDrv/Public/OpenGLDrv.h b/Engine/Source/Runtime/OpenGLDrv/Public/OpenGLDrv.h
--- a/Engine/Source/Runtime/OpenGLDrv/Public/OpenGLDrv.h
+++ b/Engine/Source/Runtime/OpenGLDrv/Public/OpenGLDrv.h
@@ -490,6 +490,18 @@ private:
 
 	/** Consumes about 100ms of GPU time (depending on resolution and GPU), useful for making sure we're not CPU bound when GPU profiling. */
 	void IssueLongGPUTask();
+
+	/** Remaps vertex attributes on devices where GL_MAX_VERTEX_ATTRIBS < 16 */
+	uint32 RemapVertexAttrib(uint32 VertexAttributeIndex)
+	{
+		if (FOpenGL::NeedsVertexAttribRemapTable())
+		{
+			check(VertexAttributeIndex < ARRAY_COUNT(PendingState.BoundShaderState->VertexShader->Bindings.VertexAttributeRemap));
+			VertexAttributeIndex = PendingState.BoundShaderState->VertexShader->Bindings.VertexAttributeRemap[VertexAttributeIndex];
+			check(VertexAttributeIndex < NUM_OPENGL_VERTEX_STREAMS); // check that this attribute has remaped correctly.
+		}
+		return VertexAttributeIndex;
+	}
 };
 
 /** Implements the OpenGLDrv module as a dynamic RHI providing module. */
diff --git a/Engine/Source/Runtime/OpenGLDrv/Public/OpenGLES2.h b/Engine/Source/Runtime/OpenGLDrv/Public/OpenGLES2.h
--- a/Engine/Source/Runtime/OpenGLDrv/Public/OpenGLES2.h
+++ b/Engine/Source/Runtime/OpenGLDrv/Public/OpenGLES2.h
@@ -133,6 +133,7 @@ struct FOpenGLES2 : public FOpenGLBase
 	static FORCEINLINE int32 GetReadHalfFloatPixelsEnum()				{ return GL_HALF_FLOAT_OES; }
 
 	static FORCEINLINE GLenum GetVertexHalfFloatFormat()				{ return GL_HALF_FLOAT_OES; }
+	static FORCEINLINE bool NeedsVertexAttribRemapTable()				{ return bNeedsVertexAttribRemap; }
 
 	// On iOS both glMapBufferOES() and glBufferSubData() for immediate vertex and index data
 	// is the slow path (they both hit GPU sync and data cache flush in driver according to profiling in driver symbols).
@@ -472,6 +473,8 @@ protected:
 	/** GL_OES_standard_derivations */
 	static bool bSupportsStandardDerivativesExtension;
 
+	/** Vertex attributes need remapping if GL_MAX_VERTEX_ATTRIBS < 16 */
+	static bool bNeedsVertexAttribRemap;
 public:
 	/* This is a hack to remove the calls to "precision sampler" defaults which are produced by the cross compiler however don't compile on some android platforms */
 	static bool bRequiresDontEmitPrecisionForTextureSamplers;
diff --git a/Engine/Source/Runtime/OpenGLDrv/Public/OpenGLShaderResources.h b/Engine/Source/Runtime/OpenGLDrv/Public/OpenGLShaderResources.h
--- a/Engine/Source/Runtime/OpenGLDrv/Public/OpenGLShaderResources.h
+++ b/Engine/Source/Runtime/OpenGLDrv/Public/OpenGLShaderResources.h
@@ -93,15 +93,18 @@ struct FOpenGLShaderBindings
 	uint8	NumUniformBuffers;
 	uint8	NumUAVs;
 	bool	bFlattenUB;
-
+	uint8	VertexAttributeRemap[16];
+	uint8	VertexRemappedMask;
 
 	FOpenGLShaderBindings() :
 		InOutMask(0),
 		NumSamplers(0),
 		NumUniformBuffers(0),
 		NumUAVs(0),
-		bFlattenUB(false)
+		bFlattenUB(false),
+		VertexRemappedMask(0)
 	{
+		FMemory::Memset(VertexAttributeRemap, 0xFF);
 	}
 
 	friend bool operator==( const FOpenGLShaderBindings &A, const FOpenGLShaderBindings& B)
@@ -190,6 +193,11 @@ inline FArchive& operator<<(FArchive& Ar, FOpenGLShaderBindings& Bindings)
 	Ar << Bindings.NumUniformBuffers;
 	Ar << Bindings.NumUAVs;
 	Ar << Bindings.bFlattenUB;
+	for (uint32 i = 0; i < ARRAY_COUNT(Bindings.VertexAttributeRemap); i++)
+	{
+		Ar << Bindings.VertexAttributeRemap[i];
+	}
+	Ar << Bindings.VertexRemappedMask;
 	return Ar;
 }
 
