[INTEGRATE] 4.7 release branch to main - Reenable use of #version 300 shaders on ES2 Android

diff --git a/Engine/Source/Runtime/OpenGLDrv/Private/Android/AndroidOpenGL.cpp b/Engine/Source/Runtime/OpenGLDrv/Private/Android/AndroidOpenGL.cpp
--- a/Engine/Source/Runtime/OpenGLDrv/Private/Android/AndroidOpenGL.cpp
+++ b/Engine/Source/Runtime/OpenGLDrv/Private/Android/AndroidOpenGL.cpp
@@ -350,7 +350,7 @@ void FAndroidOpenGL::ProcessExtensions(const FString& ExtensionsString)
 	glGetObjectLabelEXT = (PFNGLGETOBJECTLABELEXTPROC)((void*)eglGetProcAddress("glGetObjectLabelEXT"));
 
 	bSupportsETC2 = bES30Support;
-//	bUseES30ShadingLanguage = bES30Support;  // this was causing a problem on some drivers with Android 5.0.1
+	bUseES30ShadingLanguage = bES30Support;
 
 	FString RendererString = FString(ANSI_TO_TCHAR((const ANSICHAR*)glGetString(GL_RENDERER)));
 
