Killed a few NVidia perf warnings.

- even better hint flags, disregarding staging buffer at all.
- synchronous debug output (to get the callstacks for warning in the future).

(Merging CL 2459008, 2459016 to main).

diff --git a/Engine/Source/Runtime/OpenGLDrv/Private/OpenGLDevice.cpp b/Engine/Source/Runtime/OpenGLDrv/Private/OpenGLDevice.cpp
--- a/Engine/Source/Runtime/OpenGLDrv/Private/OpenGLDevice.cpp
+++ b/Engine/Source/Runtime/OpenGLDrv/Private/OpenGLDevice.cpp
@@ -422,6 +422,8 @@ void InitDebugContext()
 #if defined(GL_ARB_debug_output)
 	if (glDebugMessageCallbackARB)
 	{
+		// Synchronous output can slow things down, but we'll get better callstack if breaking in or crashing in the callback. This is debug only after all.
+		glEnable(GL_DEBUG_OUTPUT_SYNCHRONOUS);
 		glDebugMessageCallbackARB(GLDEBUGPROCARB(OpenGLDebugMessageCallbackARB), /*UserParam=*/ NULL);
 		bDebugOutputInitialized = (glGetError() == GL_NO_ERROR);
 	}
diff --git a/Engine/Source/Runtime/OpenGLDrv/Public/OpenGLResources.h b/Engine/Source/Runtime/OpenGLDrv/Public/OpenGLResources.h
--- a/Engine/Source/Runtime/OpenGLDrv/Public/OpenGLResources.h
+++ b/Engine/Source/Runtime/OpenGLDrv/Public/OpenGLResources.h
@@ -80,14 +80,9 @@ class TOpenGLBuffer : public BaseType
 
 	GLenum GetAccess()
 	{
-		GLenum Access = GL_STATIC_DRAW;
-
-		if (Type != GL_ARRAY_BUFFER || !OpenGLConsoleVariables::bUseStagingBuffer)
-		{
-			Access = bStreamDraw ? GL_STREAM_DRAW : (IsDynamic() ? GL_DYNAMIC_DRAW : GL_STATIC_DRAW);
-		}
-		
-		return Access;
+		// Previously there was special-case logic to always use GL_STATIC_DRAW for vertex buffers allocated from staging buffer.
+		// However it seems to be incorrect as NVidia drivers complain (via debug output callback) about VIDEO->HOST copying for buffers with such hints
+		return bStreamDraw ? GL_STREAM_DRAW : (IsDynamic() ? GL_DYNAMIC_DRAW : GL_STATIC_DRAW);
 	}
 public:
 
