Linux: better control  of texture streaming.

(Merging CL 2461849 to main).

diff --git a/Engine/Config/Linux/LinuxEngine.ini b/Engine/Config/Linux/LinuxEngine.ini
--- a/Engine/Config/Linux/LinuxEngine.ini
+++ b/Engine/Config/Linux/LinuxEngine.ini
@@ -6,3 +6,7 @@ r.setres=1280x720
 
 [/Script/SourceCodeAccess.SourceCodeAccessSettings]
 PreferredAccessor=KDevelopSourceCodeAccessor
+
+[TextureStreaming]
+; PoolSizeVRAMPercentage is how much percentage of GPU Dedicated VRAM should be used as a TexturePool cache for streaming textures (0 - unlimited streaming)
+PoolSizeVRAMPercentage=70
diff --git a/Engine/Source/Runtime/OpenGLDrv/Private/OpenGLDevice.cpp b/Engine/Source/Runtime/OpenGLDrv/Private/OpenGLDevice.cpp
--- a/Engine/Source/Runtime/OpenGLDrv/Private/OpenGLDevice.cpp
+++ b/Engine/Source/Runtime/OpenGLDrv/Private/OpenGLDevice.cpp
@@ -312,6 +312,12 @@ static void APIENTRY OpenGLDebugMessageCallbackARB(
 				ANSI_TO_TCHAR(Message)
 				);
 		}
+
+		// this is a debugging code to catch VIDEO->HOST copying
+		if (Id == 131186)
+		{
+			int A = 5;
+		}
 	}
 #endif
 }
@@ -480,7 +486,7 @@ static void InitRHICapabilitiesForGL()
 
 	GTexturePoolSize = 0;
 	GPoolSizeVRAMPercentage = 0;
-#if PLATFORM_WINDOWS
+#if PLATFORM_WINDOWS || PLATFORM_LINUX
 	GConfig->GetInt( TEXT( "TextureStreaming" ), TEXT( "PoolSizeVRAMPercentage" ), GPoolSizeVRAMPercentage, GEngineIni );	
 #endif
 
@@ -1187,7 +1193,7 @@ void FOpenGLDynamicRHI::Init()
 		ResourceIt->InitRHI();
 	}
 
-#if PLATFORM_WINDOWS || PLATFORM_MAC
+#if PLATFORM_WINDOWS || PLATFORM_MAC || PLATFORM_LINUX
 
 	extern int64 GOpenGLDedicatedVideoMemory;
 	extern int64 GOpenGLTotalGraphicsMemory;
