Add an extension directive for tesselation.

- Fixes compiling shaders right now, proper resolution pending (UE-11222, hlslcc should emit appropriate #extension).

(Merging CL 2461844 from my branch).

diff --git a/Engine/Source/Runtime/OpenGLDrv/Private/OpenGLShaders.cpp b/Engine/Source/Runtime/OpenGLDrv/Private/OpenGLShaders.cpp
--- a/Engine/Source/Runtime/OpenGLDrv/Private/OpenGLShaders.cpp
+++ b/Engine/Source/Runtime/OpenGLDrv/Private/OpenGLShaders.cpp
@@ -563,6 +563,14 @@ ShaderType* CompileOpenGLShader(const TArray<uint8>& Code)
 			}
 		}
 
+		// Add an extension directive (in line with https://www.opengl.org/registry/specs/ARB/tessellation_shader.txt)
+		// Note that this is a temporary hack, proper resolution is tracked as UE-11222
+		bool bNeedsTesselationExtension = (FCStringAnsi::Strstr(GlslCodeOriginal.GetData(), "gl_MaxPatchVertices") != nullptr);
+		if (bNeedsTesselationExtension)
+		{
+			AppendCString(GlslCode, "#extension GL_ARB_tessellation_shader : enable\n");
+		}
+
 #if PLATFORM_ANDROID 
 		// Temporary patch to remove #extension GL_OES_standard_derivaties if not supported
 		if (!FOpenGL::SupportsStandardDerivativesExtension())
