Better settings for GL4 code path on Linux.

- Can still crash sometimes.

(Merging CL 2459155 from my branch).

diff --git a/Engine/Source/Runtime/OpenGLDrv/Private/OpenGLCommands.cpp b/Engine/Source/Runtime/OpenGLDrv/Private/OpenGLCommands.cpp
--- a/Engine/Source/Runtime/OpenGLDrv/Private/OpenGLCommands.cpp
+++ b/Engine/Source/Runtime/OpenGLDrv/Private/OpenGLCommands.cpp
@@ -28,7 +28,7 @@ DECLARE_ISBOUNDSHADER(DomainShader)
 
 namespace OpenGLConsoleVariables
 {
-#if PLATFORM_WINDOWS
+#if PLATFORM_WINDOWS || PLATFORM_LINUX
 	int32 bUseMapBuffer = 0;
 #else
 	int32 bUseMapBuffer = 1;
@@ -48,7 +48,7 @@ namespace OpenGLConsoleVariables
 		);
 
 	//@todo-rco: Workaround Nvidia driver crash
-	int32 bUseVAB = (PLATFORM_LINUX || (PLATFORM_DESKTOP && IsRHIDeviceNVIDIA())) ? 0 : 1;	// [RCL] @FIXME - currently (2014-09-19) causes crashes
+	int32 bUseVAB = (PLATFORM_DESKTOP && !PLATFORM_LINUX && IsRHIDeviceNVIDIA()) ? 0 : 1;
 	static FAutoConsoleVariableRef CVarUseVAB(
 		TEXT("OpenGL.UseVAB"),
 		bUseVAB,
@@ -56,7 +56,7 @@ namespace OpenGLConsoleVariables
 		ECVF_ReadOnly
 		);
 
-#if PLATFORM_WINDOWS
+#if PLATFORM_WINDOWS || PLATFORM_LINUX
 	int32 MaxSubDataSize = 256*1024;
 #else
 	int32 MaxSubDataSize = 0;
