Added another command-line option 'DisableMacDrawEvents' which will disable RHI/OpenGL draw event group markers on OS X if specified as they can cause Apple's OpenGL Profiler to crash.

diff --git a/Engine/Source/Runtime/OpenGLDrv/Private/Mac/OpenGLMac.cpp b/Engine/Source/Runtime/OpenGLDrv/Private/Mac/OpenGLMac.cpp
--- a/Engine/Source/Runtime/OpenGLDrv/Private/Mac/OpenGLMac.cpp
+++ b/Engine/Source/Runtime/OpenGLDrv/Private/Mac/OpenGLMac.cpp
@@ -1001,7 +1001,7 @@ void FMacOpenGL::ProcessExtensions(const FString& ExtensionsString)
 		glPushGroupMarkerEXT = (PFNGLPUSHGROUPMARKEREXTPROC)dlsym(RTLD_SELF, "glPushGroupMarkerEXT");
 		glPopGroupMarkerEXT = (PFNGLPOPGROUPMARKEREXTPROC)dlsym(RTLD_SELF, "glPopGroupMarkerEXT");
 #if !UE_BUILD_SHIPPING // For debuggable builds emit draw events when the extension is GL_EXT_debug_marker present.
-        GEmitDrawEvents = true;
+        GEmitDrawEvents = !FParse::Param(FCommandLine::Get(), TEXT("DisableMacDrawEvents"));
 #endif
 	}
 	
