PR #709 (Contributed by [REDACTED_PERSON])
Support for Quad Buffer Stereo on OpenGL; adds missing Custom Present on Linux
https://github.com/EpicGames/UnrealEngine/pull/709

diff --git a/Engine/Source/Runtime/OpenGLDrv/Private/Linux/OpenGLLinux.cpp b/Engine/Source/Runtime/OpenGLDrv/Private/Linux/OpenGLLinux.cpp
--- a/Engine/Source/Runtime/OpenGLDrv/Private/Linux/OpenGLLinux.cpp
+++ b/Engine/Source/Runtime/OpenGLDrv/Private/Linux/OpenGLLinux.cpp
@@ -378,6 +378,14 @@ bool PlatformBlitToViewport(FPlatformOpenGLDevice* Device,
 	{
 		FScopeContext ScopeContext( Context );
 
+		if (Viewport.GetCustomPresent())
+		{
+			glDisable(GL_FRAMEBUFFER_SRGB);
+			Viewport.GetCustomPresent()->Present(SyncInterval);
+			glEnable(GL_FRAMEBUFFER_SRGB);
+			return false;
+		}
+
 		glBindFramebuffer( GL_DRAW_FRAMEBUFFER, 0 );
 		glDrawBuffer( GL_BACK );
 		glBindFramebuffer( GL_READ_FRAMEBUFFER, Context->ViewportFramebuffer );
@@ -835,6 +843,15 @@ bool PlatformInitOpenGL()
 			UE_LOG(LogLinux, Fatal, TEXT("SDL error errno=%d (%s)"), SDL_GetError());
 		}
 
+		if (FParse::Param(FCommandLine::Get(), TEXT("quad_buffer_stereo")))
+		{
+			if (SDL_GL_SetAttribute(SDL_GL_STEREO, 1))
+			{
+				verifyf(false, TEXT("OpenGLRHI: %s\n."), SDL_GetError());
+				UE_LOG(LogLinux, Fatal, TEXT("SDL error errno=%d (%s)"), SDL_GetError());
+			}
+		}
+
 		// Create a dummy context to verify opengl support.
 		FPlatformOpenGLContext DummyContext;
 		_PlatformCreateDummyGLWindow(&DummyContext);
diff --git a/Engine/Source/Runtime/OpenGLDrv/Private/OpenGLDevice.cpp b/Engine/Source/Runtime/OpenGLDrv/Private/OpenGLDevice.cpp
--- a/Engine/Source/Runtime/OpenGLDrv/Private/OpenGLDevice.cpp
+++ b/Engine/Source/Runtime/OpenGLDrv/Private/OpenGLDevice.cpp
@@ -586,6 +586,16 @@ static void InitRHICapabilitiesForGL()
 	LOG_AND_GET_GL_INT_TEMP(GL_MAX_INTEGER_SAMPLES, 1);
 	LOG_AND_GET_GL_INT_TEMP(GL_MAX_COMBINED_TEXTURE_IMAGE_UNITS, 0);
 	LOG_AND_GET_GL_INT_TEMP(GL_MAX_VERTEX_ATTRIBS, 0);
+
+	if (FParse::Param(FCommandLine::Get(), TEXT("quad_buffer_stereo")))
+	{
+		GLboolean Result = GL_FALSE;
+		glGetBooleanv(GL_STEREO, &Result);
+		// Skip any errors if any were generated
+		glGetError();
+		GSupportsQuadBufferStereo = (Result == GL_TRUE);
+	}
+
 	if( FOpenGL::SupportsTextureFilterAnisotropic())
 	{
 		LOG_AND_GET_GL_INT_TEMP(GL_MAX_TEXTURE_MAX_ANISOTROPY_EXT, 0);
diff --git a/Engine/Source/Runtime/OpenGLDrv/Private/Windows/OpenGLWindows.cpp b/Engine/Source/Runtime/OpenGLDrv/Private/Windows/OpenGLWindows.cpp
--- a/Engine/Source/Runtime/OpenGLDrv/Private/Windows/OpenGLWindows.cpp
+++ b/Engine/Source/Runtime/OpenGLDrv/Private/Windows/OpenGLWindows.cpp
@@ -101,7 +101,7 @@ static LRESULT CALLBACK PlatformDummyGLWndproc(HWND hWnd, uint32 Message, WPARAM
 /**
  * Initialize a pixel format descriptor for the given window handle.
  */
-static void PlatformInitPixelFormatForDevice(HDC DeviceContext)
+static void PlatformInitPixelFormatForDevice(HDC DeviceContext, bool bTryIsDummyContext)
 {
 	// Pixel format descriptor for the context.
 	PIXELFORMATDESCRIPTOR PixelFormatDesc;
@@ -115,6 +115,12 @@ static void PlatformInitPixelFormatForDevice(HDC DeviceContext)
 	PixelFormatDesc.cStencilBits	= 0;
 	PixelFormatDesc.iLayerType	= PFD_MAIN_PLANE;
 
+	static bool bRequestedQuadBufferStereo = FParse::Param(FCommandLine::Get(), TEXT("quad_buffer_stereo"));
+	if (bRequestedQuadBufferStereo && !bTryIsDummyContext)
+	{
+		PixelFormatDesc.dwFlags = PixelFormatDesc.dwFlags | PFD_STEREO;
+	}
+
 	// Set the pixel format and create the context.
 	int32 PixelFormat = ChoosePixelFormat(DeviceContext, &PixelFormatDesc);
 	if (!PixelFormat || !SetPixelFormat(DeviceContext, PixelFormat, &PixelFormatDesc))
@@ -166,7 +172,7 @@ static void PlatformCreateDummyGLWindow(FPlatformOpenGLContext* OutContext)
 	// Get the device context.
 	OutContext->DeviceContext = GetDC(OutContext->WindowHandle);
 	check(OutContext->DeviceContext);
-	PlatformInitPixelFormatForDevice(OutContext->DeviceContext);
+	PlatformInitPixelFormatForDevice(OutContext->DeviceContext, true);
 }
 
 /**
@@ -348,7 +354,7 @@ FPlatformOpenGLContext* PlatformCreateOpenGLContext(FPlatformOpenGLDevice* Devic
 	Context->bReleaseWindowOnDestroy = false;
 	Context->DeviceContext = GetDC(Context->WindowHandle);
 	check(Context->DeviceContext);
-	PlatformInitPixelFormatForDevice(Context->DeviceContext);
+	PlatformInitPixelFormatForDevice(Context->DeviceContext, false);
 
 	int MajorVersion = 0;
 	int MinorVersion = 0;
diff --git a/Engine/Source/Runtime/RHI/Private/RHI.cpp b/Engine/Source/Runtime/RHI/Private/RHI.cpp
--- a/Engine/Source/Runtime/RHI/Private/RHI.cpp
+++ b/Engine/Source/Runtime/RHI/Private/RHI.cpp
@@ -159,6 +159,7 @@ namespace RHIConfig
 
 bool GIsRHIInitialized = false;
 int32 GMaxTextureMipCount = MAX_TEXTURE_MIP_COUNT;
+bool GSupportsQuadBufferStereo = false;
 bool GSupportsDepthFetchDuringDepthTest = true;
 FString GRHIAdapterName;
 uint32 GRHIVendorId = 0;
diff --git a/Engine/Source/Runtime/RHI/Public/RHI.h b/Engine/Source/Runtime/RHI/Public/RHI.h
--- a/Engine/Source/Runtime/RHI/Public/RHI.h
+++ b/Engine/Source/Runtime/RHI/Public/RHI.h
@@ -49,6 +49,9 @@ class RHI_API FRHICommandList;
 /** The maximum number of mip-maps that a texture can contain. 	*/
 extern	RHI_API int32		GMaxTextureMipCount;
 
+/** true if this platform has quad buffer stereo support. */
+extern RHI_API bool GSupportsQuadBufferStereo;
+
 /** true if the RHI supports textures that may be bound as both a render target and a shader resource. */
 extern RHI_API bool GSupportsRenderDepthTargetableShaderResources;
 
