UE4 - Mac compile fix

diff --git a/Engine/Source/Runtime/OpenGLDrv/Private/Mac/OpenGLMac.h b/Engine/Source/Runtime/OpenGLDrv/Private/Mac/OpenGLMac.h
--- a/Engine/Source/Runtime/OpenGLDrv/Private/Mac/OpenGLMac.h
+++ b/Engine/Source/Runtime/OpenGLDrv/Private/Mac/OpenGLMac.h
@@ -6,7 +6,11 @@
 #include <OpenGL/gl3ext.h>
 
 #ifndef GL_COMPUTE_SHADER
-#define GL_COMPUTE_SHADER 0x91B9
+#define GL_COMPUTE_SHADER	0x91B9
+#endif
+
+#ifndef GL_STEREO
+#define GL_STEREO			0x0C33
 #endif
 
 #include "OpenGL3.h"
