UE4 - GL_STEREO compile fix

diff --git a/Engine/Source/Runtime/OpenGLDrv/Private/Mac/OpenGLMac.h b/Engine/Source/Runtime/OpenGLDrv/Private/Mac/OpenGLMac.h
--- a/Engine/Source/Runtime/OpenGLDrv/Private/Mac/OpenGLMac.h
+++ b/Engine/Source/Runtime/OpenGLDrv/Private/Mac/OpenGLMac.h
@@ -9,10 +9,6 @@
 #define GL_COMPUTE_SHADER	0x91B9
 #endif
 
-#ifndef GL_STEREO
-#define GL_STEREO			0x0C33
-#endif
-
 #include "OpenGL3.h"
 
 
diff --git a/Engine/Source/Runtime/OpenGLDrv/Private/OpenGLDevice.cpp b/Engine/Source/Runtime/OpenGLDrv/Private/OpenGLDevice.cpp
--- a/Engine/Source/Runtime/OpenGLDrv/Private/OpenGLDevice.cpp
+++ b/Engine/Source/Runtime/OpenGLDrv/Private/OpenGLDevice.cpp
@@ -9,6 +9,10 @@
 #include "HardwareInfo.h"
 #include "ShaderCache.h"
 
+#ifndef GL_STEREO
+#define GL_STEREO			0x0C33
+#endif
+
 extern GLint GMaxOpenGLColorSamples;
 extern GLint GMaxOpenGLDepthSamples;
 extern GLint GMaxOpenGLIntegerSamples;
