Fix some CoreAnimation/OpenGL interaction problems that could cause corruption problems on OS X:
- In OpenGLMac::PlatformBlitToViewport we must blit from the viewport's source texture to an intermediate render buffer so that we can safely defer blitting to the window's context on the main thread, otherwise we risk trying to read from the viewport texture while it is being written by the render thread.
- In OpenGLMac::DrawOpenGLViewport we should only try and read from render buffers that have actually been rendered too.
- In OpenGLMac::DrawOpenGLViewport we must also flush the CoreAnimation layer's OpenGL context, otherwise OpenGL won't ever be told to flip the front & back buffers which may well cause problems for any device which relies upon that behaviour (i.e. VR headsets).

diff --git a/Engine/Source/Runtime/OpenGLDrv/Private/Mac/MacOpenGLContext.cpp b/Engine/Source/Runtime/OpenGLDrv/Private/Mac/MacOpenGLContext.cpp
--- a/Engine/Source/Runtime/OpenGLDrv/Private/Mac/MacOpenGLContext.cpp
+++ b/Engine/Source/Runtime/OpenGLDrv/Private/Mac/MacOpenGLContext.cpp
@@ -208,12 +208,14 @@ FPlatformOpenGLContext::FPlatformOpenGLContext()
 , EmulatedQueries(nullptr)
 , VertexArrayObject(0)
 , ViewportFramebuffer(0)
+, ViewportRenderbuffer(0)
 , SyncInterval(0)
 , VirtualScreen(0)
 , VendorID(0)
 , RendererID(0)
 , SupportsDepthFetchDuringDepthTest(true)
 {
+	FMemory::Memzero(ViewportSize);
 }
 
 FPlatformOpenGLContext::~FPlatformOpenGLContext()
@@ -228,6 +230,13 @@ FPlatformOpenGLContext::~FPlatformOpenGLContext()
 		MakeCurrent();
 		bMadeCurrent = true;
 	}
+	
+	if (ViewportRenderbuffer)
+	{
+		// this can be done from any context, as long as it's not nil.
+		glDeleteRenderbuffers(1, &ViewportRenderbuffer);
+		ViewportRenderbuffer = 0;
+	}
 
 	if (ViewportFramebuffer)
 	{
@@ -316,6 +325,15 @@ void FPlatformOpenGLContext::Reset(void)
 		MakeCurrent();
 		bMadeCurrent = true;
 	}
+	
+	FMemory::Memzero(ViewportSize);
+	
+	if (ViewportRenderbuffer)
+	{
+		// this can be done from any context, as long as it's not nil.
+		glDeleteRenderbuffers(1, &ViewportRenderbuffer);
+		ViewportRenderbuffer = 0;
+	}
 
 	if (ViewportFramebuffer)
 	{
diff --git a/Engine/Source/Runtime/OpenGLDrv/Private/Mac/MacOpenGLContext.h b/Engine/Source/Runtime/OpenGLDrv/Private/Mac/MacOpenGLContext.h
--- a/Engine/Source/Runtime/OpenGLDrv/Private/Mac/MacOpenGLContext.h
+++ b/Engine/Source/Runtime/OpenGLDrv/Private/Mac/MacOpenGLContext.h
@@ -39,6 +39,10 @@ struct FPlatformOpenGLContext
 	GLuint					VertexArrayObject;
 	/** Output Framebuffer */
 	GLuint					ViewportFramebuffer;
+	/** Output render buffer */
+	GLuint					ViewportRenderbuffer;
+	/** Cached viewport render buffer dimensions */
+	int32					ViewportSize[2];
 
 	// Context Settings
 	/** Flush Synchronisation Interval (V-Sync) */
diff --git a/Engine/Source/Runtime/OpenGLDrv/Private/Mac/OpenGLMac.cpp b/Engine/Source/Runtime/OpenGLDrv/Private/Mac/OpenGLMac.cpp
--- a/Engine/Source/Runtime/OpenGLDrv/Private/Mac/OpenGLMac.cpp
+++ b/Engine/Source/Runtime/OpenGLDrv/Private/Mac/OpenGLMac.cpp
@@ -404,6 +404,28 @@ bool PlatformBlitToViewport(FPlatformOpenGLDevice* Device, const FOpenGLViewport
 			}
 
 			[(FCocoaWindow*)Context->WindowHandle startRendering];
+			
+			[Context->OpenGLContext lock];
+			int32 CurrentReadFramebuffer = 0;
+			glGetIntegerv(GL_READ_FRAMEBUFFER_BINDING, &CurrentReadFramebuffer);
+			glBindFramebuffer(GL_DRAW_FRAMEBUFFER, Context->ViewportFramebuffer);
+			if ( Context->ViewportSize[0] != BackbufferSizeX || Context->ViewportSize[1] != BackbufferSizeY )
+			{
+				glBindRenderbuffer(GL_RENDERBUFFER, Context->ViewportRenderbuffer);
+				glRenderbufferStorage(GL_RENDERBUFFER, GL_RGBA8, BackbufferSizeX, BackbufferSizeY);
+				glFramebufferRenderbuffer(GL_DRAW_FRAMEBUFFER, GL_COLOR_ATTACHMENT1, GL_RENDERBUFFER, Context->ViewportRenderbuffer);
+				glBindRenderbuffer(GL_RENDERBUFFER, 0);
+			}
+			glDrawBuffer(GL_COLOR_ATTACHMENT1);
+			glBindFramebuffer(GL_READ_FRAMEBUFFER, Context->ViewportFramebuffer);
+			glReadBuffer(GL_COLOR_ATTACHMENT0);
+			glDisable(GL_FRAMEBUFFER_SRGB);
+			glBlitFramebuffer(0, 0, BackbufferSizeX, BackbufferSizeY, 0, 0, BackbufferSizeX, BackbufferSizeY, GL_COLOR_BUFFER_BIT, GL_NEAREST);
+			glEnable(GL_FRAMEBUFFER_SRGB);
+			glBindFramebuffer(GL_READ_FRAMEBUFFER, CurrentReadFramebuffer);
+			Context->ViewportSize[0] = BackbufferSizeX;
+			Context->ViewportSize[1] = BackbufferSizeY;
+			[Context->OpenGLContext unlock];
 
 			MainThreadCall(^{ [Context->OpenGLView setNeedsDisplay:YES]; }, NSDefaultRunLoopMode, false);
 
@@ -422,22 +444,26 @@ bool PlatformBlitToViewport(FPlatformOpenGLDevice* Device, const FOpenGLViewport
 
 void DrawOpenGLViewport(FPlatformOpenGLContext* const Context, uint32 Width, uint32 Height)
 {
-	if ([(FCocoaWindow*)Context->WindowHandle isRenderInitialized])
+	FCocoaWindow* Window = (FCocoaWindow*)Context->WindowHandle;
+	if ([Window isRenderInitialized] && Context->ViewportSize[0] && Context->ViewportSize[1] && Context->ViewportFramebuffer)
 	{
+		[Context->OpenGLContext lock];
 		int32 CurrentReadFramebuffer = 0;
 		glGetIntegerv(GL_READ_FRAMEBUFFER_BINDING, &CurrentReadFramebuffer);
 		glBindFramebuffer(GL_READ_FRAMEBUFFER, Context->ViewportFramebuffer);
-		glReadBuffer(GL_COLOR_ATTACHMENT0);
+		glReadBuffer(GL_COLOR_ATTACHMENT1);
 		glDisable(GL_FRAMEBUFFER_SRGB);
-		glBlitFramebuffer(0, 0, Width, Height, 0, Height, Width, 0, GL_COLOR_BUFFER_BIT, GL_NEAREST);
+		glBlitFramebuffer(0, 0, Context->ViewportSize[0], Context->ViewportSize[1], 0, Height, Width, 0, GL_COLOR_BUFFER_BIT, GL_NEAREST);
 		glEnable(GL_FRAMEBUFFER_SRGB);
 		glBindFramebuffer(GL_READ_FRAMEBUFFER, CurrentReadFramebuffer);
+		[Context->OpenGLContext unlock];
 	}
 	else
 	{
 		glClearColor(0.0f, 0.0f, 0.0f, 0.0f);
 		glClear(GL_COLOR_BUFFER_BIT);
 	}
+	[[NSOpenGLContext currentContext] flushBuffer];
 }
 
 void PlatformRenderingContextSetup(FPlatformOpenGLDevice* Device)
@@ -531,10 +557,16 @@ void PlatformResizeGLContext( FPlatformOpenGLDevice* Device, FPlatformOpenGLCont
 
 		if (Context->ViewportFramebuffer == 0)
 		{
-		glGenFramebuffers(1, &Context->ViewportFramebuffer);
-		check(Context->ViewportFramebuffer);
+			glGenFramebuffers(1, &Context->ViewportFramebuffer);
+			check(Context->ViewportFramebuffer);
 		}
-
+		
+		if (Context->ViewportRenderbuffer == 0)
+		{
+			glGenRenderbuffers(1, &Context->ViewportRenderbuffer);
+			check(Context->ViewportRenderbuffer);
+		}
+		
 		glBindFramebuffer(GL_FRAMEBUFFER, Context->ViewportFramebuffer);
 		glFramebufferTexture2D(GL_FRAMEBUFFER, GL_COLOR_ATTACHMENT0, BackBufferTarget, BackBufferResource, 0);
 #if UE_BUILD_DEBUG
