UE4 - OpenGL - Fix for memory overwrite
- Fixes UE-11903

diff --git a/Engine/Source/Runtime/OpenGLDrv/Private/OpenGLShaders.cpp b/Engine/Source/Runtime/OpenGLDrv/Private/OpenGLShaders.cpp
--- a/Engine/Source/Runtime/OpenGLDrv/Private/OpenGLShaders.cpp
+++ b/Engine/Source/Runtime/OpenGLDrv/Private/OpenGLShaders.cpp
@@ -1040,7 +1040,7 @@ public:
 	static inline void SortPackedUniformInfos(const TArray<FPackedUniformInfo>& ReflectedUniformInfos, const TArray<CrossCompiler::FPackedArrayInfo>& PackedGlobalArrays, TArray<FPackedUniformInfo>& OutPackedUniformInfos)
 	{
 		check(OutPackedUniformInfos.Num() == 0);
-		OutPackedUniformInfos.AddUninitialized(PackedGlobalArrays.Num());
+		OutPackedUniformInfos.Empty(PackedGlobalArrays.Num());
 		for (int32 Index = 0; Index < PackedGlobalArrays.Num(); ++Index)
 		{
 			auto& PackedArray = PackedGlobalArrays[Index];
@@ -1053,17 +1053,10 @@ public:
 				if (ReflectedInfo.ArrayType == PackedArray.TypeName)
 				{
 					FoundIndex = FindIndex;
-					OutPackedUniformInfos[Index] = ReflectedInfo;
+					OutPackedUniformInfos.Add(ReflectedInfo);
 					break;
 				}
 			}
-
-			if (FoundIndex == -1)
-			{
-				OutPackedUniformInfos[Index].Location = -1;
-				OutPackedUniformInfos[Index].ArrayType = -1;
-				OutPackedUniformInfos[Index].Index = -1;
-			}
 		}
 	}
 };
@@ -2626,7 +2619,8 @@ void FOpenGLShaderParameterCache::CommitPackedGlobals(const FOpenGLLinkedProgram
 	for (int32 PackedUniform = 0; PackedUniform < PackedUniforms.Num(); ++PackedUniform)
 	{
 		const FOpenGLLinkedProgram::FPackedUniformInfo& UniformInfo = PackedUniforms[PackedUniform];
-		const int32 ArrayIndex = UniformInfo.Index;
+		const uint32 ArrayIndex = UniformInfo.Index;
+		check(ArrayIndex < CrossCompiler::PACKED_TYPEINDEX_MAX);
 		const int32 NumVectors = PackedArrays[PackedUniform].Size / BytesPerRegister;
 		GLint Location = UniformInfo.Location;
 		const void* UniformData = PackedGlobalUniforms[ArrayIndex];
