UE4 - OpenGL - Fx crash when loading second map

diff --git a/Engine/Source/Runtime/OpenGLDrv/Private/OpenGLCommands.cpp b/Engine/Source/Runtime/OpenGLDrv/Private/OpenGLCommands.cpp
--- a/Engine/Source/Runtime/OpenGLDrv/Private/OpenGLCommands.cpp
+++ b/Engine/Source/Runtime/OpenGLDrv/Private/OpenGLCommands.cpp
@@ -53,8 +53,7 @@ namespace OpenGLConsoleVariables
 		TEXT("If true, don't issue dispatch work.")
 		);
 
-	//@todo-rco: Workaround Nvidia driver crash
-	int32 bUseVAB = (PLATFORM_DESKTOP && !PLATFORM_LINUX && IsRHIDeviceNVIDIA()) ? 0 : 1;
+	int32 bUseVAB = 1;
 	static FAutoConsoleVariableRef CVarUseVAB(
 		TEXT("OpenGL.UseVAB"),
 		bUseVAB,
diff --git a/Engine/Source/Runtime/OpenGLDrv/Private/OpenGLDrv.cpp b/Engine/Source/Runtime/OpenGLDrv/Private/OpenGLDrv.cpp
--- a/Engine/Source/Runtime/OpenGLDrv/Private/OpenGLDrv.cpp
+++ b/Engine/Source/Runtime/OpenGLDrv/Private/OpenGLDrv.cpp
@@ -487,11 +487,29 @@ void FOpenGLBase::ProcessExtensions( const FString& ExtensionsString )
 
 #if PLATFORM_WINDOWS || PLATFORM_LINUX
 	FString VendorName( ANSI_TO_TCHAR((const ANSICHAR*)glGetString(GL_VENDOR) ) );
-	if (VendorName.Contains(TEXT("ATI ")) || VendorName.Contains(TEXT("Intel ")))
+	if (VendorName.Contains(TEXT("ATI ")))
 	{
 		bAmdWorkaround = true;
+		GRHIVendorId = 0x1002;
+	}
+	else if (VendorName.Contains(TEXT("Intel ")))
+	{
+		bAmdWorkaround = true;
+		GRHIVendorId = 0x8086;
+	}
+	else if (VendorName.Contains(TEXT("NVIDIA ")))
+	{
+		GRHIVendorId = 0x10DE;
 	}
 #endif
+
+	// Setup CVars that require the RHI initialized
+
+	//@todo-rco: Workaround Nvidia driver crash
+	if (PLATFORM_DESKTOP && !PLATFORM_LINUX && IsRHIDeviceNVIDIA())
+	{
+		OpenGLConsoleVariables::bUseVAB = 0;
+	}
 }
 
 void GetExtensionsString( FString& ExtensionsString)
diff --git a/Engine/Source/Runtime/RHI/Private/RHI.cpp b/Engine/Source/Runtime/RHI/Private/RHI.cpp
--- a/Engine/Source/Runtime/RHI/Private/RHI.cpp
+++ b/Engine/Source/Runtime/RHI/Private/RHI.cpp
@@ -347,12 +347,14 @@ EShaderPlatform ShaderFormatToLegacyShaderPlatform(FName ShaderFormat)
 
 RHI_API bool IsRHIDeviceAMD()
 {
+	check(GRHIVendorId != 0);
 	// AMD's drivers tested on July 11 2013 have hitching problems with async resource streaming, setting single threaded for now until fixed.
 	return GRHIVendorId == 0x1002;
 }
 
 RHI_API bool IsRHIDeviceIntel()
 {
+	check(GRHIVendorId != 0);
 	// Intel GPUs are integrated and use both DedicatedVideoMemory and SharedSystemMemory.
 	// The hardware has fast clears so we disable exclude rects (see r.ClearWithExcludeRects)
 	return GRHIVendorId == 0x8086;
@@ -360,6 +362,7 @@ RHI_API bool IsRHIDeviceIntel()
 
 RHI_API bool IsRHIDeviceNVIDIA()
 {
+	check(GRHIVendorId != 0);
 	// NVIDIA GPUs are discrete and use DedicatedVideoMemory only.
 	return GRHIVendorId == 0x10DE;
 }
