UE4 - OpenGL - Fix mismatched new/delete

diff --git a/Engine/Source/Runtime/OpenGLDrv/Private/OpenGLShaders.cpp b/Engine/Source/Runtime/OpenGLDrv/Private/OpenGLShaders.cpp
--- a/Engine/Source/Runtime/OpenGLDrv/Private/OpenGLShaders.cpp
+++ b/Engine/Source/Runtime/OpenGLDrv/Private/OpenGLShaders.cpp
@@ -1780,7 +1780,7 @@ static FString GetShaderStageSource(TOpenGLStage* Shader)
 				ANSICHAR* Code = new ANSICHAR[Len + 1];
 				glGetShaderSource(Shaders[i], Len + 1, &Len, Code);
 				Source += Code;
-				delete Code;
+				delete [] Code;
 			}
 		}
 	}
