On OS X each CoreAnimation layer-backed NSView needs its own OpenGL context that is part of the OpenGL sharing group - we can't just use the shared context directly. The shared context is made current by other threads (typically Game & Render) so accessing it on Main can cause intermittent crashes as the context is swiped by another thread. We must also finish all rendering to the intermediate render buffer during PlatformBlitToViewport as there's no implicit synchronisation of resource access across contexts.

diff --git a/Engine/Source/Runtime/OpenGLDrv/Private/Mac/OpenGLMac.cpp b/Engine/Source/Runtime/OpenGLDrv/Private/Mac/OpenGLMac.cpp
--- a/Engine/Source/Runtime/OpenGLDrv/Private/Mac/OpenGLMac.cpp
+++ b/Engine/Source/Runtime/OpenGLDrv/Private/Mac/OpenGLMac.cpp
@@ -140,7 +140,7 @@ static void DrawOpenGLViewport(FPlatformOpenGLContext* const Context, uint32 Wid
 	self = [super init];
 	if (self)
 	{
-		self.Context = context;
+		self.Context = [[NSOpenGLContext alloc] initWithFormat:pixelFormat shareContext:context];
 		self.PixelFormat = pixelFormat;
 	}
 	return self;
@@ -425,6 +425,8 @@ bool PlatformBlitToViewport(FPlatformOpenGLDevice* Device, const FOpenGLViewport
 			glBindFramebuffer(GL_READ_FRAMEBUFFER, CurrentReadFramebuffer);
 			Context->ViewportSize[0] = BackbufferSizeX;
 			Context->ViewportSize[1] = BackbufferSizeY;
+			glFinishRenderAPPLE();	// Just in case the blit to screen happens immediately after unlock make sure that the command stream has been processed at this point.
+									// There's no synchronisation of resource access across contexts.
 			[Context->OpenGLContext unlock];
 
 			MainThreadCall(^{ [Context->OpenGLView setNeedsDisplay:YES]; }, NSDefaultRunLoopMode, false);
@@ -445,7 +447,7 @@ bool PlatformBlitToViewport(FPlatformOpenGLDevice* Device, const FOpenGLViewport
 void DrawOpenGLViewport(FPlatformOpenGLContext* const Context, uint32 Width, uint32 Height)
 {
 	FCocoaWindow* Window = (FCocoaWindow*)Context->WindowHandle;
-	if ([Window isRenderInitialized] && Context->ViewportSize[0] && Context->ViewportSize[1] && Context->ViewportFramebuffer)
+	if ([Window isRenderInitialized] && Context->ViewportSize[0] && Context->ViewportSize[1] && Context->ViewportFramebuffer && Context->ViewportRenderbuffer)
 	{
 		[Context->OpenGLContext lock];
 		int32 CurrentReadFramebuffer = 0;
