Thanks to [REDACTED_PERSON] finding & fixing the bug in the ResourcePool (CL 2458630) we don't need the workaround from CL 2420476 to invalidate the buffer range on unsynchronised buffer writes.

diff --git a/Engine/Source/Runtime/OpenGLDrv/Private/Mac/OpenGLMac.h b/Engine/Source/Runtime/OpenGLDrv/Private/Mac/OpenGLMac.h
--- a/Engine/Source/Runtime/OpenGLDrv/Private/Mac/OpenGLMac.h
+++ b/Engine/Source/Runtime/OpenGLDrv/Private/Mac/OpenGLMac.h
@@ -89,8 +89,8 @@ public:
 			case RLM_WriteOnly:
 				Access = (GL_MAP_INVALIDATE_BUFFER_BIT | GL_MAP_WRITE_BIT);
 				break;
-			case RLM_WriteOnlyUnsynchronized: // Invalidate range is only required to avoid UE-7915 and is a workaround, not a fix.
-				Access = (GL_MAP_WRITE_BIT | GL_MAP_UNSYNCHRONIZED_BIT | GL_MAP_INVALIDATE_RANGE_BIT);
+			case RLM_WriteOnlyUnsynchronized:
+				Access = (GL_MAP_WRITE_BIT | GL_MAP_UNSYNCHRONIZED_BIT);
 				break;
 			case RLM_WriteOnlyPersistent:
 				Access = (GL_MAP_WRITE_BIT | GL_MAP_PERSISTENT_BIT | GL_MAP_COHERENT_BIT);
