Remove the glFinishRenderAPPLE from PlatformBlitToViewport as it does not appear to be necessary.

diff --git a/Engine/Source/Runtime/OpenGLDrv/Private/Mac/OpenGLMac.cpp b/Engine/Source/Runtime/OpenGLDrv/Private/Mac/OpenGLMac.cpp
--- a/Engine/Source/Runtime/OpenGLDrv/Private/Mac/OpenGLMac.cpp
+++ b/Engine/Source/Runtime/OpenGLDrv/Private/Mac/OpenGLMac.cpp
@@ -425,8 +425,6 @@ bool PlatformBlitToViewport(FPlatformOpenGLDevice* Device, const FOpenGLViewport
 			glBindFramebuffer(GL_READ_FRAMEBUFFER, CurrentReadFramebuffer);
 			Context->ViewportSize[0] = BackbufferSizeX;
 			Context->ViewportSize[1] = BackbufferSizeY;
-			glFinishRenderAPPLE();	// Just in case the blit to screen happens immediately after unlock make sure that the command stream has been processed at this point.
-									// There's no synchronisation of resource access across contexts.
 			[Context->OpenGLContext unlock];
 
 			MainThreadCall(^{ [Context->OpenGLView setNeedsDisplay:YES]; }, NSDefaultRunLoopMode, false);
