On OS X make the CoreAnimation layer share the viewport's context by ensuring that access to the context is restricted to one thread at a time using the NSOpenGLContext lock/unlock operations. This ensures that only one thread can have this context current & be calling OpenGL using it at any given time.

diff --git a/Engine/Source/Runtime/OpenGLDrv/Private/Mac/OpenGLMac.cpp b/Engine/Source/Runtime/OpenGLDrv/Private/Mac/OpenGLMac.cpp
--- a/Engine/Source/Runtime/OpenGLDrv/Private/Mac/OpenGLMac.cpp
+++ b/Engine/Source/Runtime/OpenGLDrv/Private/Mac/OpenGLMac.cpp
@@ -140,12 +140,22 @@ static void DrawOpenGLViewport(FPlatformOpenGLContext* const Context, uint32 Wid
 	self = [super init];
 	if (self)
 	{
-		self.Context = [[NSOpenGLContext alloc] initWithFormat:pixelFormat shareContext:context];
+		self.Context = context;
 		self.PixelFormat = pixelFormat;
 	}
 	return self;
 }
 
+- (BOOL)canDrawInOpenGLContext:(NSOpenGLContext *)context pixelFormat:(NSOpenGLPixelFormat *)pixelFormat forLayerTime:(CFTimeInterval)timeInterval displayTime:(const CVTimeStamp *)timeStamp
+{
+	BOOL bOK = [super canDrawInOpenGLContext:context pixelFormat:pixelFormat forLayerTime:timeInterval displayTime:timeStamp];
+	if ( bOK && context && (self.Context == context) )
+	{
+		[context lock];
+	}
+	return bOK;
+}
+
 @end
 
 /**
@@ -384,6 +394,7 @@ bool PlatformBlitToViewport(FPlatformOpenGLDevice* Device, const FOpenGLViewport
 	check(Context && Context->OpenGLView);
 
 	FScopeLock ScopeLock(Device->ContextUsageGuard);
+	[Context->OpenGLContext lock];
 	{
 		FScopeContext ScopeContext(Context->OpenGLContext);
 
@@ -405,7 +416,6 @@ bool PlatformBlitToViewport(FPlatformOpenGLDevice* Device, const FOpenGLViewport
 
 			[(FCocoaWindow*)Context->WindowHandle startRendering];
 			
-			[Context->OpenGLContext lock];
 			int32 CurrentReadFramebuffer = 0;
 			glGetIntegerv(GL_READ_FRAMEBUFFER_BINDING, &CurrentReadFramebuffer);
 			glBindFramebuffer(GL_DRAW_FRAMEBUFFER, Context->ViewportFramebuffer);
@@ -425,7 +435,6 @@ bool PlatformBlitToViewport(FPlatformOpenGLDevice* Device, const FOpenGLViewport
 			glBindFramebuffer(GL_READ_FRAMEBUFFER, CurrentReadFramebuffer);
 			Context->ViewportSize[0] = BackbufferSizeX;
 			Context->ViewportSize[1] = BackbufferSizeY;
-			[Context->OpenGLContext unlock];
 
 			MainThreadCall(^{ [Context->OpenGLView setNeedsDisplay:YES]; }, NSDefaultRunLoopMode, false);
 
@@ -438,6 +447,7 @@ bool PlatformBlitToViewport(FPlatformOpenGLDevice* Device, const FOpenGLViewport
 			}
 		}
 	}
+	[Context->OpenGLContext unlock];
 
 	return !Viewport.GetCustomPresent();
 }
@@ -447,7 +457,6 @@ void DrawOpenGLViewport(FPlatformOpenGLContext* const Context, uint32 Width, uin
 	FCocoaWindow* Window = (FCocoaWindow*)Context->WindowHandle;
 	if ([Window isRenderInitialized] && Context->ViewportSize[0] && Context->ViewportSize[1] && Context->ViewportFramebuffer && Context->ViewportRenderbuffer)
 	{
-		[Context->OpenGLContext lock];
 		int32 CurrentReadFramebuffer = 0;
 		glGetIntegerv(GL_READ_FRAMEBUFFER_BINDING, &CurrentReadFramebuffer);
 		glBindFramebuffer(GL_READ_FRAMEBUFFER, Context->ViewportFramebuffer);
@@ -456,14 +465,14 @@ void DrawOpenGLViewport(FPlatformOpenGLContext* const Context, uint32 Width, uin
 		glBlitFramebuffer(0, 0, Context->ViewportSize[0], Context->ViewportSize[1], 0, Height, Width, 0, GL_COLOR_BUFFER_BIT, GL_NEAREST);
 		glEnable(GL_FRAMEBUFFER_SRGB);
 		glBindFramebuffer(GL_READ_FRAMEBUFFER, CurrentReadFramebuffer);
-		[Context->OpenGLContext unlock];
 	}
 	else
 	{
 		glClearColor(0.0f, 0.0f, 0.0f, 0.0f);
 		glClear(GL_COLOR_BUFFER_BIT);
 	}
-	[[NSOpenGLContext currentContext] flushBuffer];
+	[Context->OpenGLContext flushBuffer];
+	[Context->OpenGLContext unlock];
 }
 
 void PlatformRenderingContextSetup(FPlatformOpenGLDevice* Device)
@@ -528,6 +537,7 @@ void PlatformRebindResources(FPlatformOpenGLDevice* Device)
 void PlatformResizeGLContext( FPlatformOpenGLDevice* Device, FPlatformOpenGLContext* Context, uint32 SizeX, uint32 SizeY, bool bFullscreen, bool bWasFullscreen, GLenum BackBufferTarget, GLuint BackBufferResource)
 {
 	FScopeLock ScopeLock(Device->ContextUsageGuard);
+	[Context->OpenGLContext lock];
 	{
 		FScopeContext ScopeContext(Context->OpenGLContext);
 
@@ -579,6 +589,7 @@ void PlatformResizeGLContext( FPlatformOpenGLDevice* Device, FPlatformOpenGLCont
 		}
 #endif
 	}
+	[Context->OpenGLContext unlock];
 }
 
 void PlatformGetSupportedResolution(uint32 &Width, uint32 &Height)
