Remove avoidable temporary string construction in OpenGL uniform look-up. This happens a lot more often with SSOs as we must rebind uniforms for SSOs as shaders will have different bindings for each of the multiple bound-shader-states they belong too.

diff --git a/Engine/Source/Runtime/OpenGLDrv/Private/OpenGLShaders.cpp b/Engine/Source/Runtime/OpenGLDrv/Private/OpenGLShaders.cpp
--- a/Engine/Source/Runtime/OpenGLDrv/Private/OpenGLShaders.cpp
+++ b/Engine/Source/Runtime/OpenGLDrv/Private/OpenGLShaders.cpp
@@ -865,9 +865,29 @@ void FOpenGLDynamicRHI::BindUniformBufferBase(FOpenGLContextState& ContextState,
 
 // ============================================================================================================================
 
-static TMap<GLuint, TMap<FAnsiCharArray, int64>>& GetOpenGLUniformBlockLocations()
+struct FOpenGLUniformName
 {
-	static TMap<GLuint, TMap<FAnsiCharArray, int64>> UniformBlockLocations;
+	FOpenGLUniformName()
+	{
+		FMemory::Memzero(Buffer);
+	}
+	
+	ANSICHAR Buffer[10];
+	
+	friend bool operator ==(const FOpenGLUniformName& A,const FOpenGLUniformName& B)
+	{
+		return FMemory::Memcmp(A.Buffer, B.Buffer, sizeof(A.Buffer)) == 0;
+	}
+	
+	friend uint32 GetTypeHash(const FOpenGLUniformName &Key)
+	{
+		return FCrc::MemCrc32(Key.Buffer, sizeof(Key.Buffer));
+	}
+};
+
+static TMap<GLuint, TMap<FOpenGLUniformName, int64>>& GetOpenGLUniformBlockLocations()
+{
+	static TMap<GLuint, TMap<FOpenGLUniformName, int64>> UniformBlockLocations;
 	return UniformBlockLocations;
 }
 
@@ -877,20 +897,18 @@ static TMap<GLuint, TMap<int64, int64>>& GetOpenGLUniformBlockBindings()
 	return UniformBlockBindings;
 }
 
-static GLuint GetOpenGLProgramUniformBlockIndex(GLuint Program, const GLchar* UniformBlockName)
+static GLuint GetOpenGLProgramUniformBlockIndex(GLuint Program, const FOpenGLUniformName& UniformBlockName)
 {
-	TMap<FAnsiCharArray, int64>& Locations = GetOpenGLUniformBlockLocations().FindOrAdd(Program);
-	FAnsiCharArray Name;
-	Name.Append(UniformBlockName, FCStringAnsi::Strlen(UniformBlockName) + 1);
-	int64* Location = Locations.Find(Name);
+	TMap<FOpenGLUniformName, int64>& Locations = GetOpenGLUniformBlockLocations().FindOrAdd(Program);
+	int64* Location = Locations.Find(UniformBlockName);
 	if(Location)
 	{
 		return *Location;
 	}
 	else
 	{
-		int64& Loc = Locations.Emplace(Name);
-		Loc = (int64)FOpenGL::GetUniformBlockIndex(Program, UniformBlockName);
+		int64& Loc = Locations.Emplace(UniformBlockName);
+		Loc = (int64)FOpenGL::GetUniformBlockIndex(Program, UniformBlockName.Buffer);
 		return Loc;
 	}
 }
@@ -1087,19 +1105,16 @@ void FOpenGLLinkedProgram::VerifyUniformBlockBindings( int Stage, uint32 FirstUn
 	if ( FOpenGL::SupportsSeparateShaderObjects() && FOpenGL::SupportsUniformBuffers() )
 	{
 		static const ANSICHAR StagePrefix[CrossCompiler::NUM_SHADER_STAGES] = { 'v', 'p', 'g', 'h', 'd', 'c'};
-		ANSICHAR NameBuffer[10] = {0};
-		NameBuffer[0] = StagePrefix[Stage];
-		NameBuffer[1] = 'b';
-		NameBuffer[2] = 0;
-		NameBuffer[3] = 0;
-		NameBuffer[4] = 0;
+		FOpenGLUniformName Name;
+		Name.Buffer[0] = StagePrefix[Stage];
+		Name.Buffer[1] = 'b';
 		
 		GLuint StageProgram = Config.Shaders[Stage].Resource;
 
 		for (int32 BufferIndex = 0; BufferIndex < Config.Shaders[Stage].Bindings.NumUniformBuffers; ++BufferIndex)
 		{
-			SetIndex(NameBuffer, 2, BufferIndex);
-			GLint Location = GetOpenGLProgramUniformBlockIndex(StageProgram, NameBuffer);
+			SetIndex(Name.Buffer, 2, BufferIndex);
+			GLint Location = GetOpenGLProgramUniformBlockIndex(StageProgram, Name);
 			if (Location >= 0)
 			{
 				GetOpenGLProgramUniformBlockBinding(StageProgram, Location, FirstUniformBuffer + BufferIndex);
@@ -1129,31 +1144,31 @@ void FOpenGLLinkedProgram::ConfigureShaderStage( int Stage, uint32 FirstUniformB
 		OGL_UAV_NOT_SUPPORTED_FOR_GRAPHICS_UNIT,
 		FOpenGL::GetFirstComputeUAVUnit()
 	};
-	ANSICHAR NameBuffer[10] = {0};
-
+	
 	// verify that only CS uses UAVs
 	check((Stage != CrossCompiler::SHADER_STAGE_COMPUTE) ? (CountSetBits(UAVStageNeeds) == 0) : true);
 
 	SCOPE_CYCLE_COUNTER(STAT_OpenGLShaderBindParameterTime);
 	VERIFY_GL_SCOPE();
 
-	NameBuffer[0] = StagePrefix[Stage];
+	FOpenGLUniformName Name;
+	Name.Buffer[0] = StagePrefix[Stage];
 
 	GLuint StageProgram = FOpenGL::SupportsSeparateShaderObjects() ? Config.Shaders[Stage].Resource : Program;
 	
 	// Bind Global uniform arrays (vu_h, pu_i, etc)
 	{
-		NameBuffer[1] = 'u';
-		NameBuffer[2] = '_';
-		NameBuffer[3] = 0;
-		NameBuffer[4] = 0;
+		Name.Buffer[1] = 'u';
+		Name.Buffer[2] = '_';
+		Name.Buffer[3] = 0;
+		Name.Buffer[4] = 0;
 
 		TArray<FPackedUniformInfo> PackedUniformInfos;
 		for (uint8 Index = 0; Index < CrossCompiler::PACKED_TYPEINDEX_MAX; ++Index)
 		{
 			uint8 ArrayIndexType = CrossCompiler::PackedTypeIndexToTypeName(Index);
-			NameBuffer[3] = ArrayIndexType;
-			GLint Location = glGetUniformLocation(StageProgram, NameBuffer);
+			Name.Buffer[3] = ArrayIndexType;
+			GLint Location = glGetUniformLocation(StageProgram, Name.Buffer);
 			if ((int32)Location != -1)
 			{
 				FPackedUniformInfo Info = {Location, ArrayIndexType, Index};
@@ -1166,23 +1181,23 @@ void FOpenGLLinkedProgram::ConfigureShaderStage( int Stage, uint32 FirstUniformB
 
 	// Bind uniform buffer packed arrays (vc0_h, pc2_i, etc)
 	{
-		NameBuffer[1] = 'c';
-		NameBuffer[2] = 0;
-		NameBuffer[3] = 0;
-		NameBuffer[4] = 0;
-		NameBuffer[5] = 0;
-		NameBuffer[6] = 0;
+		Name.Buffer[1] = 'c';
+		Name.Buffer[2] = 0;
+		Name.Buffer[3] = 0;
+		Name.Buffer[4] = 0;
+		Name.Buffer[5] = 0;
+		Name.Buffer[6] = 0;
 		for (uint8 UB = 0; UB < Config.Shaders[Stage].Bindings.NumUniformBuffers; ++UB)
 		{
 			TArray<FPackedUniformInfo> PackedBuffers;
-			ANSICHAR* Str = SetIndex(NameBuffer, 2, UB);
+			ANSICHAR* Str = SetIndex(Name.Buffer, 2, UB);
 			*Str++ = '_';
 			Str[1] = 0;
 			for (uint8 Index = 0; Index < CrossCompiler::PACKED_TYPEINDEX_MAX; ++Index)
 			{
 				uint8 ArrayIndexType = CrossCompiler::PackedTypeIndexToTypeName(Index);
 				Str[0] = ArrayIndexType;
-				GLint Location = glGetUniformLocation(StageProgram, NameBuffer);
+				GLint Location = glGetUniformLocation(StageProgram, Name.Buffer);
 				if ((int32)Location != -1)
 				{
 					FPackedUniformInfo Info = {Location, ArrayIndexType, Index};
@@ -1199,28 +1214,28 @@ void FOpenGLLinkedProgram::ConfigureShaderStage( int Stage, uint32 FirstUniformB
 	StagePackedUniformInfo[Stage].LastEmulatedUniformBufferSet.AddZeroed(Config.Shaders[Stage].Bindings.NumUniformBuffers);
 
 	// Bind samplers.
-	NameBuffer[1] = 's';
-	NameBuffer[2] = 0;
-	NameBuffer[3] = 0;
-	NameBuffer[4] = 0;
+	Name.Buffer[1] = 's';
+	Name.Buffer[2] = 0;
+	Name.Buffer[3] = 0;
+	Name.Buffer[4] = 0;
 	int32 LastFoundIndex = -1;
 	for (int32 SamplerIndex = 0; SamplerIndex < Config.Shaders[Stage].Bindings.NumSamplers; ++SamplerIndex)
 	{
-		SetIndex(NameBuffer, 2, SamplerIndex);
-		GLint Location = glGetUniformLocation(StageProgram, NameBuffer);
+		SetIndex(Name.Buffer, 2, SamplerIndex);
+		GLint Location = glGetUniformLocation(StageProgram, Name.Buffer);
 		if (Location == -1)
 		{
 			if (LastFoundIndex != -1)
 			{
 				// It may be an array of samplers. Get the initial element location, if available, and count from it.
-				SetIndex(NameBuffer, 2, LastFoundIndex);
+				SetIndex(Name.Buffer, 2, LastFoundIndex);
 				int32 OffsetOfArraySpecifier = (LastFoundIndex>9)?4:3;
 				int32 ArrayIndex = SamplerIndex-LastFoundIndex;
-				NameBuffer[OffsetOfArraySpecifier] = '[';
-				ANSICHAR* EndBracket = SetIndex(NameBuffer, OffsetOfArraySpecifier+1, ArrayIndex);
+				Name.Buffer[OffsetOfArraySpecifier] = '[';
+				ANSICHAR* EndBracket = SetIndex(Name.Buffer, OffsetOfArraySpecifier+1, ArrayIndex);
 				*EndBracket++ = ']';
 				*EndBracket = 0;
-				Location = glGetUniformLocation(StageProgram, NameBuffer);
+				Location = glGetUniformLocation(StageProgram, Name.Buffer);
 			}
 		}
 		else
@@ -1249,28 +1264,28 @@ void FOpenGLLinkedProgram::ConfigureShaderStage( int Stage, uint32 FirstUniformB
 	}
 
 	// Bind UAVs/images.
-	NameBuffer[1] = 'i';
-	NameBuffer[2] = 0;
-	NameBuffer[3] = 0;
-	NameBuffer[4] = 0;
+	Name.Buffer[1] = 'i';
+	Name.Buffer[2] = 0;
+	Name.Buffer[3] = 0;
+	Name.Buffer[4] = 0;
 	int32 LastFoundUAVIndex = -1;
 	for (int32 UAVIndex = 0; UAVIndex < Config.Shaders[Stage].Bindings.NumUAVs; ++UAVIndex)
 	{
-		SetIndex(NameBuffer, 2, UAVIndex);
-		GLint Location = glGetUniformLocation(StageProgram, NameBuffer);
+		SetIndex(Name.Buffer, 2, UAVIndex);
+		GLint Location = glGetUniformLocation(StageProgram, Name.Buffer);
 		if (Location == -1)
 		{
 			if (LastFoundUAVIndex != -1)
 			{
 				// It may be an array of UAVs. Get the initial element location, if available, and count from it.
-				SetIndex(NameBuffer, 2, LastFoundUAVIndex);
+				SetIndex(Name.Buffer, 2, LastFoundUAVIndex);
 				int32 OffsetOfArraySpecifier = (LastFoundUAVIndex>9)?4:3;
 				int32 ArrayIndex = UAVIndex-LastFoundUAVIndex;
-				NameBuffer[OffsetOfArraySpecifier] = '[';
-				ANSICHAR* EndBracket = SetIndex(NameBuffer, OffsetOfArraySpecifier+1, ArrayIndex);
+				Name.Buffer[OffsetOfArraySpecifier] = '[';
+				ANSICHAR* EndBracket = SetIndex(Name.Buffer, OffsetOfArraySpecifier+1, ArrayIndex);
 				*EndBracket++ = ']';
 				*EndBracket = '\0';
-				Location = glGetUniformLocation(StageProgram, NameBuffer);
+				Location = glGetUniformLocation(StageProgram, Name.Buffer);
 			}
 		}
 		else
@@ -1290,14 +1305,14 @@ void FOpenGLLinkedProgram::ConfigureShaderStage( int Stage, uint32 FirstUniformB
 	// Bind uniform buffers.
 	if (FOpenGL::SupportsUniformBuffers())
 	{
-		NameBuffer[1] = 'b';
-		NameBuffer[2] = 0;
-		NameBuffer[3] = 0;
-		NameBuffer[4] = 0;
+		Name.Buffer[1] = 'b';
+		Name.Buffer[2] = 0;
+		Name.Buffer[3] = 0;
+		Name.Buffer[4] = 0;
 		for (int32 BufferIndex = 0; BufferIndex < Config.Shaders[Stage].Bindings.NumUniformBuffers; ++BufferIndex)
 		{
-			SetIndex(NameBuffer, 2, BufferIndex);
-			GLint Location = GetOpenGLProgramUniformBlockIndex(StageProgram, NameBuffer);
+			SetIndex(Name.Buffer, 2, BufferIndex);
+			GLint Location = GetOpenGLProgramUniformBlockIndex(StageProgram, Name);
 			if (Location >= 0)
 			{
 				GetOpenGLProgramUniformBlockBinding(StageProgram, Location, FirstUniformBuffer + BufferIndex);
