Cache the GL_QUERY objects in the Mac timestamp query emulation code to prevent some drivers from unmapping the query buffer during outstanding rendering & crashing. This allows us to enable this by default in non-shipping builds (with a small performance cost) - the Mac will now report GPU timing data but the accuracy will be lower than other platforms.

diff --git a/Engine/Source/Runtime/OpenGLDrv/Private/Mac/MacOpenGLQuery.cpp b/Engine/Source/Runtime/OpenGLDrv/Private/Mac/MacOpenGLQuery.cpp
--- a/Engine/Source/Runtime/OpenGLDrv/Private/Mac/MacOpenGLQuery.cpp
+++ b/Engine/Source/Runtime/OpenGLDrv/Private/Mac/MacOpenGLQuery.cpp
@@ -8,21 +8,36 @@
  OpenGL static variables.
  ------------------------------------------------------------------------------*/
 
-bool GIsEmulatingTimestamp = false; // @todo: Now crashing on Nvidia cards, but not on AMD...
-static int32 GNVDelayEmulatingTimeStamp = 60;
-static FAutoConsoleVariableRef CVarNVDelayEmulatingTimeStamp(
-	  TEXT("r.Mac.Nvidia.DelayEmulatingTimeStamp"),
-	  GNVDelayEmulatingTimeStamp,
-	  TEXT("The number of frames to defer emulating GL_TIMESTAMP on Nvidia cards under Mac OS X to avoid crashes on startup, set to 0 to disable. (Default: 60)"),
-	  ECVF_RenderThreadSafe
-	  );
+bool GIsEmulatingTimestamp = false;
+static const uint32 GMacQueryNameCacheSize = 32 * OPENGL_NAME_CACHE_SIZE;
+
+char const* const GMacTimerVertShader4 =	"#version 150\n"
+											"const int VertexCount = 4;\n"
+											"const vec2 Position[VertexCount] = vec2[](	vec2(-1.0,-1.0),\n"
+											"											vec2( 1.0,-1.0),\n"
+											"											vec2(-1.0, 1.0),\n"
+											"											vec2( 1.0, 1.0));\n"
+											"void main()\n"
+											"{\n"
+											"	gl_Position = vec4(Position[gl_VertexID], 0.0, 1.0);\n"
+											"}\n";
+
+char const* const GMacTimerFragShader =	"#version 150\n"
+										"out vec4 Color;\n"
+										"void main()\n"
+										"{\n"
+										"	Color = vec4(0.0, 0.0, 0.0, 1.0);\n"
+										"}\n";
+
+GLuint FMacOpenGLQueryEmu::DummyDrawProgram = 0;
 
 /*------------------------------------------------------------------------------
  OpenGL query emulation.
  ------------------------------------------------------------------------------*/
 
-FMacOpenGLTimer::FMacOpenGLTimer(FPlatformOpenGLContext* InContext)
+FMacOpenGLTimer::FMacOpenGLTimer(FPlatformOpenGLContext* InContext, FMacOpenGLQueryEmu* InEmu)
 : Context(InContext)
+, Emu(InEmu)
 , Name(0)
 , Result(0)
 , Accumulated(0)
@@ -31,10 +46,19 @@ FMacOpenGLTimer::FMacOpenGLTimer(FPlatformOpenGLContext* InContext)
 , Running(false)
 {
 	check(Context);
-	if (!IsRHIDeviceNVIDIA() || GFrameNumberRenderThread > GNVDelayEmulatingTimeStamp)
+	check(Emu);
+	
+	if(!Emu->FreeQueries.FindOrAdd(GL_TIME_ELAPSED).Num())
 	{
-		glGenQueries(1, &Name);
+		if(IsRHIDeviceNVIDIA())
+		{
+			glFinish();
+		}
+		
+		Emu->FreeQueries.FindOrAdd(GL_TIME_ELAPSED).AddUninitialized(GMacQueryNameCacheSize);
+		glGenQueries(GMacQueryNameCacheSize, Emu->FreeQueries.FindOrAdd(GL_TIME_ELAPSED).GetData());
 	}
+	Name = Emu->FreeQueries.FindOrAdd(GL_TIME_ELAPSED).Pop();
 }
 
 FMacOpenGLTimer::~FMacOpenGLTimer()
@@ -43,7 +67,8 @@ FMacOpenGLTimer::~FMacOpenGLTimer()
 	Next.Reset();
 	if (Name)
 	{
-		glDeleteQueries(1, &Name);
+		Emu->FreeQueries.FindOrAdd(GL_TIME_ELAPSED).Push(Name);
+		Name = 0;
 	}
 }
 
@@ -51,10 +76,7 @@ void FMacOpenGLTimer::Begin(void)
 {
 	if(!Running)
 	{
-		if (!IsRHIDeviceNVIDIA() || GFrameNumberRenderThread > GNVDelayEmulatingTimeStamp)
-		{
-			glBeginQuery(GL_TIME_ELAPSED, Name);
-		}
+		glBeginQuery(GL_TIME_ELAPSED, Name);
 		Running = true;
 	}
 }
@@ -64,10 +86,7 @@ void FMacOpenGLTimer::End(void)
 	if(Running)
 	{
 		Running = false;
-		if (!IsRHIDeviceNVIDIA() || GFrameNumberRenderThread > GNVDelayEmulatingTimeStamp)
-		{
-			glEndQuery(GL_TIME_ELAPSED);
-		}
+		glEndQuery(GL_TIME_ELAPSED);
 	}
 }
 
@@ -90,14 +109,7 @@ int32 FMacOpenGLTimer::GetResultAvailable(void)
 	check(!Running);
 	if(!Available)
 	{
-		if (!IsRHIDeviceNVIDIA() || GFrameNumberRenderThread > GNVDelayEmulatingTimeStamp)
-		{
-			glGetQueryObjectuiv(Name, GL_QUERY_RESULT_AVAILABLE, &Available);
-		}
-		else
-		{
-			Available = true;
-		}
+		glGetQueryObjectuiv(Name, GL_QUERY_RESULT_AVAILABLE, &Available);
 	}
 	return Available;
 }
@@ -107,14 +119,7 @@ void FMacOpenGLTimer::CacheResult(void)
 	check(!Running);
 	if(!Cached)
 	{
-		if (!IsRHIDeviceNVIDIA() || GFrameNumberRenderThread > GNVDelayEmulatingTimeStamp)
-		{
-			glGetQueryObjectui64v(Name, GL_QUERY_RESULT, &Result);
-		}
-		else
-		{
-			Result = 0;
-		}
+		glGetQueryObjectui64v(Name, GL_QUERY_RESULT, &Result);
 		
 		if(Previous.IsValid())
 		{
@@ -130,16 +135,18 @@ void FMacOpenGLTimer::CacheResult(void)
 	}
 }
 
-FMacOpenGLQuery::FMacOpenGLQuery(FPlatformOpenGLContext* InContext)
+FMacOpenGLQuery::FMacOpenGLQuery(FPlatformOpenGLContext* InContext, FMacOpenGLQueryEmu* InEmu)
 	: Name(0)
 	, Target(0)
 	, Context(InContext)
+	, Emu(InEmu)
 	, Result(0)
 	, Available(0)
 	, Running(0)
 	, Cached(false)
 {
 	check(Context);
+	check(Emu);
 }
 
 FMacOpenGLQuery::~FMacOpenGLQuery()
@@ -150,7 +157,21 @@ FMacOpenGLQuery::~FMacOpenGLQuery()
 	}
 	if(Name)
 	{
-		glDeleteQueries(1, &Name);
+		switch(Target)
+		{
+			case GL_TIMESTAMP:
+			case GL_TIME_ELAPSED:
+			{
+				Emu->FreeQueries.FindOrAdd(GL_TIME_ELAPSED).Push(Name);
+				break;
+			}
+			default:
+			{
+				Emu->FreeQueries.FindOrAdd(Target).Push(Name);
+				break;
+			}
+		}
+		Name = 0;
 	}
 }
 
@@ -170,7 +191,7 @@ void FMacOpenGLQuery::Begin(GLenum InTarget)
 			case GL_TIMESTAMP:
 			{
 				// There can be a lot of timestamps, clear emulated queries out to avoid problems.
-				for( auto It : Context->EmulatedQueries->Queries )
+				for( auto It : Emu->Queries )
 				{
 					if(It.IsValid() && It->Running == false && (It->Target == GL_TIMESTAMP || It->Target == GL_TIME_ELAPSED))
 					{
@@ -185,36 +206,46 @@ void FMacOpenGLQuery::Begin(GLenum InTarget)
 					}
 				}
 				
-				check(Context->EmulatedQueries->LastTimer.IsValid());
-				Context->EmulatedQueries->LastTimer->End();
-				Start = Context->EmulatedQueries->LastTimer;
+				check(Emu->LastTimer.IsValid());
+				Emu->LastTimer->End();
+				Start = Emu->LastTimer;
 				
-				TSharedPtr<FMacOpenGLTimer, ESPMode::Fast> Current(new FMacOpenGLTimer(Context));
+				TSharedPtr<FMacOpenGLTimer, ESPMode::Fast> Current(new FMacOpenGLTimer(Context, Emu));
 				Current->Previous = Start;
 				Start->Next = Current;
 				Current->Begin();
-				Context->EmulatedQueries->LastTimer = Current;
+				Emu->LastTimer = Current;
 				
 				break;
 			}
 			case GL_TIME_ELAPSED:
 			{
-				check(Context->EmulatedQueries->LastTimer.IsValid());
-				Context->EmulatedQueries->LastTimer->End();
+				check(Emu->LastTimer.IsValid());
+				Emu->LastTimer->End();
 				
-				Start = TSharedPtr<FMacOpenGLTimer, ESPMode::Fast>(new FMacOpenGLTimer(Context));
-				Start->Previous = Context->EmulatedQueries->LastTimer;
-				Context->EmulatedQueries->LastTimer->Next = Start;
+				Start = TSharedPtr<FMacOpenGLTimer, ESPMode::Fast>(new FMacOpenGLTimer(Context, Emu));
+				Start->Previous = Emu->LastTimer;
+				Emu->LastTimer->Next = Start;
 				Start->Begin();
 				
-				Context->EmulatedQueries->LastTimer = Start;
+				Emu->LastTimer = Start;
 				break;
 			}
 			default:
 			{
 				if(!Name)
 				{
-					glGenQueries(1, &Name);
+					if(!Emu->FreeQueries.FindOrAdd(Target).Num())
+					{
+						if(IsRHIDeviceNVIDIA())
+						{
+							glFinish();
+						}
+						
+						Emu->FreeQueries.FindOrAdd(Target).AddUninitialized(GMacQueryNameCacheSize);
+						glGenQueries(GMacQueryNameCacheSize, Emu->FreeQueries.FindOrAdd(Target).GetData());
+					}
+					Name = Emu->FreeQueries.FindOrAdd(Target).Pop();
 				}
 				glBeginQuery(Target, Name);
 				break;
@@ -228,13 +259,13 @@ void FMacOpenGLQuery::End()
 	Running = false;
 	if(Target == GL_TIME_ELAPSED)
 	{
-		Finish = Context->EmulatedQueries->LastTimer;
-		TSharedPtr<FMacOpenGLTimer, ESPMode::Fast> Current(new FMacOpenGLTimer(Context));
+		Finish = Emu->LastTimer;
+		TSharedPtr<FMacOpenGLTimer, ESPMode::Fast> Current(new FMacOpenGLTimer(Context, Emu));
 		Current->Previous = Finish;
 		Finish->Next = Current;
 		Finish->End();
 		Current->Begin();
-		Context->EmulatedQueries->LastTimer = Current;
+		Emu->LastTimer = Current;
 		
 		if (Start == Finish)
 		{
@@ -295,27 +326,71 @@ int32 FMacOpenGLQuery::GetResultAvailable(void)
 
 FMacOpenGLQueryEmu::FMacOpenGLQueryEmu(FPlatformOpenGLContext* InContext)
 : PlatformContext(InContext)
+, bOwnsDummyDrawProgram( InContext == nullptr )
 {
 	check(PlatformContext);
 	
 	// Only use the timestamp emulation in a non-shipping build - end-users shouldn't care about this profiling feature.
 #if (!UE_BUILD_SHIPPING)
-	// Opt-in since it crashes Nvidia cards at the moment
-	GIsEmulatingTimestamp = FParse::Param(FCommandLine::Get(), TEXT("EnableMacGPUTimestamp"));
+	GIsEmulatingTimestamp = !FParse::Param(FCommandLine::Get(), TEXT("DisableMacGPUTimestamp"));
 #endif
 	
 	if ( GIsEmulatingTimestamp )
 	{
-		LastTimer = TSharedPtr<FMacOpenGLTimer, ESPMode::Fast>(new FMacOpenGLTimer(InContext));
+		FreeQueries.FindOrAdd(GL_TIME_ELAPSED).AddUninitialized(GMacQueryNameCacheSize);
+		glGenQueries(GMacQueryNameCacheSize, FreeQueries.FindOrAdd(GL_TIME_ELAPSED).GetData());
+		
+		FreeQueries.FindOrAdd(GL_SAMPLES_PASSED).AddUninitialized(GMacQueryNameCacheSize);
+		glGenQueries(GMacQueryNameCacheSize, FreeQueries.FindOrAdd(GL_SAMPLES_PASSED).GetData());
+		
+		GLuint Query = FreeQueries.FindOrAdd(GL_TIME_ELAPSED).Pop();
+		glBeginQuery(GL_TIME_ELAPSED, Query);
+		{
+			GLuint ShaderProgram = GetDummyDrawProgram();
+			glUseProgram(ShaderProgram);
+			
+			glDrawArraysInstanced(GL_TRIANGLE_STRIP, 0, 6, 1);
+			
+			glUseProgram(0);
+		}
+		glEndQuery(GL_TIME_ELAPSED);
+		glFlushRenderAPPLE();
+		FreeQueries.FindOrAdd(GL_TIME_ELAPSED).Push(Query);
+		
+		LastTimer = TSharedPtr<FMacOpenGLTimer, ESPMode::Fast>(new FMacOpenGLTimer(InContext, this));
 		LastTimer->Begin();
 	}
 }
 
 FMacOpenGLQueryEmu::~FMacOpenGLQueryEmu()
 {
-	if ( GIsEmulatingTimestamp && LastTimer.IsValid() )
+	if ( GIsEmulatingTimestamp )
 	{
-		LastTimer->End();
+		if( LastTimer.IsValid() )
+		{
+			LastTimer->End();
+			LastTimer = nullptr;
+		}
+			
+		Queries.Empty();
+		RunningQueries.Empty();
+		
+		if(IsRHIDeviceNVIDIA())
+		{
+			glFinish();
+		}
+		
+		for (auto Pair : FreeQueries)
+		{
+			glDeleteQueries(Pair.Value.Num(), Pair.Value.GetData());
+		}
+		FreeQueries.Empty();
+		
+		if ( bOwnsDummyDrawProgram )
+		{
+			glDeleteProgram(DummyDrawProgram);
+			DummyDrawProgram = 0;
+		}
 	}
 }
 	
@@ -339,7 +414,7 @@ GLuint FMacOpenGLQueryEmu::CreateQuery(void)
 	GLuint Result = 0;
 	if ( Queries.Num() < UINT_MAX )
 	{
-		TSharedPtr<FMacOpenGLQuery, ESPMode::Fast> Query(new FMacOpenGLQuery(PlatformContext));
+		TSharedPtr<FMacOpenGLQuery, ESPMode::Fast> Query(new FMacOpenGLQuery(PlatformContext, this));
 		for(int32 Index = 0; Index < Queries.Num(); ++Index)
 		{
 			if(!Queries[Index].IsValid())
@@ -386,3 +461,30 @@ void FMacOpenGLQueryEmu::EndQuery(GLenum const QueryType)
 		RunningQueries[QueryType].Reset();
 	}
 }
+
+GLuint FMacOpenGLQueryEmu::GetDummyDrawProgram()
+{
+	if ( !DummyDrawProgram )
+	{
+		GLuint VertexShader = glCreateShader(GL_VERTEX_SHADER);
+		glShaderSource(VertexShader, 1, &GMacTimerVertShader4, NULL);
+		glCompileShader(VertexShader);
+		GLuint FragmentShader = glCreateShader(GL_FRAGMENT_SHADER);
+		glShaderSource(FragmentShader, 1, &GMacTimerFragShader, NULL);
+		glCompileShader(FragmentShader);
+		GLuint ShaderProgram = glCreateProgram();
+		glAttachShader(ShaderProgram, VertexShader);
+		glAttachShader(ShaderProgram, FragmentShader);
+		glBindFragDataLocation(ShaderProgram, 0, "Color");
+		
+		glLinkProgram(ShaderProgram);
+		glValidateProgram(ShaderProgram);
+		
+		glDeleteShader(VertexShader);
+		glDeleteShader(FragmentShader);
+		
+		DummyDrawProgram = ShaderProgram;
+	}
+	
+	return DummyDrawProgram;
+}
diff --git a/Engine/Source/Runtime/OpenGLDrv/Private/Mac/MacOpenGLQuery.h b/Engine/Source/Runtime/OpenGLDrv/Private/Mac/MacOpenGLQuery.h
--- a/Engine/Source/Runtime/OpenGLDrv/Private/Mac/MacOpenGLQuery.h
+++ b/Engine/Source/Runtime/OpenGLDrv/Private/Mac/MacOpenGLQuery.h
@@ -5,10 +5,13 @@
 #include "SharedPointer.h"
 #include "MacOpenGLContext.h"
 
+class FMacOpenGLQueryEmu;
+typedef TMap<GLenum, TArray<GLuint>> FOpenGLQueryCache;
+
 class FMacOpenGLTimer : public TSharedFromThis<FMacOpenGLTimer, ESPMode::Fast>
 {
 public:
-	FMacOpenGLTimer(FPlatformOpenGLContext* InContext);
+	FMacOpenGLTimer(FPlatformOpenGLContext* InContext, FMacOpenGLQueryEmu* InEmu);
 	virtual ~FMacOpenGLTimer();
 	
 	void Begin(void);
@@ -26,6 +29,7 @@ private:
 	
 private:
 	FPlatformOpenGLContext* Context;
+	FMacOpenGLQueryEmu* Emu;
 	uint32 Name;
 	uint64 Result;
 	uint64 Accumulated;
@@ -37,7 +41,7 @@ private:
 class FMacOpenGLQuery : public TSharedFromThis<FMacOpenGLQuery, ESPMode::Fast>
 {
 public:
-	FMacOpenGLQuery(FPlatformOpenGLContext* InContext);
+	FMacOpenGLQuery(FPlatformOpenGLContext* InContext, FMacOpenGLQueryEmu* InEmu);
 	virtual ~FMacOpenGLQuery();
 	
 	void Begin(GLenum InTarget);
@@ -50,6 +54,7 @@ public:
 	
 private:
 	FPlatformOpenGLContext* Context;
+	FMacOpenGLQueryEmu* Emu;
 	uint64 Result;
 	uint32 Available;
 	bool Running;
@@ -61,6 +66,7 @@ private:
 class FMacOpenGLQueryEmu
 {
 	friend class FMacOpenGLQuery;
+	friend class FMacOpenGLTimer;
 public:
 	FMacOpenGLQueryEmu(FPlatformOpenGLContext* InContext);
 	~FMacOpenGLQueryEmu();
@@ -71,12 +77,18 @@ public:
 	void DeleteQuery(GLuint const QueryID);
 	void BeginQuery(GLenum const QueryType, GLuint const QueryID);
 	void EndQuery(GLenum const QueryType);
+
+private:
+	static GLuint GetDummyDrawProgram();
 	
 private:
+	static GLuint DummyDrawProgram;
 	FPlatformOpenGLContext* PlatformContext;
+	FOpenGLQueryCache FreeQueries;
 	TArray<TSharedPtr<FMacOpenGLQuery, ESPMode::Fast>> Queries;
 	TMap<GLenum, TSharedPtr<FMacOpenGLQuery, ESPMode::Fast>> RunningQueries;
 	TSharedPtr<FMacOpenGLTimer, ESPMode::Fast> LastTimer;
+	bool bOwnsDummyDrawProgram;
 };
 
 extern bool GIsEmulatingTimestamp;
