Linux: use SDL_GL_GetDrawableSize() instead of window size.

- PR #985 contributed by [REDACTED_PERSON].

diff --git a/Engine/Source/Runtime/OpenGLDrv/Private/Linux/OpenGLLinux.cpp b/Engine/Source/Runtime/OpenGLDrv/Private/Linux/OpenGLLinux.cpp
--- a/Engine/Source/Runtime/OpenGLDrv/Private/Linux/OpenGLLinux.cpp
+++ b/Engine/Source/Runtime/OpenGLDrv/Private/Linux/OpenGLLinux.cpp
@@ -392,8 +392,9 @@ bool PlatformBlitToViewport(FPlatformOpenGLDevice* Device,
 		glReadBuffer( GL_COLOR_ATTACHMENT0 );
 		glDisable(GL_FRAMEBUFFER_SRGB);
 
-		int WinW, WinH;
-		SDL_GetWindowSize(Context->hWnd, &WinW, &WinH);
+		int WinW = 0;
+		int WinH = 0;
+		SDL_GL_GetDrawableSize(Context->hWnd, &WinW, &WinH);
 		GLenum BlitFilter;
 		GLint DestX0, DestY0, DestX1, DestY1;
 
@@ -928,18 +929,18 @@ EOpenGLCurrentContext PlatformOpenGLCurrentContext( FPlatformOpenGLDevice* Devic
 
 void PlatformGetBackbufferDimensions( uint32& OutWidth, uint32& OutHeight )
 {
-	SDL_HWindow h_cwnd = SDL_GL_GetCurrentWindow();
+	SDL_HWindow CurrentWindow = SDL_GL_GetCurrentWindow();
 
-	if	( h_cwnd )
-	{
-		SDL_Surface *p_surf = SDL_GetWindowSurface( h_cwnd );
+	int Width = 0;
+	int Height = 0;
 
-		if	( p_surf )
-		{
-			OutWidth  = p_surf->w;
-			OutHeight = p_surf->h;
-		}
+	if (CurrentWindow)
+	{
+		SDL_GL_GetDrawableSize(CurrentWindow, &Width, &Height);
 	}
+
+	OutWidth = Width;
+	OutHeight = Height;
 }
 
 // =============================================================
