Removed unnecessary shaders & program from Mac OS X GL timestamp emulation and changed the way the query cache is initialised to reduce crashes on project startup using some GPUs.

diff --git a/Engine/Source/Runtime/OpenGLDrv/Private/Mac/MacOpenGLQuery.cpp b/Engine/Source/Runtime/OpenGLDrv/Private/Mac/MacOpenGLQuery.cpp
--- a/Engine/Source/Runtime/OpenGLDrv/Private/Mac/MacOpenGLQuery.cpp
+++ b/Engine/Source/Runtime/OpenGLDrv/Private/Mac/MacOpenGLQuery.cpp
@@ -11,26 +11,6 @@
 bool GIsEmulatingTimestamp = false;
 static const uint32 GMacQueryNameCacheSize = 32 * OPENGL_NAME_CACHE_SIZE;
 
-char const* const GMacTimerVertShader4 =	"#version 150\n"
-											"const int VertexCount = 4;\n"
-											"const vec2 Position[VertexCount] = vec2[](	vec2(-1.0,-1.0),\n"
-											"											vec2( 1.0,-1.0),\n"
-											"											vec2(-1.0, 1.0),\n"
-											"											vec2( 1.0, 1.0));\n"
-											"void main()\n"
-											"{\n"
-											"	gl_Position = vec4(Position[gl_VertexID], 0.0, 1.0);\n"
-											"}\n";
-
-char const* const GMacTimerFragShader =	"#version 150\n"
-										"out vec4 Color;\n"
-										"void main()\n"
-										"{\n"
-										"	Color = vec4(0.0, 0.0, 0.0, 1.0);\n"
-										"}\n";
-
-GLuint FMacOpenGLQueryEmu::DummyDrawProgram = 0;
-
 /*------------------------------------------------------------------------------
  OpenGL query emulation.
  ------------------------------------------------------------------------------*/
@@ -326,7 +306,6 @@ int32 FMacOpenGLQuery::GetResultAvailable(void)
 
 FMacOpenGLQueryEmu::FMacOpenGLQueryEmu(FPlatformOpenGLContext* InContext)
 : PlatformContext(InContext)
-, bOwnsDummyDrawProgram( InContext == nullptr )
 {
 	check(PlatformContext);
 	
@@ -343,19 +322,41 @@ FMacOpenGLQueryEmu::FMacOpenGLQueryEmu(FPlatformOpenGLContext* InContext)
 		FreeQueries.FindOrAdd(GL_SAMPLES_PASSED).AddUninitialized(GMacQueryNameCacheSize);
 		glGenQueries(GMacQueryNameCacheSize, FreeQueries.FindOrAdd(GL_SAMPLES_PASSED).GetData());
 		
-		GLuint Query = FreeQueries.FindOrAdd(GL_TIME_ELAPSED).Pop();
-		glBeginQuery(GL_TIME_ELAPSED, Query);
+		float red = 0;
+		float green = 0;
+		float blue = 0;
+		
+		glClearColor(1.0, 0.5, 0.0, 1.0);
+		glClear(GL_COLOR_BUFFER_BIT);
+		
+		for(uint32 i = 0; i < FreeQueries.FindOrAdd(GL_TIME_ELAPSED).Num(); i++)
 		{
-			GLuint ShaderProgram = GetDummyDrawProgram();
-			glUseProgram(ShaderProgram);
+			GLuint Query = FreeQueries.FindOrAdd(GL_TIME_ELAPSED)[i];
+			glBeginQuery(GL_TIME_ELAPSED, Query);
 			
-			glDrawArraysInstanced(GL_TRIANGLE_STRIP, 0, 6, 1);
+			red = (float)(i & 0xff);
+			green = (float)((i >> 8) & 0xff);
+			blue = (float)((i >> 16) & 0xff);
+			glClearColor(red, green, blue, 1.0);
+			glClear(GL_COLOR_BUFFER_BIT);
 			
-			glUseProgram(0);
+			glEndQuery(GL_TIME_ELAPSED);
 		}
-		glEndQuery(GL_TIME_ELAPSED);
-		glFlushRenderAPPLE();
-		FreeQueries.FindOrAdd(GL_TIME_ELAPSED).Push(Query);
+		for(uint32 i = 0; i < FreeQueries.FindOrAdd(GL_TIME_ELAPSED).Num(); i++)
+		{
+			GLuint64 Result;
+			GLuint Query = FreeQueries.FindOrAdd(GL_TIME_ELAPSED)[i];
+			glGetQueryObjectui64v(Query, GL_QUERY_RESULT, &Result);
+		}
+		
+		glClearColor(0.0, 0.0, 0.0, 1.0);
+		glClear(GL_COLOR_BUFFER_BIT);
+		
+		glDeleteQueries(FreeQueries.FindOrAdd(GL_TIME_ELAPSED).Num(), FreeQueries.FindOrAdd(GL_TIME_ELAPSED).GetData());
+		FreeQueries.FindOrAdd(GL_TIME_ELAPSED).Empty();
+		
+		FreeQueries.FindOrAdd(GL_TIME_ELAPSED).AddUninitialized(GMacQueryNameCacheSize);
+		glGenQueries(GMacQueryNameCacheSize, FreeQueries.FindOrAdd(GL_TIME_ELAPSED).GetData());
 		
 		LastTimer = TSharedPtr<FMacOpenGLTimer, ESPMode::Fast>(new FMacOpenGLTimer(InContext, this));
 		LastTimer->Begin();
@@ -385,12 +386,6 @@ FMacOpenGLQueryEmu::~FMacOpenGLQueryEmu()
 			glDeleteQueries(Pair.Value.Num(), Pair.Value.GetData());
 		}
 		FreeQueries.Empty();
-		
-		if ( bOwnsDummyDrawProgram )
-		{
-			glDeleteProgram(DummyDrawProgram);
-			DummyDrawProgram = 0;
-		}
 	}
 }
 	
@@ -461,30 +456,3 @@ void FMacOpenGLQueryEmu::EndQuery(GLenum const QueryType)
 		RunningQueries[QueryType].Reset();
 	}
 }
-
-GLuint FMacOpenGLQueryEmu::GetDummyDrawProgram()
-{
-	if ( !DummyDrawProgram )
-	{
-		GLuint VertexShader = glCreateShader(GL_VERTEX_SHADER);
-		glShaderSource(VertexShader, 1, &GMacTimerVertShader4, NULL);
-		glCompileShader(VertexShader);
-		GLuint FragmentShader = glCreateShader(GL_FRAGMENT_SHADER);
-		glShaderSource(FragmentShader, 1, &GMacTimerFragShader, NULL);
-		glCompileShader(FragmentShader);
-		GLuint ShaderProgram = glCreateProgram();
-		glAttachShader(ShaderProgram, VertexShader);
-		glAttachShader(ShaderProgram, FragmentShader);
-		glBindFragDataLocation(ShaderProgram, 0, "Color");
-		
-		glLinkProgram(ShaderProgram);
-		glValidateProgram(ShaderProgram);
-		
-		glDeleteShader(VertexShader);
-		glDeleteShader(FragmentShader);
-		
-		DummyDrawProgram = ShaderProgram;
-	}
-	
-	return DummyDrawProgram;
-}
diff --git a/Engine/Source/Runtime/OpenGLDrv/Private/Mac/MacOpenGLQuery.h b/Engine/Source/Runtime/OpenGLDrv/Private/Mac/MacOpenGLQuery.h
--- a/Engine/Source/Runtime/OpenGLDrv/Private/Mac/MacOpenGLQuery.h
+++ b/Engine/Source/Runtime/OpenGLDrv/Private/Mac/MacOpenGLQuery.h
@@ -79,16 +79,11 @@ public:
 	void EndQuery(GLenum const QueryType);
 
 private:
-	static GLuint GetDummyDrawProgram();
-	
-private:
-	static GLuint DummyDrawProgram;
 	FPlatformOpenGLContext* PlatformContext;
 	FOpenGLQueryCache FreeQueries;
 	TArray<TSharedPtr<FMacOpenGLQuery, ESPMode::Fast>> Queries;
 	TMap<GLenum, TSharedPtr<FMacOpenGLQuery, ESPMode::Fast>> RunningQueries;
 	TSharedPtr<FMacOpenGLTimer, ESPMode::Fast> LastTimer;
-	bool bOwnsDummyDrawProgram;
 };
 
 extern bool GIsEmulatingTimestamp;
