When resizing an OS X OpenGL context clear the new back-buffer resource to black & only blit the viewport back-buffer to the window when its contents can be guaranteed to be valid, otherwise clear to black.

The Editor will not be affected by this change, but standalone games would render garbage on startup & when resizing because they don't necessarily render into the viewport backbuffer between a resize message arriving & the next time Cocoa needs to draw the window.

diff --git a/Engine/Source/Runtime/OpenGLDrv/Private/Mac/OpenGLMac.cpp b/Engine/Source/Runtime/OpenGLDrv/Private/Mac/OpenGLMac.cpp
--- a/Engine/Source/Runtime/OpenGLDrv/Private/Mac/OpenGLMac.cpp
+++ b/Engine/Source/Runtime/OpenGLDrv/Private/Mac/OpenGLMac.cpp
@@ -455,7 +455,7 @@ bool PlatformBlitToViewport(FPlatformOpenGLDevice* Device, const FOpenGLViewport
 void DrawOpenGLViewport(FPlatformOpenGLContext* const Context, uint32 Width, uint32 Height)
 {
 	FCocoaWindow* Window = (FCocoaWindow*)Context->WindowHandle;
-	if ([Window isRenderInitialized] && Context->ViewportSize[0] && Context->ViewportSize[1] && Context->ViewportFramebuffer && Context->ViewportRenderbuffer)
+	if ([Window isRenderInitialized] && Context->ViewportSize[0] && Context->ViewportSize[1] && Context->ViewportFramebuffer && Context->ViewportRenderbuffer && ([Window styleMask] & (NSTexturedBackgroundWindowMask|NSFullSizeContentViewWindowMask) || !Window.inLiveResize))
 	{
 		int32 CurrentReadFramebuffer = 0;
 		glGetIntegerv(GL_READ_FRAMEBUFFER_BINDING, &CurrentReadFramebuffer);
@@ -571,14 +571,23 @@ void PlatformResizeGLContext( FPlatformOpenGLDevice* Device, FPlatformOpenGLCont
 			check(Context->ViewportFramebuffer);
 		}
 		
-		if (Context->ViewportRenderbuffer == 0)
+		glBindFramebuffer(GL_FRAMEBUFFER, Context->ViewportFramebuffer);
+		glFramebufferTexture2D(GL_FRAMEBUFFER, GL_COLOR_ATTACHMENT0, BackBufferTarget, BackBufferResource, 0);
+		glFramebufferTexture2D(GL_FRAMEBUFFER, GL_COLOR_ATTACHMENT1, GL_TEXTURE_2D, 0, 0);
+		glViewport(0, 0, SizeX, SizeY);
+		glDrawBuffer(GL_COLOR_ATTACHMENT0);
+		glClearColor(0.f, 0.f, 0.f, 0.f);
+		glClear(GL_COLOR_BUFFER_BIT);
+		
+		if (Context->ViewportRenderbuffer != 0)
+		{
+			glFramebufferRenderbuffer(GL_FRAMEBUFFER, GL_COLOR_ATTACHMENT1, GL_RENDERBUFFER, Context->ViewportRenderbuffer);
+		}
+		else
 		{
 			glGenRenderbuffers(1, &Context->ViewportRenderbuffer);
 			check(Context->ViewportRenderbuffer);
 		}
-		
-		glBindFramebuffer(GL_FRAMEBUFFER, Context->ViewportFramebuffer);
-		glFramebufferTexture2D(GL_FRAMEBUFFER, GL_COLOR_ATTACHMENT0, BackBufferTarget, BackBufferResource, 0);
 #if UE_BUILD_DEBUG
 		glReadBuffer(GL_COLOR_ATTACHMENT0);
 		glDrawBuffer(GL_COLOR_ATTACHMENT0);
