Initialise the VendorID immediately on OS X to avoid an assert by recently added calls to IsRHIDeviceNVIDIA within the extension processing code.

diff --git a/Engine/Source/Runtime/OpenGLDrv/Private/Mac/OpenGLMac.cpp b/Engine/Source/Runtime/OpenGLDrv/Private/Mac/OpenGLMac.cpp
--- a/Engine/Source/Runtime/OpenGLDrv/Private/Mac/OpenGLMac.cpp
+++ b/Engine/Source/Runtime/OpenGLDrv/Private/Mac/OpenGLMac.cpp
@@ -1001,6 +1001,21 @@ void FMacOpenGL::ProcessQueryGLInt()
 
 void FMacOpenGL::ProcessExtensions(const FString& ExtensionsString)
 {
+	// Get the Vendor ID by parsing the renderer string
+	FString VendorName( ANSI_TO_TCHAR((const ANSICHAR*)glGetString(GL_VENDOR)));
+	if (VendorName.Contains(TEXT("Intel ")))
+	{
+		GRHIVendorId = 0x8086;
+	}
+	else if (VendorName.Contains(TEXT("NVIDIA ")))
+	{
+		GRHIVendorId = 0x10DE;
+	}
+	else if (VendorName.Contains(TEXT("ATi ")) || VendorName.Contains(TEXT("AMD ")))
+	{
+		GRHIVendorId = 0x1002;
+	}
+	
 	ProcessQueryGLInt();
 	FOpenGL3::ProcessExtensions(ExtensionsString);
 	
