Further harden GPU vendor ID resolution on OS X to try and stop crash on startup.

diff --git a/Engine/Source/Runtime/OpenGLDrv/Private/Mac/MacOpenGLContext.cpp b/Engine/Source/Runtime/OpenGLDrv/Private/Mac/MacOpenGLContext.cpp
--- a/Engine/Source/Runtime/OpenGLDrv/Private/Mac/MacOpenGLContext.cpp
+++ b/Engine/Source/Runtime/OpenGLDrv/Private/Mac/MacOpenGLContext.cpp
@@ -432,19 +432,49 @@ void FPlatformOpenGLContext::VerifyCurrentContext()
 			
 			// Get the Vendor ID by parsing the renderer string
 			FString VendorName( ANSI_TO_TCHAR((const ANSICHAR*)glGetString(GL_VENDOR)));
-			if (VendorName.Contains(TEXT("Intel ")))
+			if (VendorName.Contains(TEXT("Intel")))
 			{
 				PlatformContext->VendorID = 0x8086;
 			}
-			else if (VendorName.Contains(TEXT("NVIDIA ")))
+			else if (VendorName.Contains(TEXT("NVIDIA")))
 			{
 				PlatformContext->VendorID = 0x10DE;
 			}
-			else if (VendorName.Contains(TEXT("ATi ")) || VendorName.Contains(TEXT("AMD ")))
+			else if (VendorName.Contains(TEXT("ATi")) || VendorName.Contains(TEXT("AMD")))
 			{
 				PlatformContext->VendorID = 0x1002;
 			}
 			
+			if(PlatformContext->VendorID == 0)
+			{
+				// Get the current renderer ID
+				switch(RendererID & 0x000ff000)
+				{
+					case 0x00021000:
+					{
+						PlatformContext->VendorID = 0x1002;
+						break;
+					}
+					case 0x00022000:
+					{
+						PlatformContext->VendorID = 0x10DE;
+						break;
+					}
+					case 0x00024000:
+					{
+						PlatformContext->VendorID = 0x8086;
+						break;
+					}
+					default:
+					{
+						// Unknown GPU vendor - assuming Intel!
+						PlatformContext->VendorID = 0x8086;
+						break;
+					}
+				}
+			}
+			check(PlatformContext->VendorID != 0);
+			
 			// Renderer IDs matchup to driver kexts, so switching based on them will allow us to target workarouds to many GPUs
 			// which exhibit the same unfortunate driver bugs without having to parse their individual ID strings.
 			switch((PlatformContext->RendererID & kCGLRendererIDMatchingMask))
diff --git a/Engine/Source/Runtime/OpenGLDrv/Private/Mac/OpenGLMac.cpp b/Engine/Source/Runtime/OpenGLDrv/Private/Mac/OpenGLMac.cpp
--- a/Engine/Source/Runtime/OpenGLDrv/Private/Mac/OpenGLMac.cpp
+++ b/Engine/Source/Runtime/OpenGLDrv/Private/Mac/OpenGLMac.cpp
@@ -1003,19 +1003,55 @@ void FMacOpenGL::ProcessExtensions(const FString& ExtensionsString)
 {
 	// Get the Vendor ID by parsing the renderer string
 	FString VendorName( ANSI_TO_TCHAR((const ANSICHAR*)glGetString(GL_VENDOR)));
-	if (VendorName.Contains(TEXT("Intel ")))
+	if (VendorName.Contains(TEXT("Intel")))
 	{
 		GRHIVendorId = 0x8086;
 	}
-	else if (VendorName.Contains(TEXT("NVIDIA ")))
+	else if (VendorName.Contains(TEXT("NVIDIA")))
 	{
 		GRHIVendorId = 0x10DE;
 	}
-	else if (VendorName.Contains(TEXT("ATi ")) || VendorName.Contains(TEXT("AMD ")))
+	else if (VendorName.Contains(TEXT("ATi")) || VendorName.Contains(TEXT("AMD")))
 	{
 		GRHIVendorId = 0x1002;
 	}
 	
+	if(GRHIVendorId == 0)
+	{
+		// Get the current renderer ID
+		CGLContextObj Current = CGLGetCurrentContext();
+		GLint RendererID = 0;
+		CGLError Error = CGLGetParameter(Current, kCGLCPCurrentRendererID, &RendererID);
+		if(Error == kCGLNoError)
+		{
+			switch(RendererID & 0x000ff000)
+			{
+				case 0x00021000:
+				{
+					GRHIVendorId = 0x1002;
+					break;
+				}
+				case 0x00022000:
+				{
+					GRHIVendorId = 0x10DE;
+					break;
+				}
+				case 0x00024000:
+				{
+					GRHIVendorId = 0x8086;
+					break;
+				}
+				default:
+				{
+					// Unknown GPU vendor - assuming Intel!
+					GRHIVendorId = 0x8086;
+					break;
+				}
+			}
+		}
+	}
+	check(GRHIVendorId != 0);
+	
 	ProcessQueryGLInt();
 	FOpenGL3::ProcessExtensions(ExtensionsString);
 	
