HTML5: add more logging for the framebuffer completeness check.

diff --git a/Engine/Source/Runtime/OpenGLDrv/Private/HTML5/HTML5OpenGL.cpp b/Engine/Source/Runtime/OpenGLDrv/Private/HTML5/HTML5OpenGL.cpp
--- a/Engine/Source/Runtime/OpenGLDrv/Private/HTML5/HTML5OpenGL.cpp
+++ b/Engine/Source/Runtime/OpenGLDrv/Private/HTML5/HTML5OpenGL.cpp
@@ -84,6 +84,8 @@ void FHTML5OpenGL::ProcessExtensions( const FString& ExtensionsString )
         // See http://www.khronos.org/webgl/public-mailing-list/archives/1211/msg00133.html
         // for more information.
 
+		UE_LOG(LogRHI, Warning, TEXT("Trying to enable fp rendering without explicit EXT_color_buffer_half_float by checking for framebuffer completeness"));
+
         GLenum err = glGetError();
         if (err != GL_NO_ERROR) {
             UE_LOG(LogRHI, Warning, TEXT("Detected OpenGL error 0x%04x before checking for implicit half-float fb support"), err);
@@ -121,9 +123,14 @@ void FHTML5OpenGL::ProcessExtensions( const FString& ExtensionsString )
 #endif 
         bSupportsColorBufferHalfFloat = fbstatus == GL_FRAMEBUFFER_COMPLETE && err == GL_NO_ERROR;
 
-        if (bSupportsColorBufferHalfFloat) {
+        if (bSupportsColorBufferHalfFloat) 
+		{
             UE_LOG(LogRHI, Log, TEXT("Enabling implicit ColorBufferHalfFloat after checking fb completeness"));
         }
+		else
+		{
+			UE_LOG(LogRHI, Log, TEXT("Could not enable implicit ColorBufferHalfFloat after checking fb completeness"));
+		}
 
         glBindFramebuffer(GL_FRAMEBUFFER, 0);
         glDeleteFramebuffers(1, &fb);
