Disable viewport frame sync on Android (can cause issues with tile rendering)
#ue4
#android

diff --git a/Engine/Source/Runtime/OpenGLDrv/Private/Android/AndroidES31OpenGL.cpp b/Engine/Source/Runtime/OpenGLDrv/Private/Android/AndroidES31OpenGL.cpp
--- a/Engine/Source/Runtime/OpenGLDrv/Private/Android/AndroidES31OpenGL.cpp
+++ b/Engine/Source/Runtime/OpenGLDrv/Private/Android/AndroidES31OpenGL.cpp
@@ -305,7 +305,9 @@ bool PlatformBlitToViewport( FPlatformOpenGLDevice* Device, const FOpenGLViewpor
 			GRenderThreadNumIdle[ERenderThreadIdleTypes::WaitingForGPUPresent]++;
 		}
 	}
-	return true;
+//	return true;
+	//Do not want WaitForFrameEventCompletion
+	return false;
 }
 
 void PlatformRenderingContextSetup(FPlatformOpenGLDevice* Device)
diff --git a/Engine/Source/Runtime/OpenGLDrv/Private/Android/AndroidOpenGL.cpp b/Engine/Source/Runtime/OpenGLDrv/Private/Android/AndroidOpenGL.cpp
--- a/Engine/Source/Runtime/OpenGLDrv/Private/Android/AndroidOpenGL.cpp
+++ b/Engine/Source/Runtime/OpenGLDrv/Private/Android/AndroidOpenGL.cpp
@@ -127,7 +127,9 @@ bool PlatformBlitToViewport( FPlatformOpenGLDevice* Device, const FOpenGLViewpor
 	{
 		AndroidEGL::GetInstance()->SwapBuffers();
 	}
-	return true;
+//	return true;
+	//Do not want WaitForFrameEventCompletion
+	return false;
 }
 
 void PlatformRenderingContextSetup(FPlatformOpenGLDevice* Device)
