Workaround for odd compile problems seen on SGX540.

Fixes UE-13933 and potentially UE-3004

diff --git a/Engine/Source/Developer/ShaderFormatOpenGL/Private/GlslBackend.cpp b/Engine/Source/Developer/ShaderFormatOpenGL/Private/GlslBackend.cpp
--- a/Engine/Source/Developer/ShaderFormatOpenGL/Private/GlslBackend.cpp
+++ b/Engine/Source/Developer/ShaderFormatOpenGL/Private/GlslBackend.cpp
@@ -2856,6 +2856,18 @@ public:
 			ralloc_asprintf_append(buffer, "precision %s sampler2D;\n", DefaultPrecision);
 			ralloc_asprintf_append(buffer, "precision %s samplerCube;\n\n", DefaultPrecision);
 			ralloc_asprintf_append(buffer, "#endif\n");
+
+			// SGX540 compiler can get upset with some operations that mix highp and mediump.
+			// this results in a shader compile fail with output "compile failed."
+			// Although the actual cause of the failure hasnt been determined this code appears to prevent
+			// compile failure for cases so far seen.
+			ralloc_asprintf_append(buffer, "\n#ifdef TEXCOORDPRECISIONWORKAROUND\n");
+			ralloc_asprintf_append(buffer, "vec4 texture2DTexCoordPrecisionWorkaround(sampler2D p, vec2 tcoord)\n");
+			ralloc_asprintf_append(buffer, "{\n");
+			ralloc_asprintf_append(buffer, "	return texture2D(p, tcoord);\n");
+			ralloc_asprintf_append(buffer, "}\n");
+			ralloc_asprintf_append(buffer, "#define texture2D texture2DTexCoordPrecisionWorkaround\n");
+			ralloc_asprintf_append(buffer, "#endif\n");
 		}
 
 		if ((state->language_version == 310) && (ShaderTarget == fragment_shader))
diff --git a/Engine/Source/Developer/ShaderFormatOpenGL/Private/ShaderFormatOpenGL.cpp b/Engine/Source/Developer/ShaderFormatOpenGL/Private/ShaderFormatOpenGL.cpp
--- a/Engine/Source/Developer/ShaderFormatOpenGL/Private/ShaderFormatOpenGL.cpp
+++ b/Engine/Source/Developer/ShaderFormatOpenGL/Private/ShaderFormatOpenGL.cpp
@@ -27,11 +27,11 @@ class FShaderFormatGLSL : public IShaderFormat
 		UE_SHADER_GLSL_150_VER = 26,
 		UE_SHADER_GLSL_150_MAC_VER = 6,
 		UE_SHADER_GLSL_430_VER = 57,
-		UE_SHADER_GLSL_ES2_VER = 14,
-		UE_SHADER_GLSL_150ES2_VER = 17,
-		UE_SHADER_GLSL_150ES2NOUB_VER = 1,
+		UE_SHADER_GLSL_ES2_VER = 15,
+		UE_SHADER_GLSL_150ES2_VER = 18,
+		UE_SHADER_GLSL_150ES2NOUB_VER = 2,
 		UE_SHADER_GLSL_150ES3_1_VER = 1,
-		UE_SHADER_GLSL_ES2_VER_WEBGL  = 15,
+		UE_SHADER_GLSL_ES2_VER_WEBGL  = 16,
 		UE_SHADER_GLSL_ES2_IOS_VER  = 5,
 		UE_SHADER_GLSL_310_ES_EXT_VER = 3,
 	}; 
diff --git a/Engine/Source/Runtime/OpenGLDrv/Private/Android/AndroidOpenGL.cpp b/Engine/Source/Runtime/OpenGLDrv/Private/Android/AndroidOpenGL.cpp
--- a/Engine/Source/Runtime/OpenGLDrv/Private/Android/AndroidOpenGL.cpp
+++ b/Engine/Source/Runtime/OpenGLDrv/Private/Android/AndroidOpenGL.cpp
@@ -360,6 +360,7 @@ void FAndroidOpenGL::ProcessExtensions(const FString& ExtensionsString)
 	{
 		UE_LOG(LogRHI, Warning, TEXT("Disabling support for GL_OES_packed_depth_stencil on SGX 540"));
 		bSupportsPackedDepthStencil = false;
+		bRequiresTexture2DPrecisionHack = true;
 	}
 
 	if( RendererString.Contains(TEXT("Adreno")) )
diff --git a/Engine/Source/Runtime/OpenGLDrv/Private/OpenGLES2.cpp b/Engine/Source/Runtime/OpenGLDrv/Private/OpenGLES2.cpp
--- a/Engine/Source/Runtime/OpenGLDrv/Private/OpenGLES2.cpp
+++ b/Engine/Source/Runtime/OpenGLDrv/Private/OpenGLES2.cpp
@@ -127,6 +127,9 @@ bool FOpenGLES2::bRequiresGLFragCoordVaryingLimitHack = false;
 /** Vertex attributes need remapping if GL_MAX_VERTEX_ATTRIBS < 16 */
 bool FOpenGLES2::bNeedsVertexAttribRemap = false;
 
+/* This hack fixes an issue with SGX540 compiler which can get upset with some operations that mix highp and mediump */
+bool FOpenGLES2::bRequiresTexture2DPrecisionHack = false;
+
 bool FOpenGLES2::SupportsDisjointTimeQueries()
 {
 	bool bAllowDisjointTimerQueries = false;
diff --git a/Engine/Source/Runtime/OpenGLDrv/Private/OpenGLES31.cpp b/Engine/Source/Runtime/OpenGLDrv/Private/OpenGLES31.cpp
--- a/Engine/Source/Runtime/OpenGLDrv/Private/OpenGLES31.cpp
+++ b/Engine/Source/Runtime/OpenGLDrv/Private/OpenGLES31.cpp
@@ -124,6 +124,9 @@ bool FOpenGLES31::bRequiresTextureCubeLodEXTToTextureCubeLodDefine = false;
 /* This is a hack to remove the gl_FragCoord if shader will fail to link if exceeding the max varying on android platforms */
 bool FOpenGLES31::bRequiresGLFragCoordVaryingLimitHack = false;
 
+/* This hack fixes an issue with SGX540 compiler which can get upset with some operations that mix highp and mediump */
+bool FOpenGLES2::bRequiresTexture2DPrecisionHack = false;
+
 /** GL_EXT_disjoint_timer_query or GL_NV_timer_query*/
 bool FOpenGLES31::bSupportsDisjointTimeQueries = false;
 
diff --git a/Engine/Source/Runtime/OpenGLDrv/Private/OpenGLShaders.cpp b/Engine/Source/Runtime/OpenGLDrv/Private/OpenGLShaders.cpp
--- a/Engine/Source/Runtime/OpenGLDrv/Private/OpenGLShaders.cpp
+++ b/Engine/Source/Runtime/OpenGLDrv/Private/OpenGLShaders.cpp
@@ -658,6 +658,12 @@ ShaderType* CompileOpenGLShader(const TArray<uint8>& Code)
 							ReplaceCString(GlslCodeOriginal, "gl_FragCoord.xy", "vec2(400.5,240.5)");
 						}
 					}
+
+					if (FOpenGL::RequiresTexture2DPrecisionHack())
+					{
+						AppendCString(GlslCode,	"#define TEXCOORDPRECISIONWORKAROUND \n");
+					}
+
 				}
 			}
 		}
diff --git a/Engine/Source/Runtime/OpenGLDrv/Public/OpenGLES2.h b/Engine/Source/Runtime/OpenGLDrv/Public/OpenGLES2.h
--- a/Engine/Source/Runtime/OpenGLDrv/Public/OpenGLES2.h
+++ b/Engine/Source/Runtime/OpenGLDrv/Public/OpenGLES2.h
@@ -124,11 +124,14 @@ struct FOpenGLES2 : public FOpenGLBase
 	static FORCEINLINE bool SupportsCopyTextureLevels()					{ return bSupportsCopyTextureLevels; }
 	static FORCEINLINE GLenum GetDepthFormat()							{ return GL_DEPTH_COMPONENT; }
 	static FORCEINLINE GLenum GetShadowDepthFormat()					{ return GL_DEPTH_COMPONENT; }
+	static FORCEINLINE bool SupportsFramebufferSRGBEnable()				{ return false; }
+
 
 	static FORCEINLINE bool RequiresDontEmitPrecisionForTextureSamplers() { return bRequiresDontEmitPrecisionForTextureSamplers; }
 	static FORCEINLINE bool RequiresTextureCubeLodEXTToTextureCubeLodDefine() { return bRequiresTextureCubeLodEXTToTextureCubeLodDefine; }
 	static FORCEINLINE bool SupportsStandardDerivativesExtension()		{ return bSupportsStandardDerivativesExtension; }
 	static FORCEINLINE bool RequiresGLFragCoordVaryingLimitHack()		{ return bRequiresGLFragCoordVaryingLimitHack; }
+	static FORCEINLINE bool RequiresTexture2DPrecisionHack()			{ return bRequiresTexture2DPrecisionHack; }
 
 	static FORCEINLINE int32 GetReadHalfFloatPixelsEnum()				{ return GL_HALF_FLOAT_OES; }
 
@@ -484,6 +487,9 @@ public:
 
 	/* This is a hack to remove the gl_FragCoord if shader will fail to link if exceeding the max varying on android platforms */
 	static bool bRequiresGLFragCoordVaryingLimitHack;
+
+	/* This hack fixes an issue with SGX540 compiler which can get upset with some operations that mix highp and mediump */
+	static bool bRequiresTexture2DPrecisionHack;
 };
 
 
diff --git a/Engine/Source/Runtime/OpenGLDrv/Public/OpenGLES31.h b/Engine/Source/Runtime/OpenGLDrv/Public/OpenGLES31.h
--- a/Engine/Source/Runtime/OpenGLDrv/Public/OpenGLES31.h
+++ b/Engine/Source/Runtime/OpenGLDrv/Public/OpenGLES31.h
@@ -86,6 +86,7 @@ struct FOpenGLES31 : public FOpenGLBase
 	static FORCEINLINE bool SupportsStandardDerivativesExtension()		{ return true; }
 	static FORCEINLINE bool RequiresGLFragCoordVaryingLimitHack()		{ return bRequiresGLFragCoordVaryingLimitHack; }
 	static FORCEINLINE GLenum GetVertexHalfFloatFormat()				{ return bES2Fallback ? GL_HALF_FLOAT_OES : GL_HALF_FLOAT; }
+	static FORCEINLINE bool RequiresTexture2DPrecisionHack()			{ return bRequiresTexture2DPrecisionHack; }
 
 
 	// On iOS both glMapBufferOES() and glBufferSubData() for immediate vertex and index data
@@ -937,6 +938,9 @@ public:
 
 	/* This is a hack to remove the gl_FragCoord if shader will fail to link if exceeding the max varying on android platforms */
 	static bool bRequiresGLFragCoordVaryingLimitHack;
+
+	/* This hack fixes an issue with SGX540 compiler which can get upset with some operations that mix highp and mediump */
+	static bool bRequiresTexture2DPrecisionHack;
 };
 
 // yes they are different between the ES2 extension and ES3.x and GL3.x core
