Fixed crashes on OS X Mavericks

diff --git a/Engine/Source/Runtime/Core/Private/Mac/MacPlatformMemory.cpp b/Engine/Source/Runtime/Core/Private/Mac/MacPlatformMemory.cpp
--- a/Engine/Source/Runtime/Core/Private/Mac/MacPlatformMemory.cpp
+++ b/Engine/Source/Runtime/Core/Private/Mac/MacPlatformMemory.cpp
@@ -24,7 +24,20 @@ void FMacPlatformMemory::Init()
 
 FMalloc* FMacPlatformMemory::BaseAllocator()
 {
-	if(getenv("UE4_FORCE_MALLOC_ANSI") != nullptr)
+	bool bIsMavericks = false;
+
+	char OSRelease[PATH_MAX] = {};
+	size_t OSReleaseBufferSize = PATH_MAX;
+	if (sysctlbyname("kern.osrelease", OSRelease, &OSReleaseBufferSize, NULL, 0) == 0)
+	{
+		int32 OSVersionMajor = 0;
+		if (sscanf(OSRelease, "%d", &OSVersionMajor) == 1)
+		{
+			bIsMavericks = OSVersionMajor <= 13;
+		}
+	}
+
+	if(getenv("UE4_FORCE_MALLOC_ANSI") != nullptr || bIsMavericks)
 	{
 		return new FMallocAnsi();
 	}
diff --git a/Engine/Source/Runtime/OpenGLDrv/Private/Mac/OpenGLMac.cpp b/Engine/Source/Runtime/OpenGLDrv/Private/Mac/OpenGLMac.cpp
--- a/Engine/Source/Runtime/OpenGLDrv/Private/Mac/OpenGLMac.cpp
+++ b/Engine/Source/Runtime/OpenGLDrv/Private/Mac/OpenGLMac.cpp
@@ -114,6 +114,30 @@ private:
 static void DeleteQueriesForCurrentContext( NSOpenGLContext* Context );
 static void DrawOpenGLViewport(FPlatformOpenGLContext* const Context, uint32 Width, uint32 Height);
 
+static void LockGLContext(NSOpenGLContext* Context)
+{
+	if (FPlatformMisc::IsRunningOnMavericks())
+	{
+		CGLLockContext([Context CGLContextObj]);
+	}
+	else
+	{
+		[Context lock];
+	}
+}
+
+static void UnlockGLContext(NSOpenGLContext* Context)
+{
+	if (FPlatformMisc::IsRunningOnMavericks())
+	{
+		CGLUnlockContext([Context CGLContextObj]);
+	}
+	else
+	{
+		[Context unlock];
+	}
+}
+
 /*------------------------------------------------------------------------------
  OpenGL view.
  ------------------------------------------------------------------------------*/
@@ -151,7 +175,7 @@ static void DrawOpenGLViewport(FPlatformOpenGLContext* const Context, uint32 Wid
 	BOOL bOK = [super canDrawInOpenGLContext:context pixelFormat:pixelFormat forLayerTime:timeInterval displayTime:timeStamp];
 	if ( bOK && context && (self.Context == context) )
 	{
-		[context lock];
+		LockGLContext(context);
 	}
 	return bOK;
 }
@@ -396,7 +420,7 @@ bool PlatformBlitToViewport(FPlatformOpenGLDevice* Device, const FOpenGLViewport
 	check(Context && Context->OpenGLView);
 
 	FScopeLock ScopeLock(Device->ContextUsageGuard);
-	[Context->OpenGLContext lock];
+	LockGLContext(Context->OpenGLContext);
 	{
 		FScopeContext ScopeContext(Context->OpenGLContext);
 
@@ -449,7 +473,7 @@ bool PlatformBlitToViewport(FPlatformOpenGLDevice* Device, const FOpenGLViewport
 			}
 		}
 	}
-	[Context->OpenGLContext unlock];
+	UnlockGLContext(Context->OpenGLContext);
 
 	return !Viewport.GetCustomPresent();
 }
@@ -474,7 +498,7 @@ void DrawOpenGLViewport(FPlatformOpenGLContext* const Context, uint32 Width, uin
 		glClear(GL_COLOR_BUFFER_BIT);
 	}
 	[Context->OpenGLContext flushBuffer];
-	[Context->OpenGLContext unlock];
+	UnlockGLContext(Context->OpenGLContext);
 }
 
 void PlatformRenderingContextSetup(FPlatformOpenGLDevice* Device)
@@ -539,7 +563,7 @@ void PlatformRebindResources(FPlatformOpenGLDevice* Device)
 void PlatformResizeGLContext( FPlatformOpenGLDevice* Device, FPlatformOpenGLContext* Context, uint32 SizeX, uint32 SizeY, bool bFullscreen, bool bWasFullscreen, GLenum BackBufferTarget, GLuint BackBufferResource)
 {
 	FScopeLock ScopeLock(Device->ContextUsageGuard);
-	[Context->OpenGLContext lock];
+	LockGLContext(Context->OpenGLContext);
 	{
 		FScopeContext ScopeContext(Context->OpenGLContext);
 
@@ -600,7 +624,7 @@ void PlatformResizeGLContext( FPlatformOpenGLDevice* Device, FPlatformOpenGLCont
 		}
 #endif
 	}
-	[Context->OpenGLContext unlock];
+	UnlockGLContext(Context->OpenGLContext);
 }
 
 void PlatformGetSupportedResolution(uint32 &Width, uint32 &Height)
