Fix Intel HD 5000 driver crash - clear colours must be valid: [UE-14795] Crash occurred trying to close tutorial window.

diff --git a/Engine/Source/Runtime/OpenGLDrv/Private/Mac/MacOpenGLQuery.cpp b/Engine/Source/Runtime/OpenGLDrv/Private/Mac/MacOpenGLQuery.cpp
--- a/Engine/Source/Runtime/OpenGLDrv/Private/Mac/MacOpenGLQuery.cpp
+++ b/Engine/Source/Runtime/OpenGLDrv/Private/Mac/MacOpenGLQuery.cpp
@@ -334,9 +334,9 @@ FMacOpenGLQueryEmu::FMacOpenGLQueryEmu(FPlatformOpenGLContext* InContext)
 			GLuint Query = FreeQueries.FindOrAdd(GL_TIME_ELAPSED)[i];
 			glBeginQuery(GL_TIME_ELAPSED, Query);
 			
-			red = (float)(i & 0xff);
-			green = (float)((i >> 8) & 0xff);
-			blue = (float)((i >> 16) & 0xff);
+			red = (1.0f / (i + 1.0f));
+			green = (1.0f / (i + 2.0f));
+			blue = (1.0f / (i + 3.0f));
 			glClearColor(red, green, blue, 1.0);
 			glClear(GL_COLOR_BUFFER_BIT);
 			
