Fix android ES3.1 build

diff --git a/Engine/Source/Runtime/OpenGLDrv/Private/OpenGLES31.cpp b/Engine/Source/Runtime/OpenGLDrv/Private/OpenGLES31.cpp
--- a/Engine/Source/Runtime/OpenGLDrv/Private/OpenGLES31.cpp
+++ b/Engine/Source/Runtime/OpenGLDrv/Private/OpenGLES31.cpp
@@ -125,7 +125,7 @@ bool FOpenGLES31::bRequiresTextureCubeLodEXTToTextureCubeLodDefine = false;
 bool FOpenGLES31::bRequiresGLFragCoordVaryingLimitHack = false;
 
 /* This hack fixes an issue with SGX540 compiler which can get upset with some operations that mix highp and mediump */
-bool FOpenGLES2::bRequiresTexture2DPrecisionHack = false;
+bool FOpenGLES31::bRequiresTexture2DPrecisionHack = false;
 
 /** GL_EXT_disjoint_timer_query or GL_NV_timer_query*/
 bool FOpenGLES31::bSupportsDisjointTimeQueries = false;
