Custom present fix for Mac OpenGL viewport

diff --git a/Engine/Source/Runtime/OpenGLDrv/Private/Mac/OpenGLMac.cpp b/Engine/Source/Runtime/OpenGLDrv/Private/Mac/OpenGLMac.cpp
--- a/Engine/Source/Runtime/OpenGLDrv/Private/Mac/OpenGLMac.cpp
+++ b/Engine/Source/Runtime/OpenGLDrv/Private/Mac/OpenGLMac.cpp
@@ -430,6 +430,22 @@ bool PlatformBlitToViewport(FPlatformOpenGLDevice* Device, const FOpenGLViewport
 			[Context->OpenGLContext update];
 		}
 
+		if (Viewport.GetCustomPresent())
+		{
+			SCOPED_AUTORELEASE_POOL;
+
+			// Clear the Alpha channel
+			glColorMask(GL_FALSE, GL_FALSE, GL_FALSE, GL_TRUE);
+			glClearColor(0.f, 0.f, 0.f, 1.f);
+			glClear(GL_COLOR_BUFFER_BIT);
+			glColorMask(GL_TRUE, GL_TRUE, GL_TRUE, GL_TRUE);
+			glClearColor(0.f, 0.f, 0.f, 0.f);
+
+			glDisable(GL_FRAMEBUFFER_SRGB);
+			bPresent = Viewport.GetCustomPresent()->Present(SyncInterval);
+			glEnable(GL_FRAMEBUFFER_SRGB);
+		}
+
 		if (bPresent)
 		{
 			int32 RealSyncInterval = bLockToVsync ? SyncInterval : 0;
