Linux: OpenGL: working around crash in vertex attrib binding code (UE-15240).

- Disabled VAB for Linux again; not stable enough and causes crashes when painting foliage.
- Also some initialization cleanup.

diff --git a/Engine/Source/Runtime/Core/Private/Linux/LinuxApplication.cpp b/Engine/Source/Runtime/Core/Private/Linux/LinuxApplication.cpp
--- a/Engine/Source/Runtime/Core/Private/Linux/LinuxApplication.cpp
+++ b/Engine/Source/Runtime/Core/Private/Linux/LinuxApplication.cpp
@@ -47,12 +47,15 @@ FLinuxApplication* FLinuxApplication::CreateLinuxApplication()
 
 	LinuxApplication = new FLinuxApplication();
 
-	SDLControllerState *controllerState = LinuxApplication->ControllerStates;
-	for (int i = 0; i < SDL_NumJoysticks(); ++i) {
-	    if (SDL_IsGameController(i)) {
-		    controllerState->controller = SDL_GameControllerOpen(i);
-			if ( controllerState++->controller == NULL ) {
-				UE_LOG(LogLoad, Warning, TEXT("Could not open gamecontroller %i: %s\n"), i, SDL_GetError() );
+	SDLControllerState* ControllerState = LinuxApplication->ControllerStates;
+	for (int i = 0; i < SDL_NumJoysticks(); ++i)
+	{
+		if (SDL_IsGameController(i))
+		{
+			ControllerState->controller = SDL_GameControllerOpen(i);
+			if (ControllerState++->controller == nullptr)
+			{
+				UE_LOG(LogLoad, Warning, TEXT("Could not open gamecontroller %i: %s\n"), i, UTF8_TO_TCHAR(SDL_GetError()));
 			}
 		}
 	}
diff --git a/Engine/Source/Runtime/Core/Private/Linux/LinuxWindow.cpp b/Engine/Source/Runtime/Core/Private/Linux/LinuxWindow.cpp
--- a/Engine/Source/Runtime/Core/Private/Linux/LinuxWindow.cpp
+++ b/Engine/Source/Runtime/Core/Private/Linux/LinuxWindow.cpp
@@ -238,10 +238,10 @@ void FLinuxWindow::Initialize( FLinuxApplication* const Application, const TShar
 	// desired client area space.
 	ReshapeWindow( X, Y, ClientWidth, ClientHeight );
 
-	if( HWnd == NULL )
+	if (HWnd == nullptr)
 	{
 		// @todo Error message should be localized!
-		checkf(0, TEXT("Window Creation Failed (%d)"), SDL_GetError() );
+		checkf(false, TEXT("Window creation failed (%s)"), UTF8_TO_TCHAR(SDL_GetError()));
 		return;
 	}
 
diff --git a/Engine/Source/Runtime/OpenGLDrv/Private/Linux/OpenGLLinux.cpp b/Engine/Source/Runtime/OpenGLDrv/Private/Linux/OpenGLLinux.cpp
--- a/Engine/Source/Runtime/OpenGLDrv/Private/Linux/OpenGLLinux.cpp
+++ b/Engine/Source/Runtime/OpenGLDrv/Private/Linux/OpenGLLinux.cpp
@@ -208,10 +208,7 @@ struct FPlatformOpenGLDevice
 
         ContextUsageGuard = new FCriticalSection;
 
-		if	( SDL_GL_SetAttribute( SDL_GL_SHARE_WITH_CURRENT_CONTEXT, 0 ) )
-		{
-			verifyf( false, TEXT("OpenGLRHI: %s\n."), SDL_GetError() );
-		}
+		verifyf(SDL_GL_SetAttribute(SDL_GL_SHARE_WITH_CURRENT_CONTEXT, 0) == 0, TEXT("SDL_GL_SetAttribute(SDL_GL_SHARE_WITH_CURRENT_CONTEXT, 0) failed: %s\n."), UTF8_TO_TCHAR(SDL_GetError()));
 
 		_PlatformCreateDummyGLWindow( &SharedContext );
 		_PlatformCreateOpenGLContextCore( &SharedContext );
@@ -225,10 +222,8 @@ struct FPlatformOpenGLDevice
 			InitDefaultGLContextState();
 		}
 
-		if	( SDL_GL_SetAttribute( SDL_GL_SHARE_WITH_CURRENT_CONTEXT, 1 ) )
-		{
-			verifyf( false, TEXT("OpenGLRHI: %s\n."), SDL_GetError() );
-		}
+		verifyf(SDL_GL_SetAttribute(SDL_GL_SHARE_WITH_CURRENT_CONTEXT, 1) == 0, TEXT("SDL_GL_SetAttribute(SDL_GL_SHARE_WITH_CURRENT_CONTEXT, 1) failed: %s\n."), UTF8_TO_TCHAR(SDL_GetError()));
+
 		_ContextMakeCurrent( SharedContext.hWnd, SharedContext.hGLContext );
 
 		_PlatformCreateDummyGLWindow( &RenderingContext );
@@ -282,10 +277,7 @@ FPlatformOpenGLContext* PlatformCreateOpenGLContext(FPlatformOpenGLDevice* Devic
 	check( Device->SharedContext.hGLContext )
 	{
 		FScopeContext Scope(&(Device->SharedContext));
-		if	( SDL_GL_SetAttribute( SDL_GL_SHARE_WITH_CURRENT_CONTEXT, 1 ) )
-		{
-			verifyf( false, TEXT("OpenGLRHI: %s\n."), SDL_GetError() );
-		}
+		verifyf(SDL_GL_SetAttribute(SDL_GL_SHARE_WITH_CURRENT_CONTEXT, 1) == 0, TEXT("SDL_GL_SetAttribute(SDL_GL_SHARE_WITH_CURRENT_CONTEXT, 1) failed: %s\n."), UTF8_TO_TCHAR(SDL_GetError()));
 		_PlatformCreateOpenGLContextCore( Context );
 	}
 
@@ -806,7 +798,7 @@ bool PlatformInitOpenGL()
 	{
 		if (SDL_GL_LoadLibrary(NULL))
 		{
-			verifyf(false, TEXT("Unable to dynamically load libGL: %s\n."), SDL_GetError());
+			UE_LOG(LogLinux, Fatal, TEXT("Unable to dynamically load libGL: %s\n."), UTF8_TO_TCHAR(SDL_GetError()));
 		}
 
 		int MajorVersion = 0;
@@ -820,36 +812,31 @@ bool PlatformInitOpenGL()
 		}
 	
 		_PlatformOpenGLVersionFromCommandLine(MajorVersion, MinorVersion);
-		if	( SDL_GL_SetAttribute( SDL_GL_CONTEXT_MAJOR_VERSION, MajorVersion ) )
+		if (SDL_GL_SetAttribute( SDL_GL_CONTEXT_MAJOR_VERSION, MajorVersion) != 0)
 		{
-			verifyf( false, TEXT("OpenGLRHI: %s\n."), SDL_GetError() );
-			UE_LOG(LogLinux, Fatal, TEXT("SDL error errno=%d (%s)"), SDL_GetError());
+			UE_LOG(LogLinux, Fatal, TEXT("SDL_GL_SetAttribute(SDL_GL_CONTEXT_MAJOR_VERSION, %d) failed: %s"), MajorVersion, UTF8_TO_TCHAR(SDL_GetError()));
 		}
 
-		if	( SDL_GL_SetAttribute( SDL_GL_CONTEXT_MINOR_VERSION, MinorVersion ) )
+		if (SDL_GL_SetAttribute(SDL_GL_CONTEXT_MINOR_VERSION, MinorVersion) != 0)
 		{
-			verifyf( false, TEXT("OpenGLRHI: %s\n."), SDL_GetError() );
-			UE_LOG(LogLinux, Fatal, TEXT("SDL error errno=%d (%s)"), SDL_GetError());
+			UE_LOG(LogLinux, Fatal, TEXT("SDL_GL_SetAttribute(SDL_GL_CONTEXT_MINOR_VERSION, %d) failed: %s"), MinorVersion, UTF8_TO_TCHAR(SDL_GetError()));
 		}
 
-		if	( SDL_GL_SetAttribute( SDL_GL_CONTEXT_FLAGS, DebugFlag ) )
+		if (SDL_GL_SetAttribute(SDL_GL_CONTEXT_FLAGS, DebugFlag) != 0)
 		{
-			verifyf( false, TEXT("OpenGLRHI: %s\n."), SDL_GetError() );
-			UE_LOG(LogLinux, Fatal, TEXT("SDL error errno=%d (%s)"), SDL_GetError());
+			UE_LOG(LogLinux, Fatal, TEXT("SDL_GL_SetAttribute(SDL_GL_CONTEXT_FLAGS, %d) failed: %s"), DebugFlag, UTF8_TO_TCHAR(SDL_GetError()));
 		}
 
-		if	( SDL_GL_SetAttribute( SDL_GL_CONTEXT_PROFILE_MASK, SDL_GL_CONTEXT_PROFILE_CORE ) )
+		if (SDL_GL_SetAttribute(SDL_GL_CONTEXT_PROFILE_MASK, SDL_GL_CONTEXT_PROFILE_CORE) != 0)
 		{
-			verifyf( false, TEXT("OpenGLRHI: %s\n."), SDL_GetError() );
-			UE_LOG(LogLinux, Fatal, TEXT("SDL error errno=%d (%s)"), SDL_GetError());
+			UE_LOG(LogLinux, Fatal, TEXT("SDL_GL_SetAttribute(SDL_GL_CONTEXT_PROFILE_MASK, SDL_GL_CONTEXT_PROFILE_CORE) failed: %s"), UTF8_TO_TCHAR(SDL_GetError()));
 		}
 
 		if (FParse::Param(FCommandLine::Get(), TEXT("quad_buffer_stereo")))
 		{
-			if (SDL_GL_SetAttribute(SDL_GL_STEREO, 1))
+			if (SDL_GL_SetAttribute(SDL_GL_STEREO, 1) != 0)
 			{
-				verifyf(false, TEXT("OpenGLRHI: %s\n."), SDL_GetError());
-				UE_LOG(LogLinux, Fatal, TEXT("SDL error errno=%d (%s)"), SDL_GetError());
+				UE_LOG(LogLinux, Fatal, TEXT("SDL_GL_SetAttribute(SDL_GL_STEREO, 1) failed: %s"), UTF8_TO_TCHAR(SDL_GetError()));
 			}
 		}
 
diff --git a/Engine/Source/Runtime/OpenGLDrv/Private/OpenGLCommands.cpp b/Engine/Source/Runtime/OpenGLDrv/Private/OpenGLCommands.cpp
--- a/Engine/Source/Runtime/OpenGLDrv/Private/OpenGLCommands.cpp
+++ b/Engine/Source/Runtime/OpenGLDrv/Private/OpenGLCommands.cpp
@@ -53,7 +53,7 @@ namespace OpenGLConsoleVariables
 		TEXT("If true, don't issue dispatch work.")
 		);
 
-	int32 bUseVAB = 1;
+	int32 bUseVAB = (PLATFORM_LINUX) ? 0 : 1;	// disable VAB on Linux until vertex attrib binding code path is fixed, see UE-15240
 	static FAutoConsoleVariableRef CVarUseVAB(
 		TEXT("OpenGL.UseVAB"),
 		bUseVAB,
@@ -1922,7 +1922,7 @@ void FOpenGLDynamicRHI::SetupVertexArrays(FOpenGLContextState& ContextState, uin
 {
 	SCOPE_CYCLE_COUNTER_DETAILED(STAT_OpenGLVBOSetupTime);
 	if (FOpenGL::SupportsVertexAttribBinding() && OpenGLConsoleVariables::bUseVAB)
-{
+	{
 		SetupVertexArraysVAB(ContextState, BaseVertexIndex, Streams, NumStreams, MaxVertices);
 		return;
 	}
@@ -2090,7 +2090,7 @@ void FOpenGLDynamicRHI::SetupVertexArraysVAB(FOpenGLContextState& ContextState,
 				else
 				{
 					// bogus stream, make sure current value is zero to match D3D
-					float data[4] = { 0.0f };
+					static float data[4] = { 0.0f };
 
 					glVertexAttrib4fv(AttributeIndex, data);
 
