Fixed compile warning

diff --git a/Engine/Source/Runtime/OpenGLDrv/Private/OpenGLES31.cpp b/Engine/Source/Runtime/OpenGLDrv/Private/OpenGLES31.cpp
--- a/Engine/Source/Runtime/OpenGLDrv/Private/OpenGLES31.cpp
+++ b/Engine/Source/Runtime/OpenGLDrv/Private/OpenGLES31.cpp
@@ -142,18 +142,18 @@ GLint FOpenGLES31::MinorVersion = 0;
 bool FOpenGLES31::SupportsAdvancedFeatures()
 {
 	bool bResult = true;
-	GLint MajorVersion = 0;
-	GLint MinorVersion = 0;
+	GLint LocalMajorVersion = 0;
+	GLint LocalMinorVersion = 0;
 
 	const FString ExtensionsString = ANSI_TO_TCHAR((const ANSICHAR*)glGetString(GL_EXTENSIONS));
 
 	if (FString(ANSI_TO_TCHAR((const ANSICHAR*)glGetString(GL_VERSION))).Contains(TEXT("OpenGL ES 3.")))
 	{
-		glGetIntegerv(GL_MAJOR_VERSION, &MajorVersion);
-		glGetIntegerv(GL_MINOR_VERSION, &MinorVersion);
+		glGetIntegerv(GL_MAJOR_VERSION, &LocalMajorVersion);
+		glGetIntegerv(GL_MINOR_VERSION, &LocalMinorVersion);
 
 		// Check for minimum ES 3.1 + extensions support to avoid the ES2 fallback
-		bResult = (MajorVersion == 3 && MinorVersion >= 1);
+		bResult = (LocalMajorVersion == 3 && LocalMinorVersion >= 1);
 
 		bResult &= ExtensionsString.Contains(TEXT("GL_ANDROID_extension_pack_es31a"));
 		bResult &= ExtensionsString.Contains(TEXT("GL_EXT_color_buffer_half_float"));
