Fixed releasing queries on Android ES 3.1

diff --git a/Engine/Source/Runtime/OpenGLDrv/Private/Android/AndroidES31OpenGL.cpp b/Engine/Source/Runtime/OpenGLDrv/Private/Android/AndroidES31OpenGL.cpp
--- a/Engine/Source/Runtime/OpenGLDrv/Private/Android/AndroidES31OpenGL.cpp
+++ b/Engine/Source/Runtime/OpenGLDrv/Private/Android/AndroidES31OpenGL.cpp
@@ -498,7 +498,7 @@ void PlatformReleaseRenderQuery(GLuint Query, uint64 QueryContext)
 	EGLContext Context = eglGetCurrentContext();
 	if ((uint64)Context == QueryContext)
 	{
-		FOpenGL::GenQueries(1, &Query);
+		FOpenGL::DeleteQueries(1, &Query);
 	}
 	else
 	{
