#ue4 Fix for build break because of templated virtual destructor in TOpenGLTexture

diff --git a/Engine/Source/Runtime/OpenGLDrv/Private/OpenGLTexture.cpp b/Engine/Source/Runtime/OpenGLDrv/Private/OpenGLTexture.cpp
--- a/Engine/Source/Runtime/OpenGLDrv/Private/OpenGLTexture.cpp
+++ b/Engine/Source/Runtime/OpenGLDrv/Private/OpenGLTexture.cpp
@@ -1392,54 +1392,6 @@ void TOpenGLTexture<RHIResourceType>::CloneViaPBO( TOpenGLTexture<RHIResourceTyp
 	CachedBindPixelUnpackBuffer(0);
 }
 
-template<typename RHIResourceType>
-TOpenGLTexture<RHIResourceType>::~TOpenGLTexture()
-{
-	if ( GIsRHIInitialized )
-	{
-		VERIFY_GL_SCOPE();
-
-		OpenGLTextureDeleted( this );
-
-		if( Resource != 0 )
-		{
-			switch(Target)
-			{
-				case GL_TEXTURE_2D:
- 				case GL_TEXTURE_2D_MULTISAMPLE:
-				case GL_TEXTURE_3D:
-				case GL_TEXTURE_CUBE_MAP:
-				case GL_TEXTURE_2D_ARRAY:
-				case GL_TEXTURE_CUBE_MAP_ARRAY:
-				{
-					OpenGLRHI->InvalidateTextureResourceInCache( Resource );
-					FOpenGL::DeleteTextures( 1, &Resource );
-					break;
-				}
-				case GL_RENDERBUFFER:
-				{
-					if (!(this->GetFlags() & TexCreate_Presentable))
-					{
-						glDeleteRenderbuffers(1, &Resource);
-					}
-					break;
-				}
-				default:
-				{
-					checkNoEntry();
-				}
-			}
-		}
-		
-		if(TextureRange)
-		{
-			delete [] TextureRange;
-			TextureRange = nullptr;
-		}
-
-		ReleaseOpenGLFramebuffers(OpenGLRHI,this);
-	}
-}
 
 /*-----------------------------------------------------------------------------
 	2D texture support.
diff --git a/Engine/Source/Runtime/OpenGLDrv/Public/OpenGLResources.h b/Engine/Source/Runtime/OpenGLDrv/Public/OpenGLResources.h
--- a/Engine/Source/Runtime/OpenGLDrv/Public/OpenGLResources.h
+++ b/Engine/Source/Runtime/OpenGLDrv/Public/OpenGLResources.h
@@ -910,7 +910,53 @@ public:
 		}
 	}
 
-	virtual ~TOpenGLTexture();
+	virtual ~TOpenGLTexture()
+	{
+		if (GIsRHIInitialized)
+		{
+			VERIFY_GL_SCOPE();
+
+			OpenGLTextureDeleted(this);
+
+			if (Resource != 0)
+			{
+				switch (Target)
+				{
+					case GL_TEXTURE_2D:
+					case GL_TEXTURE_2D_MULTISAMPLE:
+					case GL_TEXTURE_3D:
+					case GL_TEXTURE_CUBE_MAP:
+					case GL_TEXTURE_2D_ARRAY:
+					case GL_TEXTURE_CUBE_MAP_ARRAY:
+					{
+						OpenGLRHI->InvalidateTextureResourceInCache(Resource);
+						FOpenGL::DeleteTextures(1, &Resource);
+						break;
+					}
+					case GL_RENDERBUFFER:
+					{
+						if (!(this->GetFlags() & TexCreate_Presentable))
+						{
+							glDeleteRenderbuffers(1, &Resource);
+						}
+						break;
+					}
+					default:
+					{
+						checkNoEntry();
+					}
+				}
+			}
+
+			if (TextureRange)
+			{
+				delete[] TextureRange;
+				TextureRange = nullptr;
+			}
+
+			ReleaseOpenGLFramebuffers(OpenGLRHI, this);
+		}
+	}
 
 	/**
 	 * Locks one of the texture's mip-maps.
