OpenGLRHI: fix Linux (and Mac?) build.

diff --git a/Engine/Source/Runtime/OpenGLDrv/Private/OpenGLTexture.cpp b/Engine/Source/Runtime/OpenGLDrv/Private/OpenGLTexture.cpp
--- a/Engine/Source/Runtime/OpenGLDrv/Private/OpenGLTexture.cpp
+++ b/Engine/Source/Runtime/OpenGLDrv/Private/OpenGLTexture.cpp
@@ -600,6 +600,11 @@ GLuint FOpenGLTextureBase::GetOpenGLFramebuffer(uint32 ArrayIndices, uint32 Mipm
 }
 #endif
 
+void FOpenGLTextureBase::InvalidateTextureResourceInCache()
+{
+	OpenGLRHI->InvalidateTextureResourceInCache(Resource);
+}
+
 template<typename RHIResourceType>
 void TOpenGLTexture<RHIResourceType>::Resolve(uint32 MipIndex,uint32 ArrayIndex)
 {
diff --git a/Engine/Source/Runtime/OpenGLDrv/Public/OpenGLResources.h b/Engine/Source/Runtime/OpenGLDrv/Public/OpenGLResources.h
--- a/Engine/Source/Runtime/OpenGLDrv/Public/OpenGLResources.h
+++ b/Engine/Source/Runtime/OpenGLDrv/Public/OpenGLResources.h
@@ -815,6 +815,8 @@ public:
 	GLuint GetOpenGLFramebuffer(uint32 ArrayIndices, uint32 MipmapLevels);
 #endif
 
+	void InvalidateTextureResourceInCache();
+
 private:
 	uint32 MemorySize		: 31;
 	uint32 bIsPowerOfTwo	: 1;
@@ -929,7 +931,7 @@ public:
 					case GL_TEXTURE_2D_ARRAY:
 					case GL_TEXTURE_CUBE_MAP_ARRAY:
 					{
-						OpenGLRHI->InvalidateTextureResourceInCache(Resource);
+						InvalidateTextureResourceInCache();
 						FOpenGL::DeleteTextures(1, &Resource);
 						break;
 					}
