#html5 don't use custom version macros, use inbuilt version macros instead.

diff --git a/Engine/Source/Runtime/OpenGLDrv/Private/HTML5/HTML5OpenGL.cpp b/Engine/Source/Runtime/OpenGLDrv/Private/HTML5/HTML5OpenGL.cpp
--- a/Engine/Source/Runtime/OpenGLDrv/Private/HTML5/HTML5OpenGL.cpp
+++ b/Engine/Source/Runtime/OpenGLDrv/Private/HTML5/HTML5OpenGL.cpp
@@ -361,7 +361,8 @@ extern "C"
 
 	EM_BOOL request_fullscreen_callback(int eventType, const EmscriptenMouseEvent* evt, void* user)
 	{
-#if __EMCC_VER__ >= 1290
+
+#if __EMSCRIPTEN_major__  >= 1 && __EMSCRIPTEN_minor__  >= 29 && __EMSCRIPTEN_tiny__  >= 0
 		EmscriptenFullscreenStrategy FSStrat;
 		FMemory::Memzero(FSStrat);
 		FSStrat.scaleMode = EMSCRIPTEN_FULLSCREEN_SCALE_STRETCH;//EMSCRIPTEN_FULLSCREEN_SCALE_ASPECT;// : EMSCRIPTEN_FULLSCREEN_SCALE_STRETCH;
