Fixed invisible windows on OS X 10.11

diff --git a/Engine/Source/Runtime/OpenGLDrv/Private/Mac/OpenGLMac.cpp b/Engine/Source/Runtime/OpenGLDrv/Private/Mac/OpenGLMac.cpp
--- a/Engine/Source/Runtime/OpenGLDrv/Private/Mac/OpenGLMac.cpp
+++ b/Engine/Source/Runtime/OpenGLDrv/Private/Mac/OpenGLMac.cpp
@@ -456,7 +456,7 @@ bool PlatformBlitToViewport(FPlatformOpenGLDevice* Device, const FOpenGLViewport
 				Context->SyncInterval = RealSyncInterval;
 			}
 
-			[(FCocoaWindow*)Context->WindowHandle startRendering];
+			MainThreadCall(^{ [(FCocoaWindow*)Context->WindowHandle startRendering]; }, NSDefaultRunLoopMode, false);
 			
 			int32 CurrentReadFramebuffer = 0;
 			glGetIntegerv(GL_READ_FRAMEBUFFER_BINDING, &CurrentReadFramebuffer);
