OpenGL: downgrade obscure and harmless "warning" to verbose (UE-17528).

diff --git a/Engine/Source/Developer/StandaloneRenderer/Private/Linux/OpenGL/SlateOpenGLContext.cpp b/Engine/Source/Developer/StandaloneRenderer/Private/Linux/OpenGL/SlateOpenGLContext.cpp
--- a/Engine/Source/Developer/StandaloneRenderer/Private/Linux/OpenGL/SlateOpenGLContext.cpp
+++ b/Engine/Source/Developer/StandaloneRenderer/Private/Linux/OpenGL/SlateOpenGLContext.cpp
@@ -263,7 +263,7 @@ void FSlateOpenGLContext::Initialize( void* InWindow, const FSlateOpenGLContext*
 		#ifdef GL_KHR_debug
 				glDebugMessageControlARB(GL_DEBUG_SOURCE_API_ARB, GL_DEBUG_TYPE_OTHER_ARB, GL_DEBUG_SEVERITY_NOTIFICATION, 0, NULL, GL_FALSE);
 		#endif
-				UE_LOG(LogStandaloneRenderer,Warning,TEXT("disabling reporting back of debug groups and markers to the OpenGL debug output callback"));
+				UE_LOG(LogStandaloneRenderer,Verbose,TEXT("disabling reporting back of debug groups and markers to the OpenGL debug output callback"));
 			}
 		#endif // GL_ARB_debug_output || GL_KHR_debug
 		}
diff --git a/Engine/Source/Runtime/OpenGLDrv/Private/OpenGLDevice.cpp b/Engine/Source/Runtime/OpenGLDrv/Private/OpenGLDevice.cpp
--- a/Engine/Source/Runtime/OpenGLDrv/Private/OpenGLDevice.cpp
+++ b/Engine/Source/Runtime/OpenGLDrv/Private/OpenGLDevice.cpp
@@ -463,7 +463,7 @@ void InitDebugContext()
 #ifdef GL_KHR_debug
 		glDebugMessageControlARB(GL_DEBUG_SOURCE_API_ARB, GL_DEBUG_TYPE_OTHER_ARB, GL_DEBUG_SEVERITY_NOTIFICATION, 0, NULL, GL_FALSE);
 #endif
-		UE_LOG(LogRHI,Warning,TEXT("disabling reporting back of debug groups and markers to the OpenGL debug output callback"));
+		UE_LOG(LogRHI,Verbose,TEXT("disabling reporting back of debug groups and markers to the OpenGL debug output callback"));
 	}
 #elif ENABLE_OPENGL_DEBUG_GROUPS && !defined(GL_ARB_debug_output) && GL_KHR_debug
 	if(glDebugMessageControlKHR)
@@ -472,7 +472,7 @@ void InitDebugContext()
 		glDebugMessageControlKHR(GL_DEBUG_SOURCE_APPLICATION_KHR, GL_DEBUG_TYPE_PUSH_GROUP, GL_DONT_CARE, 0, NULL, GL_FALSE);
 		glDebugMessageControlKHR(GL_DEBUG_SOURCE_APPLICATION_KHR, GL_DEBUG_TYPE_POP_GROUP, GL_DONT_CARE, 0, NULL, GL_FALSE);
 		glDebugMessageControlKHR(GL_DEBUG_SOURCE_API_KHR, GL_DEBUG_TYPE_OTHER_KHR, GL_DEBUG_SEVERITY_NOTIFICATION, 0, NULL, GL_FALSE);
-		UE_LOG(LogRHI,Warning,TEXT("disabling reporting back of debug groups and markers to the OpenGL debug output callback"));
+		UE_LOG(LogRHI,Verbose,TEXT("disabling reporting back of debug groups and markers to the OpenGL debug output callback"));
 	}
 #endif
 }
