Removed Adreno-specific check to use glTexStorage2D to create FB16 framebuffer on Android, so other GPUs can also use that path.

diff --git a/Engine/Source/Runtime/OpenGLDrv/Private/Android/AndroidOpenGL.cpp b/Engine/Source/Runtime/OpenGLDrv/Private/Android/AndroidOpenGL.cpp
--- a/Engine/Source/Runtime/OpenGLDrv/Private/Android/AndroidOpenGL.cpp
+++ b/Engine/Source/Runtime/OpenGLDrv/Private/Android/AndroidOpenGL.cpp
@@ -303,7 +303,7 @@ void PlatformReleaseRenderQuery( GLuint Query, uint64 QueryContext )
 }
 
 
-bool FAndroidOpenGL::bUseAdrenoHalfFloatTexStorage = false;
+bool FAndroidOpenGL::bUseHalfFloatTexStorage = false;
 bool FAndroidOpenGL::bUseES30ShadingLanguage = false;
 
 void FAndroidOpenGL::ProcessExtensions(const FString& ExtensionsString)
@@ -361,7 +361,7 @@ void FAndroidOpenGL::ProcessExtensions(const FString& ExtensionsString)
 		bRequiresTexture2DPrecisionHack = true;
 	}
 
-	if( RendererString.Contains(TEXT("Adreno")) )
+	if (RendererString.Contains(TEXT("Adreno")))
 	{
 		// Adreno 2xx doesn't work with packed depth stencil enabled
 		if (RendererString.Contains(TEXT("Adreno (TM) 2")))
@@ -369,17 +369,20 @@ void FAndroidOpenGL::ProcessExtensions(const FString& ExtensionsString)
 			UE_LOG(LogRHI, Warning, TEXT("Disabling support for GL_OES_packed_depth_stencil on Adreno 2xx"));
 			bSupportsPackedDepthStencil = false;
 		}
+	}
 
-		// Attempt to find ES 3.0 glTexStorage2D if we're on an Adreno device that supports it.
+	if (bES30Support)
+	{
+		// Attempt to find ES 3.0 glTexStorage2D if we're on an ES 3.0 device
 		glTexStorage2D = (PFNGLTEXSTORAGE2DPROC)((void*)eglGetProcAddress("glTexStorage2D"));
 		if( glTexStorage2D != NULL )
 		{
-			bUseAdrenoHalfFloatTexStorage = true;
+			bUseHalfFloatTexStorage = true;
 		}
 		else
 		{
 			// need to disable GL_EXT_color_buffer_half_float support because we have no way to allocate the storage and the driver doesn't work without it.
-			UE_LOG(LogRHI,Warning,TEXT("Disabling support for GL_EXT_color_buffer_half_float to avoid an Adreno driver bug"));
+			UE_LOG(LogRHI,Warning,TEXT("Disabling support for GL_EXT_color_buffer_half_float as we cannot bind glTexStorage2D"));
 			bSupportsColorBufferHalfFloat = false;
 		}
 	}
diff --git a/Engine/Source/Runtime/OpenGLDrv/Private/Android/AndroidOpenGL.h b/Engine/Source/Runtime/OpenGLDrv/Private/Android/AndroidOpenGL.h
--- a/Engine/Source/Runtime/OpenGLDrv/Private/Android/AndroidOpenGL.h
+++ b/Engine/Source/Runtime/OpenGLDrv/Private/Android/AndroidOpenGL.h
@@ -209,7 +209,7 @@ struct FAndroidOpenGL : public FOpenGLES2
 
 	static FORCEINLINE bool TexStorage2D(GLenum Target, GLint Levels, GLint InternalFormat, GLsizei Width, GLsizei Height, GLenum Format, GLenum Type, uint32 Flags)
 	{
-		if( bUseAdrenoHalfFloatTexStorage && Type == GL_HALF_FLOAT_OES && (Flags & TexCreate_RenderTargetable) != 0 )
+		if( bUseHalfFloatTexStorage && Type == GL_HALF_FLOAT_OES && (Flags & TexCreate_RenderTargetable) != 0 )
 		{
 			glTexStorage2D(Target, Levels, InternalFormat, Width, Height);
 			VERIFY_GL(glTexStorage2D)
@@ -231,8 +231,8 @@ struct FAndroidOpenGL : public FOpenGLES2
 
 	static void ProcessExtensions(const FString& ExtensionsString);
 
-	// whether to use ES 3.0 function glTexStorage2D to allocate storage for GL_HALF_FLOAT_OES render target textures, working around a driver limitation on Adreno 3x0.
-	static bool bUseAdrenoHalfFloatTexStorage;
+	// whether to use ES 3.0 function glTexStorage2D to allocate storage for GL_HALF_FLOAT_OES render target textures
+	static bool bUseHalfFloatTexStorage;
 
 	// whether to use ES 3.0 shading language
 	static bool bUseES30ShadingLanguage;
