Remove redundant PBO checks in the UpdateTexture functions in the OpenGL RHI to resolve [UE-14274] "GitHub 1086 : PBO checks in OpenGL RHIUpdateTexture2D and RHIUpdateTexture3D"
#1086

diff --git a/Engine/Source/Runtime/OpenGLDrv/Private/OpenGLTexture.cpp b/Engine/Source/Runtime/OpenGLDrv/Private/OpenGLTexture.cpp
--- a/Engine/Source/Runtime/OpenGLDrv/Private/OpenGLTexture.cpp
+++ b/Engine/Source/Runtime/OpenGLDrv/Private/OpenGLTexture.cpp
@@ -2000,7 +2000,6 @@ void FOpenGLDynamicRHI::RHIUnlockTexture2DArray(FTexture2DArrayRHIParamRef Textu
 void FOpenGLDynamicRHI::RHIUpdateTexture2D(FTexture2DRHIParamRef TextureRHI,uint32 MipIndex,const FUpdateTextureRegion2D& UpdateRegion,uint32 SourcePitch,const uint8* SourceData)
 {
 	VERIFY_GL_SCOPE();
-	check( FOpenGL::SupportsPixelBuffers() );
 
 	FOpenGLTexture2D* Texture = ResourceCast(TextureRHI);
 
@@ -2033,7 +2032,7 @@ void FOpenGLDynamicRHI::RHIUpdateTexture2D(FTexture2DRHIParamRef TextureRHI,uint
 void FOpenGLDynamicRHI::RHIUpdateTexture3D(FTexture3DRHIParamRef TextureRHI,uint32 MipIndex,const FUpdateTextureRegion3D& UpdateRegion,uint32 SourceRowPitch,uint32 SourceDepthPitch,const uint8* SourceData)
 {
 	VERIFY_GL_SCOPE();
-	check( FOpenGL::SupportsPixelBuffers() && FOpenGL::SupportsTexture3D() );
+	check( FOpenGL::SupportsTexture3D() );
 	FOpenGLTexture3D* Texture = ResourceCast(TextureRHI);
 
 	// Use a texture stage that's not likely to be used for draws, to avoid waiting
