Reset OpenGL timestamp queries in EndRenderQuery so that the result is recached. This code was broken in CL #2489182 by removing the RHIResetRenderQuery call that was essential for them to work in their previous guise.

diff --git a/Engine/Source/Runtime/OpenGLDrv/Private/OpenGLQuery.cpp b/Engine/Source/Runtime/OpenGLDrv/Private/OpenGLQuery.cpp
--- a/Engine/Source/Runtime/OpenGLDrv/Private/OpenGLQuery.cpp
+++ b/Engine/Source/Runtime/OpenGLDrv/Private/OpenGLQuery.cpp
@@ -89,6 +89,7 @@ void FOpenGLDynamicRHI::RHIEndRenderQuery(FRenderQueryRHIParamRef QueryRHI)
 		else if(Query->QueryType == RQT_AbsoluteTime)
 		{
 			FOpenGL::QueryTimestampCounter(Query->Resource);
+			Query->bResultIsCached = false;
 		}
 	}
 }
