Disable GPU timing on Nvidia GPUs under OS X as it appears to be responsible for [UE-18045] - the code to pre-allocate the queries may be insufficient.

diff --git a/Engine/Source/Runtime/OpenGLDrv/Private/Mac/MacOpenGLQuery.cpp b/Engine/Source/Runtime/OpenGLDrv/Private/Mac/MacOpenGLQuery.cpp
--- a/Engine/Source/Runtime/OpenGLDrv/Private/Mac/MacOpenGLQuery.cpp
+++ b/Engine/Source/Runtime/OpenGLDrv/Private/Mac/MacOpenGLQuery.cpp
@@ -8,9 +8,20 @@
  OpenGL static variables.
  ------------------------------------------------------------------------------*/
 
-bool GIsEmulatingTimestamp = false;
+bool GIsEmulatingTimestamp = true;
 static const uint32 GMacQueryNameCacheSize = 32 * OPENGL_NAME_CACHE_SIZE;
 
+static int32 GMacIsEmulatingTimestamp = 1;
+static FAutoConsoleVariableRef CVarMacIsEmulatingTimestamp(
+	TEXT("r.Mac.EmulateTimestamps"),
+	GMacIsEmulatingTimestamp,
+	TEXT("If enabled emulates GL_TIMESTAMP with a continuous chain of GL_TIME_ELAPSED queries - not all GPU drivers support this.\n")
+	TEXT("0: Disabled\n")
+	TEXT("1: Enabled (AMD & Intel)\n")
+	TEXT("2: Force Enabled (AMD & Intel & Nvidia)"),
+	ECVF_RenderThreadSafe|ECVF_ReadOnly
+	);
+
 /*------------------------------------------------------------------------------
  OpenGL query emulation.
  ------------------------------------------------------------------------------*/
@@ -311,7 +322,7 @@ FMacOpenGLQueryEmu::FMacOpenGLQueryEmu(FPlatformOpenGLContext* InContext)
 	
 	// Only use the timestamp emulation in a non-shipping build - end-users shouldn't care about this profiling feature.
 #if (!UE_BUILD_SHIPPING)
-	GIsEmulatingTimestamp = !FParse::Param(FCommandLine::Get(), TEXT("DisableMacGPUTimestamp"));
+	GIsEmulatingTimestamp = (!IsRHIDeviceNVIDIA() || GMacIsEmulatingTimestamp == 2) && !FParse::Param(FCommandLine::Get(), TEXT("DisableMacGPUTimestamp"));
 #endif
 	
 	if ( GIsEmulatingTimestamp )
