Ensure the GRHIVendorId is set in FPlatformOpenGLContext::VerifyCurrentContext  if it hasn't already been.

diff --git a/Engine/Source/Runtime/OpenGLDrv/Private/Mac/MacOpenGLContext.cpp b/Engine/Source/Runtime/OpenGLDrv/Private/Mac/MacOpenGLContext.cpp
--- a/Engine/Source/Runtime/OpenGLDrv/Private/Mac/MacOpenGLContext.cpp
+++ b/Engine/Source/Runtime/OpenGLDrv/Private/Mac/MacOpenGLContext.cpp
@@ -310,6 +310,9 @@ void FPlatformOpenGLContext::Initialise( NSOpenGLContext* const SharedContext )
 	
 	CGLGetVirtualScreen((CGLContextObj)[OpenGLContext CGLContextObj], &VirtualScreen);
 	
+	// Update the context dependent state.
+	VerifyCurrentContext();
+	
 	EmulatedQueries = new FMacOpenGLQueryEmu(this);
 }
 
@@ -425,7 +428,7 @@ void FPlatformOpenGLContext::VerifyCurrentContext()
 		CGLGetVirtualScreen(Current, &VirtualScreen);
 		
 		// Verify that the renderer ID hasn't changed, as if it has our capabilities may also be different
-		if(RendererID != PlatformContext->RendererID || VirtualScreen != PlatformContext->VirtualScreen)
+		if(!GRHIVendorId || RendererID != PlatformContext->RendererID || VirtualScreen != PlatformContext->VirtualScreen)
 		{
 			PlatformContext->RendererID = RendererID;
 			PlatformContext->VirtualScreen = VirtualScreen;
