Fixing the FP16 path so that is once again works on pre-ES3 Adreno based devices, such as the Galaxy S4

diff --git a/Engine/Source/Runtime/OpenGLDrv/Private/Android/AndroidOpenGL.cpp b/Engine/Source/Runtime/OpenGLDrv/Private/Android/AndroidOpenGL.cpp
--- a/Engine/Source/Runtime/OpenGLDrv/Private/Android/AndroidOpenGL.cpp
+++ b/Engine/Source/Runtime/OpenGLDrv/Private/Android/AndroidOpenGL.cpp
@@ -361,7 +361,8 @@ void FAndroidOpenGL::ProcessExtensions(const FString& ExtensionsString)
 		bRequiresTexture2DPrecisionHack = true;
 	}
 
-	if (RendererString.Contains(TEXT("Adreno")))
+	const bool bIsAdrenoBased = RendererString.Contains(TEXT("Adreno"));
+	if (bIsAdrenoBased)
 	{
 		// Adreno 2xx doesn't work with packed depth stencil enabled
 		if (RendererString.Contains(TEXT("Adreno (TM) 2")))
@@ -371,7 +372,7 @@ void FAndroidOpenGL::ProcessExtensions(const FString& ExtensionsString)
 		}
 	}
 
-	if (bES30Support)
+	if (bES30Support || bIsAdrenoBased)
 	{
 		// Attempt to find ES 3.0 glTexStorage2D if we're on an ES 3.0 device
 		glTexStorage2D = (PFNGLTEXSTORAGE2DPROC)((void*)eglGetProcAddress("glTexStorage2D"));
