Only enable Mac GPU timestamp emulation on Yosemite or later, there were driver bugs still present in Mavericks that cause it to fail intermittently.

diff --git a/Engine/Source/Runtime/OpenGLDrv/Private/Mac/MacOpenGLQuery.cpp b/Engine/Source/Runtime/OpenGLDrv/Private/Mac/MacOpenGLQuery.cpp
--- a/Engine/Source/Runtime/OpenGLDrv/Private/Mac/MacOpenGLQuery.cpp
+++ b/Engine/Source/Runtime/OpenGLDrv/Private/Mac/MacOpenGLQuery.cpp
@@ -322,7 +322,7 @@ FMacOpenGLQueryEmu::FMacOpenGLQueryEmu(FPlatformOpenGLContext* InContext)
 	
 	// Only use the timestamp emulation in a non-shipping build - end-users shouldn't care about this profiling feature.
 #if (!UE_BUILD_SHIPPING)
-	GIsEmulatingTimestamp = (!IsRHIDeviceNVIDIA() || GMacIsEmulatingTimestamp == 2) && !FParse::Param(FCommandLine::Get(), TEXT("DisableMacGPUTimestamp"));
+	GIsEmulatingTimestamp = ((FPlatformMisc::MacOSXVersionCompare(10,10,0) >= 0 || !IsRHIDeviceNVIDIA()) || GMacIsEmulatingTimestamp == 2) && !FParse::Param(FCommandLine::Get(), TEXT("DisableMacGPUTimestamp"));
 #endif
 	
 	if ( GIsEmulatingTimestamp )
