Some OpenGL ES 3.1 fixes

diff --git a/Engine/Source/Runtime/OpenGLDrv/Private/OpenGLDevice.cpp b/Engine/Source/Runtime/OpenGLDrv/Private/OpenGLDevice.cpp
--- a/Engine/Source/Runtime/OpenGLDrv/Private/OpenGLDevice.cpp
+++ b/Engine/Source/Runtime/OpenGLDrv/Private/OpenGLDevice.cpp
@@ -425,29 +425,31 @@ void InitDebugContext()
 	// Set the debug output callback if the driver supports it.
 	VERIFY_GL(__FUNCTION__);
 	bool bDebugOutputInitialized = false;
-#if defined(GL_ARB_debug_output)
-	if (glDebugMessageCallbackARB)
-	{
-		// Synchronous output can slow things down, but we'll get better callstack if breaking in or crashing in the callback. This is debug only after all.
-		glEnable(GL_DEBUG_OUTPUT_SYNCHRONOUS);
-		glDebugMessageCallbackARB(GLDEBUGPROCARB(OpenGLDebugMessageCallbackARB), /*UserParam=*/ NULL);
-		bDebugOutputInitialized = (glGetError() == GL_NO_ERROR);
-	}
-#elif defined(GL_KHR_debug)
-	// OpenGLES names the debug functions differently, but they behave the same
-	if (glDebugMessageCallbackKHR)
-	{
-		glDebugMessageCallbackKHR(GLDEBUGPROCKHR(OpenGLDebugMessageCallbackARB), /*UserParam=*/ NULL);
-		bDebugOutputInitialized = (glGetError() == GL_NO_ERROR);
-	}
-#endif // GL_ARB_debug_output / GL_KHR_debug
-#if defined(GL_AMD_debug_output)
-	if (glDebugMessageCallbackAMD && !bDebugOutputInitialized)
-	{
-		glDebugMessageCallbackAMD(OpenGLDebugMessageCallbackAMD, /*UserParam=*/ NULL);
-		bDebugOutputInitialized = (glGetError() == GL_NO_ERROR);
-	}
-#endif // GL_AMD_debug_output
+#if !ENABLE_VERIFY_GL
+	#if defined(GL_ARB_debug_output)
+		if (glDebugMessageCallbackARB)
+		{
+			// Synchronous output can slow things down, but we'll get better callstack if breaking in or crashing in the callback. This is debug only after all.
+			glEnable(GL_DEBUG_OUTPUT_SYNCHRONOUS);
+			glDebugMessageCallbackARB(GLDEBUGPROCARB(OpenGLDebugMessageCallbackARB), /*UserParam=*/ NULL);
+			bDebugOutputInitialized = (glGetError() == GL_NO_ERROR);
+		}
+	#elif defined(GL_KHR_debug)
+		// OpenGLES names the debug functions differently, but they behave the same
+		if (glDebugMessageCallbackKHR)
+		{
+			glDebugMessageCallbackKHR(GLDEBUGPROCKHR(OpenGLDebugMessageCallbackARB), /*UserParam=*/ NULL);
+			bDebugOutputInitialized = (glGetError() == GL_NO_ERROR);
+		}
+	#endif // GL_ARB_debug_output / GL_KHR_debug
+	#if defined(GL_AMD_debug_output)
+		if (glDebugMessageCallbackAMD && !bDebugOutputInitialized)
+		{
+			glDebugMessageCallbackAMD(OpenGLDebugMessageCallbackAMD, /*UserParam=*/ NULL);
+			bDebugOutputInitialized = (glGetError() == GL_NO_ERROR);
+		}
+	#endif // GL_AMD_debug_output
+#endif // !ENABLE_VERIFY_GL
 	if (!bDebugOutputInitialized && !PLATFORM_MAC)
 	{
 		UE_LOG(LogRHI,Warning,TEXT("OpenGL debug output extension not supported!"));
@@ -750,25 +752,37 @@ static void InitRHICapabilitiesForGL()
 		SetupTextureFormat( PF_A2B10G10R10,		FOpenGLTextureFormat( GL_RGB10_A2,				GL_RGB10_A2,			GL_RGBA,		GL_UNSIGNED_INT_2_10_10_10_REV,		false,	false));
 		SetupTextureFormat( PF_R16F,			FOpenGLTextureFormat( GL_R16F,					GL_R16F,				GL_RED,			GL_HALF_FLOAT,						false,	false));
 		SetupTextureFormat( PF_R16F_FILTER,		FOpenGLTextureFormat( GL_R16F,					GL_R16F,				GL_RED,			GL_HALF_FLOAT,						false,	false));
-		SetupTextureFormat( PF_FloatRGB,		FOpenGLTextureFormat( GL_R11F_G11F_B10F,		GL_R11F_G11F_B10F,		GL_RGB,			GL_UNSIGNED_INT_10F_11F_11F_REV,	false,	false));
-		SetupTextureFormat( PF_V8U8,			FOpenGLTextureFormat( GL_RG8_SNORM,				GL_NONE,				GL_RG,			GL_BYTE,							false,	false));
+		if (FOpenGL::SupportsR11G11B10F())
+		{
+			// Note: Also needs to include support for compute shaders to be defined here (e.g. glBindImageTexture)
+			SetupTextureFormat(PF_FloatRGB, FOpenGLTextureFormat(GL_R11F_G11F_B10F, GL_R11F_G11F_B10F, GL_RGB, GL_UNSIGNED_INT_10F_11F_11F_REV, false, false));
+			SetupTextureFormat(PF_FloatR11G11B10, FOpenGLTextureFormat(GL_RGBA16F, GL_RGBA16F, GL_RGB, GL_HALF_FLOAT, false, false));
+		}
+		else
+		{
+			SetupTextureFormat(PF_FloatRGB, FOpenGLTextureFormat(GL_RGBA16F, GL_RGBA16F, GL_RGB, GL_HALF_FLOAT, false, false));
+			SetupTextureFormat(PF_FloatR11G11B10, FOpenGLTextureFormat(GL_R11F_G11F_B10F, GL_R11F_G11F_B10F, GL_RGB, GL_UNSIGNED_INT_10F_11F_11F_REV, false, false));
+		}
+		SetupTextureFormat(PF_V8U8, FOpenGLTextureFormat(GL_RG8_SNORM, GL_NONE, GL_RG, GL_BYTE, false, false));
 		SetupTextureFormat( PF_R8G8,			FOpenGLTextureFormat( GL_RG8,					GL_NONE,				GL_RG,			GL_UNSIGNED_BYTE,					false,	false));
 		SetupTextureFormat( PF_BC5,				FOpenGLTextureFormat( GL_COMPRESSED_RG_RGTC2,	GL_COMPRESSED_RG_RGTC2,	GL_RG,			GL_UNSIGNED_BYTE,					true,	false));
-		SetupTextureFormat( PF_BC4,				FOpenGLTextureFormat( GL_COMPRESSED_RED_RGTC1,	GL_COMPRESSED_RED_RGTC1,	GL_RED,			GL_UNSIGNED_BYTE,					true,	false));
+		SetupTextureFormat( PF_BC4,				FOpenGLTextureFormat( GL_COMPRESSED_RED_RGTC1,	GL_COMPRESSED_RED_RGTC1,GL_RED,			GL_UNSIGNED_BYTE,					true,	false));
 		SetupTextureFormat( PF_A8,				FOpenGLTextureFormat( GL_R8,					GL_NONE,				GL_RED,			GL_UNSIGNED_BYTE,					false,	false));
 		SetupTextureFormat( PF_R32_UINT,		FOpenGLTextureFormat( GL_R32UI,					GL_NONE,				GL_RED_INTEGER,	GL_UNSIGNED_INT,					false,	false));
 		SetupTextureFormat( PF_R32_SINT,		FOpenGLTextureFormat( GL_R32I,					GL_NONE,				GL_RED_INTEGER,	GL_INT,								false,	false));
 		SetupTextureFormat( PF_R16_UINT,		FOpenGLTextureFormat( GL_R16UI,					GL_NONE,				GL_RED_INTEGER,	GL_UNSIGNED_SHORT,					false,	false));
 		SetupTextureFormat( PF_R16_SINT,		FOpenGLTextureFormat( GL_R16I,					GL_NONE,				GL_RED_INTEGER,	GL_SHORT,							false,	false));
 		SetupTextureFormat( PF_FloatRGBA,		FOpenGLTextureFormat( GL_RGBA16F,				GL_RGBA16F,				GL_RGBA,		GL_HALF_FLOAT,						false,	false));
-		SetupTextureFormat( PF_FloatR11G11B10,	FOpenGLTextureFormat( GL_R11F_G11F_B10F,		GL_R11F_G11F_B10F,		GL_RGB,			GL_UNSIGNED_INT_10F_11F_11F_REV,	false,	false));
 		if (FOpenGL::GetShaderPlatform() == EShaderPlatform::SP_OPENGL_ES31_EXT)
 		{
 			SetupTextureFormat(PF_G8, FOpenGLTextureFormat(GL_R8, GL_R8, GL_RED, GL_UNSIGNED_BYTE, false, false));
 			SetupTextureFormat(PF_B8G8R8A8, FOpenGLTextureFormat(GL_RGBA8, GL_SRGB8_ALPHA8, GL_RGBA, GL_UNSIGNED_BYTE, false, true));
 			SetupTextureFormat(PF_R8G8B8A8, FOpenGLTextureFormat(GL_RGBA8, GL_SRGB8_ALPHA8, GL_RGBA, GL_UNSIGNED_BYTE, false, false));
-			// GL_RG16 is not supported in OpenGL ES. There is currently no extension for it either. The user should check for support and implement a fallback.
-//			SetupTextureFormat(PF_G16R16, FOpenGLTextureFormat(GL_RG16, GL_RG16, GL_RG, GL_UNSIGNED_SHORT, false, false));
+			if (FOpenGL::SupportsRG16UI())
+			{
+				// The user should check for support for PF_G16R16 and implement a fallback if it's not supported!
+				SetupTextureFormat(PF_G16R16, FOpenGLTextureFormat(GL_RG16, GL_RG16, GL_RG, GL_UNSIGNED_SHORT, false, false));
+			}
 		}
 		else
 		{
diff --git a/Engine/Source/Runtime/OpenGLDrv/Private/OpenGLES31.cpp b/Engine/Source/Runtime/OpenGLDrv/Private/OpenGLES31.cpp
--- a/Engine/Source/Runtime/OpenGLDrv/Private/OpenGLES31.cpp
+++ b/Engine/Source/Runtime/OpenGLDrv/Private/OpenGLES31.cpp
@@ -64,6 +64,9 @@ bool FOpenGLES31::bSupportsTextureHalfFloat = false;
 /** GL_EXT_color_buffer_half_float */
 bool FOpenGLES31::bSupportsColorBufferHalfFloat = false;
 
+/** GL_NV_image_formats */
+bool FOpenGLES31::bSupportsNvImageFormats = false;
+
 /** GL_EXT_shader_framebuffer_fetch */
 bool FOpenGLES31::bSupportsShaderFramebufferFetch = false;
 
@@ -269,6 +272,7 @@ void FOpenGLES31::ProcessExtensions( const FString& ExtensionsString )
 	bSupportsTextureHalfFloat = ExtensionsString.Contains(TEXT("GL_OES_texture_half_float"));
 	bSupportsSGRB = ExtensionsString.Contains(TEXT("GL_EXT_sRGB"));
 	bSupportsColorBufferHalfFloat = ExtensionsString.Contains(TEXT("GL_EXT_color_buffer_half_float"));
+	bSupportsNvImageFormats = ExtensionsString.Contains(TEXT("GL_NV_image_formats"));
 	bSupportsShaderFramebufferFetch = ExtensionsString.Contains(TEXT("GL_EXT_shader_framebuffer_fetch")) || ExtensionsString.Contains(TEXT("GL_NV_shader_framebuffer_fetch"));
 	// @todo es3: SRGB support does not work with our texture format setup (ES2 docs indicate that internalFormat and format must match, but they don't at all with sRGB enabled)
 	//             One possible solution us to use GLFormat.InternalFormat[bSRGB] instead of GLFormat.Format
diff --git a/Engine/Source/Runtime/OpenGLDrv/Public/OpenGL.h b/Engine/Source/Runtime/OpenGLDrv/Public/OpenGL.h
--- a/Engine/Source/Runtime/OpenGLDrv/Public/OpenGL.h
+++ b/Engine/Source/Runtime/OpenGLDrv/Public/OpenGL.h
@@ -107,6 +107,8 @@ public:
 	static FORCEINLINE bool SupportsTextureFloat()						{ return true; }
 	static FORCEINLINE bool SupportsTextureHalfFloat()					{ return true; }
 	static FORCEINLINE bool SupportsColorBufferHalfFloat()				{ return true; }
+	static FORCEINLINE bool	SupportsRG16UI()							{ return true; }
+	static FORCEINLINE bool SupportsR11G11B10F()						{ return true; }
 	static FORCEINLINE bool SupportsVolumeTextureRendering()			{ return false; }
 	static FORCEINLINE bool SupportsShaderFramebufferFetch()			{ return false; }
 	static FORCEINLINE bool SupportsVertexArrayBGRA()					{ return true; }
diff --git a/Engine/Source/Runtime/OpenGLDrv/Public/OpenGLES2.h b/Engine/Source/Runtime/OpenGLDrv/Public/OpenGLES2.h
--- a/Engine/Source/Runtime/OpenGLDrv/Public/OpenGLES2.h
+++ b/Engine/Source/Runtime/OpenGLDrv/Public/OpenGLES2.h
@@ -105,6 +105,8 @@ struct FOpenGLES2 : public FOpenGLBase
 	static FORCEINLINE bool SupportsTextureFloat()						{ return bSupportsTextureFloat; }
 	static FORCEINLINE bool SupportsTextureHalfFloat()					{ return bSupportsTextureHalfFloat; }
 	static FORCEINLINE bool SupportsColorBufferHalfFloat()				{ return bSupportsColorBufferHalfFloat; }
+	static FORCEINLINE bool	SupportsRG16UI()							{ return false; }
+	static FORCEINLINE bool SupportsR11G11B10F()						{ return false; }
 	static FORCEINLINE bool SupportsShaderFramebufferFetch()			{ return bSupportsShaderFramebufferFetch; }
 	static FORCEINLINE bool SupportsMultisampledRenderToTexture()		{ return bSupportsMultisampledRenderToTexture; }
 	static FORCEINLINE bool SupportsVertexArrayBGRA()					{ return false; }
diff --git a/Engine/Source/Runtime/OpenGLDrv/Public/OpenGLES31.h b/Engine/Source/Runtime/OpenGLDrv/Public/OpenGLES31.h
--- a/Engine/Source/Runtime/OpenGLDrv/Public/OpenGLES31.h
+++ b/Engine/Source/Runtime/OpenGLDrv/Public/OpenGLES31.h
@@ -62,6 +62,8 @@ struct FOpenGLES31 : public FOpenGLBase
 	static FORCEINLINE bool SupportsTextureFloat()						{ return bSupportsTextureFloat || !bES2Fallback; }
 	static FORCEINLINE bool SupportsTextureHalfFloat()					{ return bSupportsTextureHalfFloat || !bES2Fallback; }
 	static FORCEINLINE bool SupportsColorBufferHalfFloat()				{ return bSupportsColorBufferHalfFloat || !bES2Fallback; }
+	static FORCEINLINE bool	SupportsRG16UI()							{ return bSupportsNvImageFormats && !bES2Fallback; }
+	static FORCEINLINE bool SupportsR11G11B10F()						{ return bSupportsNvImageFormats && !bES2Fallback; }
 	static FORCEINLINE bool SupportsShaderFramebufferFetch()			{ return bSupportsShaderFramebufferFetch; }
 	static FORCEINLINE bool SupportsVertexArrayBGRA()					{ return false; }
 	static FORCEINLINE bool SupportsBGRA8888()							{ return bSupportsBGRA8888; }
@@ -887,6 +889,9 @@ protected:
 	/** GL_EXT_color_buffer_half_float */
 	static bool bSupportsColorBufferHalfFloat;
 
+	/** GL_NV_image_formats */
+	static bool bSupportsNvImageFormats;
+
 	/** GL_EXT_shader_framebuffer_fetch */
 	static bool bSupportsShaderFramebufferFetch;
 
diff --git a/Engine/Source/Runtime/OpenGLDrv/Public/OpenGLUtil.h b/Engine/Source/Runtime/OpenGLDrv/Public/OpenGLUtil.h
--- a/Engine/Source/Runtime/OpenGLDrv/Public/OpenGLUtil.h
+++ b/Engine/Source/Runtime/OpenGLDrv/Public/OpenGLUtil.h
@@ -104,7 +104,9 @@ FORCEINLINE GLenum GetOpenGLCubeFace(ECubeFace Face)
 			}
 		}
 	};
-	#define VERIFY_GL_SCOPE_WITH_MSG_STR(MsgStr) FOpenGLErrorScope ErrorScope_##__LINE__ (MsgStr, TEXT(__FILE__), __LINE__)
+	#define MACRO_TOKENIZER(IdentifierName, Msg, FileName, LineNumber) FOpenGLErrorScope IdentifierName_ ## LineNumber (Msg, FileName, LineNumber)
+	#define MACRO_TOKENIZER2(IdentifierName, Msg, FileName, LineNumber) MACRO_TOKENIZER(IdentiferName, Msg, FileName, LineNumber)
+	#define VERIFY_GL_SCOPE_WITH_MSG_STR(MsgStr) MACRO_TOKENIZER2(ErrorScope_, MsgStr, TEXT(__FILE__), __LINE__)
 	#define VERIFY_GL_SCOPE() VERIFY_GL_SCOPE_WITH_MSG_STR(ANSI_TO_TCHAR(__FUNCTION__))
 	#define VERIFY_GL_FUNC(Func, ...) { VERIFY_GL_SCOPE_WITH_MSG_STR(TEXT(#Func)); Func(__VA_ARGS__); }
 
