Fix CIS

diff --git a/Engine/Source/Runtime/OpenGLDrv/Private/Android/AndroidOpenGL.cpp b/Engine/Source/Runtime/OpenGLDrv/Private/Android/AndroidOpenGL.cpp
--- a/Engine/Source/Runtime/OpenGLDrv/Private/Android/AndroidOpenGL.cpp
+++ b/Engine/Source/Runtime/OpenGLDrv/Private/Android/AndroidOpenGL.cpp
@@ -222,7 +222,7 @@ void PlatformDestroyOpenGLContext(FPlatformOpenGLDevice* Device, FPlatformOpenGL
 FRHITexture* PlatformCreateBuiltinBackBuffer(FOpenGLDynamicRHI* OpenGLRHI, uint32 SizeX, uint32 SizeY)
 {
 	uint32 Flags = TexCreate_RenderTargetable;
-	FOpenGLTexture2D* Texture2D = new FOpenGLTexture2D(OpenGLRHI, AndroidEGL::GetInstance()->GetOnScreenColorRenderBuffer(), GL_RENDERBUFFER, GL_COLOR_ATTACHMENT0, SizeX, SizeY, 0, 1, 1, 1, PF_B8G8R8A8, false, false, Flags, nullptr);
+	FOpenGLTexture2D* Texture2D = new FOpenGLTexture2D(OpenGLRHI, AndroidEGL::GetInstance()->GetOnScreenColorRenderBuffer(), GL_RENDERBUFFER, GL_COLOR_ATTACHMENT0, SizeX, SizeY, 0, 1, 1, 1, PF_B8G8R8A8, false, false, Flags, nullptr, FClearValueBinding::Transparent);
 	OpenGLTextureAllocated(Texture2D, Flags);
 
 	return Texture2D;
diff --git a/Engine/Source/Runtime/OpenGLDrv/Private/HTML5/HTML5OpenGL.cpp b/Engine/Source/Runtime/OpenGLDrv/Private/HTML5/HTML5OpenGL.cpp
--- a/Engine/Source/Runtime/OpenGLDrv/Private/HTML5/HTML5OpenGL.cpp
+++ b/Engine/Source/Runtime/OpenGLDrv/Private/HTML5/HTML5OpenGL.cpp
@@ -331,7 +331,7 @@ FRHITexture* PlatformCreateBuiltinBackBuffer(FOpenGLDynamicRHI* OpenGLRHI, uint3
 {
 	UE_LOG(LogHTML5OpenGL, Verbose, TEXT("PlatformCreateBuiltinBackBuffer(%d, %d)"), SizeX, SizeY);
 	uint32 Flags = TexCreate_RenderTargetable;
-	FOpenGLTexture2D* Texture2D = new FOpenGLTexture2D(OpenGLRHI, 0, GL_RENDERBUFFER, GL_COLOR_ATTACHMENT0, SizeX, SizeY, 0, 1, 1, 1, PF_B8G8R8A8, false, false, Flags, nullptr);
+	FOpenGLTexture2D* Texture2D = new FOpenGLTexture2D(OpenGLRHI, 0, GL_RENDERBUFFER, GL_COLOR_ATTACHMENT0, SizeX, SizeY, 0, 1, 1, 1, PF_B8G8R8A8, false, false, Flags, nullptr, FClearValueBinding::Transparent);
 	OpenGLTextureAllocated(Texture2D, Flags);
 
 	return Texture2D;
