Fix for shadowed variables and user defined literals (VS2015)

diff --git a/Engine/Source/Runtime/OpenGLDrv/Private/OpenGLDevice.cpp b/Engine/Source/Runtime/OpenGLDrv/Private/OpenGLDevice.cpp
--- a/Engine/Source/Runtime/OpenGLDrv/Private/OpenGLDevice.cpp
+++ b/Engine/Source/Runtime/OpenGLDrv/Private/OpenGLDevice.cpp
@@ -228,7 +228,7 @@ static const TCHAR* GetOpenGLDebugTypeStringARB(GLenum Type)
 	}
 #ifdef GL_KHR_debug	
 	{
-		static const TCHAR* TypeStrings[] =
+		static const TCHAR* DebugTypeStrings[] =
 		{
 			TEXT("Marker"),
 			TEXT("PushGroup"),
@@ -237,7 +237,7 @@ static const TCHAR* GetOpenGLDebugTypeStringARB(GLenum Type)
 
 		if (Type >= GL_DEBUG_TYPE_MARKER && Type <= GL_DEBUG_TYPE_POP_GROUP)
 		{
-			return TypeStrings[Type - GL_DEBUG_TYPE_MARKER];
+			return DebugTypeStrings[Type - GL_DEBUG_TYPE_MARKER];
 		}
 	}
 #endif
@@ -493,7 +493,7 @@ static void InitRHICapabilitiesForGL()
 #endif
 
 	// GL vendor and version information.
-#if !defined(__GNUC__) && !defined(__clang__)
+#if !defined(__GNUC__) && !defined(__clang__) && !(defined(_MSC_VER) && _MSC_VER >= 1900)
 	#define LOG_GL_STRING(StringEnum) UE_LOG(LogRHI, Log, TEXT("  ") ## TEXT(#StringEnum) ## TEXT(": %s"), ANSI_TO_TCHAR((const ANSICHAR*)glGetString(StringEnum)))
 #else
 	#define LOG_GL_STRING(StringEnum) UE_LOG(LogRHI, Log, TEXT("  " #StringEnum ": %s"), ANSI_TO_TCHAR((const ANSICHAR*)glGetString(StringEnum)))
@@ -562,7 +562,7 @@ static void InitRHICapabilitiesForGL()
 	FOpenGL::InitDebugContext();
 
 	// Log and get various limits.
-#if !defined(__GNUC__) && !defined(__clang__)
+#if !defined(__GNUC__) && !defined(__clang__) && !(defined(_MSC_VER) && _MSC_VER >= 1900)
 #define LOG_AND_GET_GL_INT_TEMP(IntEnum,Default) GLint Value_##IntEnum = Default; if (IntEnum) {glGetIntegerv(IntEnum, &Value_##IntEnum); glGetError();} else {Value_##IntEnum = Default;} UE_LOG(LogRHI, Log, TEXT("  ") ## TEXT(#IntEnum) ## TEXT(": %d"), Value_##IntEnum)
 #else
 #define LOG_AND_GET_GL_INT_TEMP(IntEnum,Default) GLint Value_##IntEnum = Default; if (IntEnum) {glGetIntegerv(IntEnum, &Value_##IntEnum); glGetError();} else {Value_##IntEnum = Default;} UE_LOG(LogRHI, Log, TEXT("  " #IntEnum ": %d"), Value_##IntEnum)
