Fix narrowing conversion for 32b constants. (VS2015)

diff --git a/Engine/Source/Runtime/OpenGLDrv/Public/OpenGLState.h b/Engine/Source/Runtime/OpenGLDrv/Public/OpenGLState.h
--- a/Engine/Source/Runtime/OpenGLDrv/Public/OpenGLState.h
+++ b/Engine/Source/Runtime/OpenGLDrv/Public/OpenGLState.h
@@ -210,8 +210,8 @@ struct FUAVStage
 	{
 	}
 };
-#define FOpenGLCachedAttr_Invalid (void*)0xFFFFFFFF
-#define FOpenGLCachedAttr_SingleVertex (void*)0xFFFFFFFE
+#define FOpenGLCachedAttr_Invalid (void*)(UPTRINT)0xFFFFFFFF
+#define FOpenGLCachedAttr_SingleVertex (void*)(UPTRINT)0xFFFFFFFE
 
 struct FOpenGLCachedAttr
 {
