Added missing clearvalue param to FOpenGLTexture2D creation call
#ue4
#android

diff --git a/Engine/Source/Runtime/OpenGLDrv/Private/Android/AndroidES31OpenGL.cpp b/Engine/Source/Runtime/OpenGLDrv/Private/Android/AndroidES31OpenGL.cpp
--- a/Engine/Source/Runtime/OpenGLDrv/Private/Android/AndroidES31OpenGL.cpp
+++ b/Engine/Source/Runtime/OpenGLDrv/Private/Android/AndroidES31OpenGL.cpp
@@ -349,7 +349,7 @@ FRHITexture* PlatformCreateBuiltinBackBuffer(FOpenGLDynamicRHI* OpenGLRHI, uint3
 	if ( FOpenGL::IsES2())
 	{
 		uint32 Flags = TexCreate_RenderTargetable;
-		Texture2D = new FOpenGLTexture2D(OpenGLRHI, AndroidEGL::GetInstance()->GetOnScreenColorRenderBuffer(), GL_RENDERBUFFER, GL_COLOR_ATTACHMENT0, SizeX, SizeY, 0, 1, 1, 1, PF_B8G8R8A8, false, false, Flags, nullptr);
+		Texture2D = new FOpenGLTexture2D(OpenGLRHI, AndroidEGL::GetInstance()->GetOnScreenColorRenderBuffer(), GL_RENDERBUFFER, GL_COLOR_ATTACHMENT0, SizeX, SizeY, 0, 1, 1, 1, PF_B8G8R8A8, false, false, Flags, nullptr, FClearValueBinding::Transparent);
 		OpenGLTextureAllocated(Texture2D, Flags);
 	}
 
