Changed FClearValueBinding::None to Transparent
#ue4

diff --git a/Engine/Source/Runtime/OpenGLDrv/Private/OpenGLViewport.cpp b/Engine/Source/Runtime/OpenGLDrv/Private/OpenGLViewport.cpp
--- a/Engine/Source/Runtime/OpenGLDrv/Private/OpenGLViewport.cpp
+++ b/Engine/Source/Runtime/OpenGLDrv/Private/OpenGLViewport.cpp
@@ -283,7 +283,7 @@ void FOpenGLViewport::Resize(uint32 InSizeX,uint32 InSizeY,bool bInIsFullscreen)
 	BackBuffer = (FOpenGLTexture2D*)PlatformCreateBuiltinBackBuffer(OpenGLRHI, InSizeX, InSizeY);
 	if (!BackBuffer)
 	{
-		BackBuffer = (FOpenGLTexture2D*)OpenGLRHI->CreateOpenGLTexture(InSizeX, InSizeY, false, false, PixelFormat, 1, 1, 1, TexCreate_RenderTargetable, FClearValueBinding::None);
+		BackBuffer = (FOpenGLTexture2D*)OpenGLRHI->CreateOpenGLTexture(InSizeX, InSizeY, false, false, PixelFormat, 1, 1, 1, TexCreate_RenderTargetable, FClearValueBinding::Transparent);
 	}
 
 	PlatformResizeGLContext(OpenGLRHI->PlatformDevice, OpenGLContext, InSizeX, InSizeY, bInIsFullscreen, bIsFullscreen, BackBuffer->Target, BackBuffer->Resource);
