Remove Mac-specific workaround for an AMD/OpenGL bug with SRGB frambuffers as that has been fixed in 10.11. Also remove the special handling on Intel HD 3000 GPUs on Mac.

diff --git a/Engine/Shaders/DeferredShadingCommon.usf b/Engine/Shaders/DeferredShadingCommon.usf
--- a/Engine/Shaders/DeferredShadingCommon.usf
+++ b/Engine/Shaders/DeferredShadingCommon.usf
@@ -116,7 +116,6 @@ float3 DecodeNormal( float3 N )
 	//return OctahedronToUnitVector( Pack888To1212( N ) * 2 - 1 );
 }
 
-#if USE_FRAMEBUFFER_SRGB
 float3 EncodeBaseColor(float3 BaseColor)
 {
 	// we use sRGB on the render target to give more precision to the darks
@@ -128,18 +127,6 @@ float3 DecodeBaseColor(float3 BaseColor)
 	// we use sRGB on the render target to give more precision to the darks
 	return BaseColor;
 }
-#else // @todo: remove once Apple fixes radr://16754329 AMD Cards don't always perform FRAMEBUFFER_SRGB if the draw FBO has mixed sRGB & non-SRGB colour attachments
-float3 EncodeBaseColor(float3 BaseColor)
-{
-	// Gamma 2.0
-	return sqrt( saturate(BaseColor) );
-}
-
-float3 DecodeBaseColor(float3 BaseColor)
-{
-	return Square( BaseColor );
-}
-#endif
 
 float3 EncodeSubsurfaceColor(float3 SubsurfaceColor)
 {
diff --git a/Engine/Source/Runtime/Engine/Private/ShaderCompiler/ShaderCompiler.cpp b/Engine/Source/Runtime/Engine/Private/ShaderCompiler/ShaderCompiler.cpp
--- a/Engine/Source/Runtime/Engine/Private/ShaderCompiler/ShaderCompiler.cpp
+++ b/Engine/Source/Runtime/Engine/Private/ShaderCompiler/ShaderCompiler.cpp
@@ -1934,15 +1934,6 @@ void GlobalBeginCompileShader(
 		static IConsoleVariable* CVar = IConsoleManager::Get().FindConsoleVariable(TEXT("r.DBuffer"));
 		Input.Environment.SetDefine(TEXT("USE_DBUFFER"), CVar ? CVar->GetInt() : 0);
 	}
-	
-	{
-		int32 UseFrameBufferSRGB = 1;
-#if PLATFORM_MAC // @todo: remove once Apple fixes radr://16754329 AMD Cards don't always perform FRAMEBUFFER_SRGB if the draw FBO has mixed sRGB & non-SRGB colour attachments
-		static const auto CVar = IConsoleManager::Get().FindTConsoleVariableDataInt(TEXT("r.Mac.UseFrameBufferSRGB"));
-		UseFrameBufferSRGB = CVar ? (CVar->GetValueOnGameThread() != 0) : 0;
-#endif
-		Input.Environment.SetDefine(TEXT("USE_FRAMEBUFFER_SRGB"), UseFrameBufferSRGB);
-	}
 
 	NewJobs.Add(NewJob);
 }
diff --git a/Engine/Source/Runtime/OpenGLDrv/Private/Mac/OpenGLMac.cpp b/Engine/Source/Runtime/OpenGLDrv/Private/Mac/OpenGLMac.cpp
--- a/Engine/Source/Runtime/OpenGLDrv/Private/Mac/OpenGLMac.cpp
+++ b/Engine/Source/Runtime/OpenGLDrv/Private/Mac/OpenGLMac.cpp
@@ -17,13 +17,8 @@
  OpenGL static variables.
  ------------------------------------------------------------------------------*/
 
-// @todo: remove once Apple fixes radr://16754329 AMD Cards don't always perform FRAMEBUFFER_SRGB if the draw FBO has mixed sRGB & non-SRGB colour attachments
-static TAutoConsoleVariable<int32> CVarMacUseFrameBufferSRGB(
-	TEXT("r.Mac.UseFrameBufferSRGB"),
-	0,
-	TEXT("Flag to toggle use of GL_FRAMEBUFFER_SRGB for better color accuracy.\n"),
-	ECVF_RenderThreadSafe
-	);
+// As of 10.11.0 Apple fixed radr://16754329 AMD Cards don't always perform FRAMEBUFFER_SRGB if the draw FBO has mixed sRGB & non-SRGB colour attachments
+bool FMacOpenGL::bUseSRGBFramebuffer = false;
 
 // @todo: remove once Apple fixes radr://15553950, TTP# 315197
 static int32 GMacFlushTexStorage = true;
@@ -1169,6 +1164,8 @@ void FMacOpenGL::ProcessExtensions(const FString& ExtensionsString)
 	
 	// SSOs require structs in geometry shaders - which only work in 10.10.0 and later
 	bSupportsSeparateShaderObjects &= (FMacPlatformMisc::MacOSXVersionCompare(10,10,0) >= 0);
+	
+	bUseSRGBFramebuffer = (!IsRHIDeviceAMD() || FMacPlatformMisc::MacOSXVersionCompare(10,11,0) >= 0);
 }
 
 void FMacOpenGL::MacQueryTimestampCounter(GLuint QueryID)
@@ -1292,20 +1289,6 @@ bool FMacOpenGL::MustFlushTexStorage(void)
 	return GMacFlushTexStorage || GMacMustFlushTexStorage;
 }
 
-/** Is the current renderer the Intel HD3000? */
-bool FMacOpenGL::IsIntelHD3000()
-{
-	// Get the current renderer ID
-	GLint RendererID = 0;
-	CGLContextObj Current = CGLGetCurrentContext();
-	if (Current)
-	{
-		CGLError Error = CGLGetParameter(Current, kCGLCPCurrentRendererID, &RendererID);
-		check(Error == kCGLNoError && RendererID != 0);
-	}
-	return (RendererID & kCGLRendererIDMatchingMask) == kCGLRendererIntelHDID;
-}
-
 void FMacOpenGL::DeleteTextures(GLsizei Number, const GLuint* Textures)
 {
 	TArray<GLuint>& TexturesToDelete = GMacTexturesToDelete[GFrameNumberRenderThread % GMacTexturePoolNum];
diff --git a/Engine/Source/Runtime/OpenGLDrv/Private/Mac/OpenGLMac.h b/Engine/Source/Runtime/OpenGLDrv/Private/Mac/OpenGLMac.h
--- a/Engine/Source/Runtime/OpenGLDrv/Private/Mac/OpenGLMac.h
+++ b/Engine/Source/Runtime/OpenGLDrv/Private/Mac/OpenGLMac.h
@@ -24,9 +24,6 @@ private:
 	/** Must we employ a workaround for radr://15553950, TTP# 315197 */
 	static bool MustFlushTexStorage(void);
 	
-	/** Is the current renderer the Intel HD3000? */
-	static bool IsIntelHD3000();
-	
 public:
 	static void ProcessQueryGLInt();
 	static void ProcessExtensions(const FString& ExtensionsString);
@@ -36,7 +33,7 @@ public:
 	{
 		// The Intel HD 3000 on OS X can't cope with our deferred renderer, but ES2 mode works fine :(
 		static bool bForceFeatureLevelES2 = FParse::Param(FCommandLine::Get(), TEXT("FeatureLevelES2"));
-		if (bForceFeatureLevelES2 || IsIntelHD3000())
+		if (bForceFeatureLevelES2)
 		{
 			return ERHIFeatureLevel::ES2;
 		}
@@ -48,7 +45,7 @@ public:
 	{
 		// The Intel HD 3000 on OS X can't cope with our deferred renderer, but ES2 mode works fine :(
 		static bool bForceFeatureLevelES2 = FParse::Param(FCommandLine::Get(), TEXT("FeatureLevelES2"));
-		if (bForceFeatureLevelES2 || IsIntelHD3000())
+		if (bForceFeatureLevelES2)
 		{
 			return SP_OPENGL_PCES2;
 		}
@@ -71,6 +68,7 @@ public:
 	static FORCEINLINE bool SupportsTextureRange()						{ return !bUsingApitrace; }
 	static FORCEINLINE bool SupportsDepthBoundsTest()					{ return true; }
 	static FORCEINLINE bool SupportsDrawIndirect()						{ return bSupportsDrawIndirect; }
+	static FORCEINLINE bool SupportsSRGBFramebuffer()					{ return bUseSRGBFramebuffer; }
 	
 	static FORCEINLINE void Flush()
 	{
@@ -353,6 +351,8 @@ private:
 	static bool bSupportsDrawIndirect;
 	static PFNGLDRAWARRAYSINDIRECTPROC glDrawArraysIndirect;
 	static PFNGLDRAWELEMENTSINDIRECTPROC glDrawElementsIndirect;
+	/** Should we use SRGB framebuffers? */
+	static bool bUseSRGBFramebuffer;
 };
 
 
diff --git a/Engine/Source/Runtime/OpenGLDrv/Private/OpenGLTexture.cpp b/Engine/Source/Runtime/OpenGLDrv/Private/OpenGLTexture.cpp
--- a/Engine/Source/Runtime/OpenGLDrv/Private/OpenGLTexture.cpp
+++ b/Engine/Source/Runtime/OpenGLDrv/Private/OpenGLTexture.cpp
@@ -255,7 +255,12 @@ FRHITexture* FOpenGLDynamicRHI::CreateOpenGLTexture(uint32 SizeX, uint32 SizeY,
 	check( bArrayTexture != (ArraySize == 1));
 #endif
 
-	if (GMaxRHIFeatureLevel <= ERHIFeatureLevel::ES2)
+	bool bNoSRGBSupport = (GMaxRHIFeatureLevel <= ERHIFeatureLevel::ES2);
+#if PLATFORM_MAC // Handle versions prior to 10.11.0 that still suffer from radr://16754329 AMD Cards don't always perform FRAMEBUFFER_SRGB if the draw FBO has mixed sRGB & non-SRGB colour attachments
+	bNoSRGBSupport |= (Flags & TexCreate_RenderTargetable) && !FOpenGL::SupportsSRGBFramebuffer();
+#endif
+	
+	if (bNoSRGBSupport)
 	{
 		// Remove sRGB read flag when not supported
 		Flags &= ~TexCreate_SRGB;
diff --git a/Engine/Source/Runtime/Renderer/Private/PostProcess/SceneRenderTargets.cpp b/Engine/Source/Runtime/Renderer/Private/PostProcess/SceneRenderTargets.cpp
--- a/Engine/Source/Runtime/Renderer/Private/PostProcess/SceneRenderTargets.cpp
+++ b/Engine/Source/Runtime/Renderer/Private/PostProcess/SceneRenderTargets.cpp
@@ -746,14 +746,7 @@ void FSceneRenderTargets::AllocGBufferTargets()
 	// Create the diffuse color g-buffer.
 	{
 		const EPixelFormat DiffuseGBufferFormat = bHighPrecisionGBuffers ? PF_FloatRGBA : PF_B8G8R8A8;
-		uint32 DiffuseGBufferFlags = TexCreate_SRGB;
-
-#if PLATFORM_MAC // @todo: remove once Apple fixes radr://16754329 AMD Cards don't always perform FRAMEBUFFER_SRGB if the draw FBO has mixed sRGB & non-SRGB colour attachments
-		static const auto CVar = IConsoleManager::Get().FindTConsoleVariableDataInt(TEXT("r.Mac.UseFrameBufferSRGB"));
-		DiffuseGBufferFlags = CVar && CVar->GetValueOnRenderThread() ? TexCreate_SRGB : TexCreate_None;
-#endif
-
-		FPooledRenderTargetDesc Desc(FPooledRenderTargetDesc::Create2DDesc(BufferSize, DiffuseGBufferFormat, FClearValueBinding::Transparent, DiffuseGBufferFlags, TexCreate_RenderTargetable, false));
+		FPooledRenderTargetDesc Desc(FPooledRenderTargetDesc::Create2DDesc(BufferSize, DiffuseGBufferFormat, FClearValueBinding::Transparent, TexCreate_SRGB, TexCreate_RenderTargetable, false));
 		GRenderTargetPool.FindFreeElement(Desc, GBufferC, TEXT("GBufferC"));
 	}
 
