Remove the Mac workaround for slow HZB readback as it is no longer relevant and resolve the inefficiencies inside the OpenGLRHI that trying to use it exposes.

diff --git a/Engine/Source/Runtime/OpenGLDrv/Private/OpenGLRenderTarget.cpp b/Engine/Source/Runtime/OpenGLDrv/Private/OpenGLRenderTarget.cpp
--- a/Engine/Source/Runtime/OpenGLDrv/Private/OpenGLRenderTarget.cpp
+++ b/Engine/Source/Runtime/OpenGLDrv/Private/OpenGLRenderTarget.cpp
@@ -436,8 +436,8 @@ void FOpenGLDynamicRHI::RHICopyToResolveTarget(FTextureRHIParamRef SourceTexture
 		
 		// For CPU readback resolve targets we should issue the resolve to the internal PBO immediately.
 		// This makes any subsequent locking of that texture much cheaper as it won't have to stall on a pixel pack op.
-		bool bLockableTarget = DestTextureRHI->GetTexture2D() && (DestTextureRHI->GetFlags() & TexCreate_CPUReadback) && !(DestTextureRHI->GetFlags() & (TexCreate_RenderTargetable|TexCreate_DepthStencilTargetable));
-		if(bTrueBlit && bLockableTarget && FOpenGL::SupportsPixelBuffers() && !ResolveParams.Rect.IsValid())
+		bool bLockableTarget = DestTextureRHI->GetTexture2D() && (DestTextureRHI->GetFlags() & TexCreate_CPUReadback) && !(DestTextureRHI->GetFlags() & (TexCreate_RenderTargetable|TexCreate_DepthStencilTargetable)) && !DestTextureRHI->IsMultisampled();
+		if(bLockableTarget && FOpenGL::SupportsPixelBuffers() && !ResolveParams.Rect.IsValid())
 		{
 			FOpenGLTexture2D* DestTex = (FOpenGLTexture2D*)DestTexture;
 			DestTex->Resolve(MipmapLevel, DestIndex);
diff --git a/Engine/Source/Runtime/OpenGLDrv/Private/OpenGLTexture.cpp b/Engine/Source/Runtime/OpenGLDrv/Private/OpenGLTexture.cpp
--- a/Engine/Source/Runtime/OpenGLDrv/Private/OpenGLTexture.cpp
+++ b/Engine/Source/Runtime/OpenGLDrv/Private/OpenGLTexture.cpp
@@ -358,7 +358,7 @@ FRHITexture* FOpenGLDynamicRHI::CreateOpenGLTexture(uint32 SizeX, uint32 SizeY,
 		else
 		{
 			// Should we use client-storage to improve update time on platforms that require it
-			bool const bRenderable = (Flags & (TexCreate_RenderTargetable|TexCreate_ResolveTargetable|TexCreate_DepthStencilTargetable)) != 0;
+			bool const bRenderable = (Flags & (TexCreate_RenderTargetable|TexCreate_ResolveTargetable|TexCreate_DepthStencilTargetable|TexCreate_CPUReadback)) != 0;
 			bool const bUseClientStorage = (FOpenGL::SupportsClientStorage() && !FOpenGL::SupportsTextureView() && !bRenderable && !GLFormat.bCompressed);
 			if(bUseClientStorage)
 			{
@@ -755,7 +755,7 @@ void* TOpenGLTexture<RHIResourceType>::Lock(uint32 InMipIndex,uint32 ArrayIndex,
 
 	// Should we use client-storage to improve update time on platforms that require it
 	const FOpenGLTextureFormat& GLFormat = GOpenGLTextureFormats[PixelFormat];
-	bool const bRenderable = (this->GetFlags() & (TexCreate_RenderTargetable|TexCreate_ResolveTargetable|TexCreate_DepthStencilTargetable)) != 0;
+	bool const bRenderable = (this->GetFlags() & (TexCreate_RenderTargetable|TexCreate_ResolveTargetable|TexCreate_DepthStencilTargetable|TexCreate_CPUReadback)) != 0;
 	bool const bUseClientStorage = FOpenGL::SupportsClientStorage() && !FOpenGL::SupportsTextureView() && !bRenderable && !this->GetSizeZ() && !GLFormat.bCompressed;
 	if(!bUseClientStorage)
 	{
@@ -816,6 +816,7 @@ void* TOpenGLTexture<RHIResourceType>::Lock(uint32 InMipIndex,uint32 ArrayIndex,
 			
 			result = ClientStorageBuffers[BufferIndex].Data;
 		}
+		ClientStorageBuffers[BufferIndex].bReadOnly = (LockMode == RLM_ReadOnly);
 	}
 	
 	return result;
@@ -885,7 +886,7 @@ void TOpenGLTexture<RHIResourceType>::Unlock(uint32 MipIndex,uint32 ArrayIndex)
 	const bool bSRGB = (this->GetFlags() & TexCreate_SRGB) != 0;
 
 	// Should we use client-storage to improve update time on platforms that require it
-	bool const bRenderable = (this->GetFlags() & (TexCreate_RenderTargetable|TexCreate_ResolveTargetable|TexCreate_DepthStencilTargetable)) != 0;
+	bool const bRenderable = (this->GetFlags() & (TexCreate_RenderTargetable|TexCreate_ResolveTargetable|TexCreate_DepthStencilTargetable|TexCreate_CPUReadback)) != 0;
 	bool const bUseClientStorage = FOpenGL::SupportsClientStorage() && !FOpenGL::SupportsTextureView() && !bRenderable && !this->GetSizeZ() && !GLFormat.bCompressed;
 	check(bUseClientStorage || IsValidRef(PixelBuffers[BufferIndex]));
 	
@@ -1058,7 +1059,7 @@ void TOpenGLTexture<RHIResourceType>::Unlock(uint32 MipIndex,uint32 ArrayIndex)
 		//need to free PBO if we aren't keeping shadow copies
 		PixelBuffers[BufferIndex] = NULL;
 	}
-	else
+	else if(!bUseClientStorage || !ClientStorageBuffers[BufferIndex].bReadOnly)
 	{
 		// Code path for non-PBO:
 		// Volume/array textures are currently only supported if PixelBufferObjects are also supported.
@@ -1899,7 +1900,7 @@ FTexture2DRHIRef FOpenGLDynamicRHI::RHIAsyncReallocateTexture2D(FTexture2DRHIPar
 	
 	// Should we use client-storage to improve update time on platforms that require it
 	const bool bCompressed = GOpenGLTextureFormats[Texture2D->GetFormat()].bCompressed;
-	bool const bRenderable = (Texture2D->GetFlags() & (TexCreate_RenderTargetable|TexCreate_ResolveTargetable|TexCreate_DepthStencilTargetable)) != 0;
+	bool const bRenderable = (Texture2D->GetFlags() & (TexCreate_RenderTargetable|TexCreate_ResolveTargetable|TexCreate_DepthStencilTargetable|TexCreate_CPUReadback)) != 0;
 	const bool bUseClientStorage = (FOpenGL::SupportsClientStorage() && !FOpenGL::SupportsTextureView() && !bRenderable && !bCompressed);
 	
 	// Use the GPU to asynchronously copy the old mip-maps into the new texture.
diff --git a/Engine/Source/Runtime/Renderer/Private/SceneOcclusion.cpp b/Engine/Source/Runtime/Renderer/Private/SceneOcclusion.cpp
--- a/Engine/Source/Runtime/Renderer/Private/SceneOcclusion.cpp
+++ b/Engine/Source/Runtime/Renderer/Private/SceneOcclusion.cpp
@@ -520,11 +520,7 @@ void FHZBOcclusionTester::InitDynamicRHI()
 {
 	if (GetFeatureLevel() >= ERHIFeatureLevel::SM4)
 	{
-#if PLATFORM_MAC // Workaround radr://16096028 Texture Readback via glReadPixels + PBOs stalls on Nvidia GPUs
-		FPooledRenderTargetDesc Desc( FPooledRenderTargetDesc::Create2DDesc( FIntPoint( SizeX, SizeY ), PF_R8G8B8A8, FClearValueBinding::None, TexCreate_CPUReadback | TexCreate_HideInVisualizeTexture, TexCreate_None, false ) );
-#else
 		FPooledRenderTargetDesc Desc( FPooledRenderTargetDesc::Create2DDesc( FIntPoint( SizeX, SizeY ), PF_B8G8R8A8, FClearValueBinding::None, TexCreate_CPUReadback | TexCreate_HideInVisualizeTexture, TexCreate_None, false ) );
-#endif
 		GRenderTargetPool.FindFreeElement( Desc, ResultsTextureCPU, TEXT("HZBResultsCPU") );
 	}
 }
@@ -728,11 +724,7 @@ void FHZBOcclusionTester::Submit(FRHICommandListImmediate& RHICmdList, const FVi
 
 	TRefCountPtr< IPooledRenderTarget >	ResultsTextureGPU;
 	{
-#if PLATFORM_MAC // Workaround radr://16096028 Texture Readback via glReadPixels + PBOs stalls on Nvidia GPUs
-		FPooledRenderTargetDesc Desc( FPooledRenderTargetDesc::Create2DDesc( FIntPoint( SizeX, SizeY ), PF_R8G8B8A8, FClearValueBinding::None, TexCreate_None, TexCreate_RenderTargetable, false ) );
-#else
 		FPooledRenderTargetDesc Desc( FPooledRenderTargetDesc::Create2DDesc( FIntPoint( SizeX, SizeY ), PF_B8G8R8A8, FClearValueBinding::None, TexCreate_None, TexCreate_RenderTargetable, false ) );
-#endif
 		GRenderTargetPool.FindFreeElement( Desc, ResultsTextureGPU, TEXT("HZBResultsGPU") );
 	}
 
