CompileOpenGLShader for Android now only checks the compile result immediately during hack testing
#ue4
#android

diff --git a/Engine/Source/Runtime/OpenGLDrv/Private/OpenGLDevice.cpp b/Engine/Source/Runtime/OpenGLDrv/Private/OpenGLDevice.cpp
--- a/Engine/Source/Runtime/OpenGLDrv/Private/OpenGLDevice.cpp
+++ b/Engine/Source/Runtime/OpenGLDrv/Private/OpenGLDevice.cpp
@@ -1020,6 +1020,8 @@ static void CheckVaryingLimit()
 		}
 
 		UE_LOG(LogRHI, Display, TEXT("Testing for gl_FragCoord requiring a varying since mosaic is enabled"));
+		FOpenGL::bIsCheckingShaderCompilerHacks = true;
+
 		const ANSICHAR* TestVertexProgram = "\n"
 			"#version 100\n"
 			"attribute vec4 in_ATTRIBUTE0;\n"
@@ -1084,15 +1086,19 @@ static void CheckVaryingLimit()
 		if (!VerifyCompiledShader(VertexShader->Resource, TestVertexProgram, false))
 		{
 			UE_LOG(LogRHI, Warning, TEXT("Vertex shader for varying test failed to compile. Try running anyway."));
+			FOpenGL::bIsCheckingShaderCompilerHacks = false;
 			return;
 		}
 		TRefCountPtr<FOpenGLPixelShader> PixelShader = (FOpenGLPixelShader*)(RHICreatePixelShader(FragmentCode).GetReference());
 		if (!VerifyCompiledShader(PixelShader->Resource, TestFragmentProgram, false))
 		{
 			UE_LOG(LogRHI, Warning, TEXT("Fragment shader for varying test failed to compile. Try running anyway."));
+			FOpenGL::bIsCheckingShaderCompilerHacks = false;
 			return;
 		}
 
+		FOpenGL::bIsCheckingShaderCompilerHacks = false;
+
 		// Now try linking them.. this is where gl_FragCoord may cause a failure
 		GLuint Program = glCreateProgram();
 		glAttachShader(Program, VertexShader->Resource);
@@ -1118,6 +1124,8 @@ static void CheckTextureCubeLodSupport()
 	if (IsES2Platform(GMaxRHIShaderPlatform))
 	{
 		UE_LOG(LogRHI, Display, TEXT("Testing for shader compiler compatibility"));
+		FOpenGL::bIsCheckingShaderCompilerHacks = true;
+
 		// This code creates a sample program and finds out which hacks are required to compile it
 		const ANSICHAR* TestFragmentProgram = "\n"
 			"#version 100\n"
@@ -1161,6 +1169,7 @@ static void CheckTextureCubeLodSupport()
 		{
 			UE_LOG(LogRHI, Display, TEXT("Shaders compile fine no need to enable hacks"));
 			// we are done
+			FOpenGL::bIsCheckingShaderCompilerHacks = false;
 			return;
 		}
 		
@@ -1177,6 +1186,7 @@ static void CheckTextureCubeLodSupport()
 			UE_LOG(LogRHI, Warning, TEXT("Enabling shader compiler hack to remove precision modifiers for texture samplers"));
 
 			// we are done
+			FOpenGL::bIsCheckingShaderCompilerHacks = false;
 			return;
 		}
 
@@ -1190,6 +1200,7 @@ static void CheckTextureCubeLodSupport()
 		{
 			UE_LOG(LogRHI, Warning, TEXT("Enabling shader compiler hack to redefine textureCubeLodEXT to textureCubeLod"));
 			// we are done
+			FOpenGL::bIsCheckingShaderCompilerHacks = false;
 			return;
 		}
 
@@ -1203,10 +1214,12 @@ static void CheckTextureCubeLodSupport()
 		{
 			UE_LOG(LogRHI, Warning, TEXT("Enabling shader compiler hack to redefine textureCubeLodEXT to textureCubeLod and remove precision modifiers"));
 			// we are done
+			FOpenGL::bIsCheckingShaderCompilerHacks = false;
 			return;
 		}
 
 		UE_LOG(LogRHI, Warning, TEXT("Unable to find a test shader that compiles try running anyway"));
+		FOpenGL::bIsCheckingShaderCompilerHacks = false;
 	}
 #endif
 }
diff --git a/Engine/Source/Runtime/OpenGLDrv/Private/OpenGLES2.cpp b/Engine/Source/Runtime/OpenGLDrv/Private/OpenGLES2.cpp
--- a/Engine/Source/Runtime/OpenGLDrv/Private/OpenGLES2.cpp
+++ b/Engine/Source/Runtime/OpenGLDrv/Private/OpenGLES2.cpp
@@ -133,6 +133,9 @@ bool FOpenGLES2::bNeedsVertexAttribRemap = false;
 /* This hack fixes an issue with SGX540 compiler which can get upset with some operations that mix highp and mediump */
 bool FOpenGLES2::bRequiresTexture2DPrecisionHack = false;
 
+/* Indicates shader compiler hack checks are being tested */
+bool FOpenGLES2::bIsCheckingShaderCompilerHacks = false;
+
 bool FOpenGLES2::SupportsDisjointTimeQueries()
 {
 	bool bAllowDisjointTimerQueries = false;
diff --git a/Engine/Source/Runtime/OpenGLDrv/Private/OpenGLES31.cpp b/Engine/Source/Runtime/OpenGLDrv/Private/OpenGLES31.cpp
--- a/Engine/Source/Runtime/OpenGLDrv/Private/OpenGLES31.cpp
+++ b/Engine/Source/Runtime/OpenGLDrv/Private/OpenGLES31.cpp
@@ -133,6 +133,9 @@ bool FOpenGLES31::bRequiresGLFragCoordVaryingLimitHack = false;
 /* This hack fixes an issue with SGX540 compiler which can get upset with some operations that mix highp and mediump */
 bool FOpenGLES31::bRequiresTexture2DPrecisionHack = false;
 
+/* Indicates shader compiler hack checks are being tested */
+bool FOpenGLES31::bIsCheckingShaderCompilerHacks = false;
+
 /** GL_EXT_disjoint_timer_query or GL_NV_timer_query*/
 bool FOpenGLES31::bSupportsDisjointTimeQueries = false;
 
diff --git a/Engine/Source/Runtime/OpenGLDrv/Private/OpenGLShaders.cpp b/Engine/Source/Runtime/OpenGLDrv/Private/OpenGLShaders.cpp
--- a/Engine/Source/Runtime/OpenGLDrv/Private/OpenGLShaders.cpp
+++ b/Engine/Source/Runtime/OpenGLDrv/Private/OpenGLShaders.cpp
@@ -714,8 +714,11 @@ ShaderType* CompileOpenGLShader(const TArray<uint8>& Code)
 
 		GLint CompileStatus = GL_TRUE;
 #if PLATFORM_ANDROID
-		// On Android the same shader is compiled with different hacks to find the right one(s) to apply so don't cache unless successful
-		glGetShaderiv(Resource, GL_COMPILE_STATUS, &CompileStatus);
+		// On Android the same shader is compiled with different hacks to find the right one(s) to apply so don't cache unless successful if currently testing them
+		if (FOpenGL::IsCheckingShaderCompilerHacks())
+		{
+			glGetShaderiv(Resource, GL_COMPILE_STATUS, &CompileStatus);
+		}
 #endif
 #if PLATFORM_HTML5
 		glGetShaderiv(Resource, GL_COMPILE_STATUS, &CompileStatus);
diff --git a/Engine/Source/Runtime/OpenGLDrv/Public/OpenGLES2.h b/Engine/Source/Runtime/OpenGLDrv/Public/OpenGLES2.h
--- a/Engine/Source/Runtime/OpenGLDrv/Public/OpenGLES2.h
+++ b/Engine/Source/Runtime/OpenGLDrv/Public/OpenGLES2.h
@@ -135,6 +135,7 @@ struct FOpenGLES2 : public FOpenGLBase
 	static FORCEINLINE bool SupportsStandardDerivativesExtension()		{ return bSupportsStandardDerivativesExtension; }
 	static FORCEINLINE bool RequiresGLFragCoordVaryingLimitHack()		{ return bRequiresGLFragCoordVaryingLimitHack; }
 	static FORCEINLINE bool RequiresTexture2DPrecisionHack()			{ return bRequiresTexture2DPrecisionHack; }
+	static FORCEINLINE bool IsCheckingShaderCompilerHacks()				{ return bIsCheckingShaderCompilerHacks; }
 
 	static FORCEINLINE int32 GetReadHalfFloatPixelsEnum()				{ return GL_HALF_FLOAT_OES; }
 
@@ -496,6 +497,9 @@ public:
 
 	/* This hack fixes an issue with SGX540 compiler which can get upset with some operations that mix highp and mediump */
 	static bool bRequiresTexture2DPrecisionHack;
+
+	/* Indicates shader compiler hack checks are being tested */
+	static bool bIsCheckingShaderCompilerHacks;
 };
 
 
diff --git a/Engine/Source/Runtime/OpenGLDrv/Public/OpenGLES31.h b/Engine/Source/Runtime/OpenGLDrv/Public/OpenGLES31.h
--- a/Engine/Source/Runtime/OpenGLDrv/Public/OpenGLES31.h
+++ b/Engine/Source/Runtime/OpenGLDrv/Public/OpenGLES31.h
@@ -90,7 +90,7 @@ struct FOpenGLES31 : public FOpenGLBase
 	static FORCEINLINE bool RequiresGLFragCoordVaryingLimitHack()		{ return bRequiresGLFragCoordVaryingLimitHack; }
 	static FORCEINLINE GLenum GetVertexHalfFloatFormat()				{ return bES2Fallback ? GL_HALF_FLOAT_OES : GL_HALF_FLOAT; }
 	static FORCEINLINE bool RequiresTexture2DPrecisionHack()			{ return bRequiresTexture2DPrecisionHack; }
-
+	static FORCEINLINE bool IsCheckingShaderCompilerHacks()				{ return bIsCheckingShaderCompilerHacks; }
 
 	// On iOS both glMapBufferOES() and glBufferSubData() for immediate vertex and index data
 	// is the slow path (they both hit GPU sync and data cache flush in driver according to profiling in driver symbols).
@@ -950,6 +950,9 @@ public:
 
 	/* This hack fixes an issue with SGX540 compiler which can get upset with some operations that mix highp and mediump */
 	static bool bRequiresTexture2DPrecisionHack;
+
+	/* Indicates shader compiler hack checks are being tested */
+	static bool bIsCheckingShaderCompilerHacks;
 };
 
 // yes they are different between the ES2 extension and ES3.x and GL3.x core
