When running under APITrace use old-fashined NSView/OpenGL rather than CoreAnimation/OpenGL as it can't cope with the way CoreAnimation drives OpenGL making UE4 untraceable.

diff --git a/Engine/Source/Runtime/OpenGLDrv/Private/Mac/OpenGLMac.cpp b/Engine/Source/Runtime/OpenGLDrv/Private/Mac/OpenGLMac.cpp
--- a/Engine/Source/Runtime/OpenGLDrv/Private/Mac/OpenGLMac.cpp
+++ b/Engine/Source/Runtime/OpenGLDrv/Private/Mac/OpenGLMac.cpp
@@ -189,7 +189,14 @@ static void UnlockGLContext(NSOpenGLContext* Context)
 
 - (CALayer*)makeBackingLayer
 {
-	return [[FSlateOpenGLLayer alloc] initWithContext:self.Context->OpenGLContext andPixelFormat:self.Context->OpenGLPixelFormat];
+	if(FMacOpenGL::SupportsCoreAnimation())
+	{
+		return [[FSlateOpenGLLayer alloc] initWithContext:self.Context->OpenGLContext andPixelFormat:self.Context->OpenGLPixelFormat];
+	}
+	else
+	{
+		return nil;
+	}
 }
 
 - (id)initWithFrame:(NSRect)frameRect context:(FPlatformOpenGLContext*)context
@@ -217,7 +224,10 @@ static void UnlockGLContext(NSOpenGLContext* Context)
 
 - (void)drawRect:(NSRect)DirtyRect
 {
-	DrawOpenGLViewport(self.Context, self.frame.size.width, self.frame.size.height);
+	if(FMacOpenGL::SupportsCoreAnimation())
+	{
+		DrawOpenGLViewport(self.Context, self.frame.size.width, self.frame.size.height);
+	}
 }
 
 - (BOOL)isOpaque
@@ -362,21 +372,34 @@ FPlatformOpenGLContext* PlatformCreateOpenGLContext(FPlatformOpenGLDevice* Devic
 		{
 			NSView* SuperView = [[Context->WindowHandle contentView] superview];
 			[SuperView addSubview:Context->OpenGLView];
-			[SuperView setWantsLayer:YES];
+			if(FMacOpenGL::SupportsCoreAnimation())
+			{
+				[SuperView setWantsLayer:YES];
+			}
 			[SuperView addSubview:[Context->WindowHandle standardWindowButton:NSWindowCloseButton]];
 			[SuperView addSubview:[Context->WindowHandle standardWindowButton:NSWindowMiniaturizeButton]];
 			[SuperView addSubview:[Context->WindowHandle standardWindowButton:NSWindowZoomButton]];
 		}
 		else
 		{
-			[Context->OpenGLView setWantsLayer:YES];
+			if(FMacOpenGL::SupportsCoreAnimation())
+			{
+				[Context->OpenGLView setWantsLayer:YES];
+			}
 			[Context->WindowHandle setContentView:Context->OpenGLView];
 		}
 
 		[[Context->WindowHandle standardWindowButton:NSWindowCloseButton] setAction:@selector(performClose:)];
-
-		Context->OpenGLView.layer.magnificationFilter = kCAFilterNearest;
-		Context->OpenGLView.layer.minificationFilter = kCAFilterNearest;
+		
+		if(FMacOpenGL::SupportsCoreAnimation())
+		{
+			Context->OpenGLView.layer.magnificationFilter = kCAFilterNearest;
+			Context->OpenGLView.layer.minificationFilter = kCAFilterNearest;
+		}
+		else
+		{
+			[Context->OpenGLContext setView:Context->OpenGLView];
+		}
 	}, NSDefaultRunLoopMode, true);
 
 	return Context;
@@ -473,7 +496,14 @@ bool PlatformBlitToViewport(FPlatformOpenGLDevice* Device, const FOpenGLViewport
 			Context->ViewportSize[0] = BackbufferSizeX;
 			Context->ViewportSize[1] = BackbufferSizeY;
 
-			MainThreadCall(^{ [Context->OpenGLView setNeedsDisplay:YES]; }, NSDefaultRunLoopMode, false);
+			if(FMacOpenGL::SupportsCoreAnimation())
+			{
+				MainThreadCall(^{ [Context->OpenGLView setNeedsDisplay:YES]; }, NSDefaultRunLoopMode, false);
+			}
+			else
+			{
+				DrawOpenGLViewport(Context, BackbufferSizeX, BackbufferSizeY);
+			}
 
 			TArray<GLuint>& TexturesToDelete = GMacTexturesToDelete[(GFrameNumberRenderThread - (GMacTexturePoolNum - 1)) % GMacTexturePoolNum];
 			if(TexturesToDelete.Num() && TexturesToDelete.Num() > GMacMinTexturesToDeletePerFrame)
@@ -492,6 +522,16 @@ bool PlatformBlitToViewport(FPlatformOpenGLDevice* Device, const FOpenGLViewport
 void DrawOpenGLViewport(FPlatformOpenGLContext* const Context, uint32 Width, uint32 Height)
 {
 	FCocoaWindow* Window = (FCocoaWindow*)Context->WindowHandle;
+	
+	int32 CurrentDrawFramebuffer = 0;
+	
+	if(!FMacOpenGL::SupportsCoreAnimation())
+	{
+		glGetIntegerv(GL_DRAW_FRAMEBUFFER_BINDING, &CurrentDrawFramebuffer);
+		glBindFramebuffer(GL_DRAW_FRAMEBUFFER, 0);
+		glDrawBuffer(GL_BACK);
+	}
+	
 	if ([Window isRenderInitialized] && Context->ViewportSize[0] && Context->ViewportSize[1] && Context->ViewportFramebuffer && Context->ViewportRenderbuffer && ([Window styleMask] & (NSTexturedBackgroundWindowMask|NSFullSizeContentViewWindowMask) || !Window.inLiveResize))
 	{
 		int32 CurrentReadFramebuffer = 0;
@@ -508,8 +548,17 @@ void DrawOpenGLViewport(FPlatformOpenGLContext* const Context, uint32 Width, uin
 		glClearColor(0.0f, 0.0f, 0.0f, 0.0f);
 		glClear(GL_COLOR_BUFFER_BIT);
 	}
+	
 	[Context->OpenGLContext flushBuffer];
-	UnlockGLContext(Context->OpenGLContext);
+
+	if(FMacOpenGL::SupportsCoreAnimation())
+	{
+		UnlockGLContext(Context->OpenGLContext);
+	}
+	else
+	{
+		glBindFramebuffer(GL_DRAW_FRAMEBUFFER, CurrentDrawFramebuffer);
+	}
 }
 
 void PlatformRenderingContextSetup(FPlatformOpenGLDevice* Device)
diff --git a/Engine/Source/Runtime/OpenGLDrv/Private/Mac/OpenGLMac.h b/Engine/Source/Runtime/OpenGLDrv/Private/Mac/OpenGLMac.h
--- a/Engine/Source/Runtime/OpenGLDrv/Private/Mac/OpenGLMac.h
+++ b/Engine/Source/Runtime/OpenGLDrv/Private/Mac/OpenGLMac.h
@@ -70,6 +70,8 @@ public:
 	static FORCEINLINE bool SupportsDrawIndirect()						{ return bSupportsDrawIndirect; }
 	static FORCEINLINE bool SupportsSRGBFramebuffer()					{ return bUseSRGBFramebuffer; }
 	
+	static FORCEINLINE bool SupportsCoreAnimation()						{ return !bUsingApitrace; }
+	
 	static FORCEINLINE void Flush()
 	{
 		glFlushRenderAPPLE();
