fix potential performance hit for shipping builds. #html5

diff --git a/Engine/Source/Runtime/OpenGLDrv/Private/OpenGLShaders.cpp b/Engine/Source/Runtime/OpenGLDrv/Private/OpenGLShaders.cpp
--- a/Engine/Source/Runtime/OpenGLDrv/Private/OpenGLShaders.cpp
+++ b/Engine/Source/Runtime/OpenGLDrv/Private/OpenGLShaders.cpp
@@ -720,7 +720,7 @@ ShaderType* CompileOpenGLShader(const TArray<uint8>& Code)
 			glGetShaderiv(Resource, GL_COMPILE_STATUS, &CompileStatus);
 		}
 #endif
-#if PLATFORM_HTML5
+#if PLATFORM_HTML5 && !UE_BUILD_SHIPPING
 		glGetShaderiv(Resource, GL_COMPILE_STATUS, &CompileStatus);
 		if (CompileStatus == GL_FALSE)
 		{
