fixed null pointer access when trying to set up clear colors
UE-19836

diff --git a/Engine/Source/Runtime/OpenGLDrv/Private/OpenGLCommands.cpp b/Engine/Source/Runtime/OpenGLDrv/Private/OpenGLCommands.cpp
--- a/Engine/Source/Runtime/OpenGLDrv/Private/OpenGLCommands.cpp
+++ b/Engine/Source/Runtime/OpenGLDrv/Private/OpenGLCommands.cpp
@@ -1836,9 +1836,12 @@ void FOpenGLDynamicRHI::RHISetRenderTargetsAndClear(const FRHISetRenderTargetsIn
 		{
 			for (int32 i = 0; i < RenderTargetsInfo.NumColorRenderTargets; ++i)
 			{
-				const FClearValueBinding& ClearValue = RenderTargetsInfo.ColorRenderTarget[i].Texture->GetClearBinding();
-				checkf(ClearValue.ColorBinding == EClearBinding::EColorBound, TEXT("Texture: %s does not have a color bound for fast clears"), *RenderTargetsInfo.ColorRenderTarget[i].Texture->GetName().GetPlainNameString());
-				ClearColors[i] = ClearValue.GetClearColor();
+				if (RenderTargetsInfo.ColorRenderTarget[i].Texture != nullptr)
+				{
+					const FClearValueBinding& ClearValue = RenderTargetsInfo.ColorRenderTarget[i].Texture->GetClearBinding();
+					checkf(ClearValue.ColorBinding == EClearBinding::EColorBound, TEXT("Texture: %s does not have a color bound for fast clears"), *RenderTargetsInfo.ColorRenderTarget[i].Texture->GetName().GetPlainNameString());
+					ClearColors[i] = ClearValue.GetClearColor();
+				}
 			}
 		}
 		if (RenderTargetsInfo.bClearDepth || RenderTargetsInfo.bClearStencil)
