Fixed UE-19217 - strange square reflection appears on starter content floor in Basic Code Template on Mac - the change to enable Seamless Cubemaps only when the extensions is exported breaks on OS X which conforms to GL 3.3 or 4.1 and so doesn't export it as support is implied by the GL version number. As it is always present on OS X just enable seamless mode in the Mac code.

diff --git a/Engine/Source/Runtime/OpenGLDrv/Private/Mac/MacOpenGLContext.cpp b/Engine/Source/Runtime/OpenGLDrv/Private/Mac/MacOpenGLContext.cpp
--- a/Engine/Source/Runtime/OpenGLDrv/Private/Mac/MacOpenGLContext.cpp
+++ b/Engine/Source/Runtime/OpenGLDrv/Private/Mac/MacOpenGLContext.cpp
@@ -298,6 +298,10 @@ void FPlatformOpenGLContext::Initialise( NSOpenGLContext* const SharedContext )
 	
 	InitDefaultGLContextState();
 	
+	// As someone has made enabling this rely on the extension being exported set it directly here as it is always available on OS X.
+	// Conformant desktop GL 3.2 or above do not have to export this extension as it was promoted to core!
+	glEnable(GL_TEXTURE_CUBE_MAP_SEAMLESS);
+	
 	if ( !VertexArrayObject )
 	{
 		glGenVertexArrays(1,&VertexArrayObject);
