Fix bulk-loading of cubemaps on OpenGL - we weren't incrementing the data offset for each TexSubImage operation and would have uploaded 6 copies of the first face.

diff --git a/Engine/Source/Runtime/OpenGLDrv/Private/OpenGLTexture.cpp b/Engine/Source/Runtime/OpenGLDrv/Private/OpenGLTexture.cpp
--- a/Engine/Source/Runtime/OpenGLDrv/Private/OpenGLTexture.cpp
+++ b/Engine/Source/Runtime/OpenGLDrv/Private/OpenGLTexture.cpp
@@ -453,6 +453,10 @@ FRHITexture* FOpenGLDynamicRHI::CreateOpenGLTexture(uint32 SizeX, uint32 SizeY,
 
 			for(uint32 MipIndex = 0; MipIndex < NumMips; MipIndex++)
 			{
+				uint32 NumBlocksX = FMath::Max<uint32>(1,(SizeX >> MipIndex) / GPixelFormats[Format].BlockSizeX);
+				uint32 NumBlocksY = FMath::Max<uint32>(1,(SizeY >> MipIndex) / GPixelFormats[Format].BlockSizeY);
+				uint32 NumLayers = FMath::Max<uint32>(1,ArraySize);
+				
 				if(bArrayTexture )
 				{
 					if(bCubeTexture)
@@ -477,6 +481,8 @@ FRHITexture* FOpenGLDynamicRHI::CreateOpenGLTexture(uint32 SizeX, uint32 SizeY,
 						// @todo: refactor 2d texture arrays here?
 						check(!bCubeTexture);
 					}
+					
+					MipOffset += NumBlocksX * NumBlocksY * NumLayers * GPixelFormats[Format].BlockBytes;
 				}
 				else
 				{
@@ -496,14 +502,10 @@ FRHITexture* FOpenGLDynamicRHI::CreateOpenGLTexture(uint32 SizeX, uint32 SizeY,
 							/*Type=*/ GLFormat.Type,
 							/*Data=*/ &Data[MipOffset]
 							);
+						
+						MipOffset += NumBlocksX * NumBlocksY * NumLayers * GPixelFormats[Format].BlockBytes;
 					}
 				}
-				uint32 NumBlocksX = FMath::Max<uint32>(1,(SizeX >> MipIndex) / GPixelFormats[Format].BlockSizeX);
-				uint32 NumBlocksY = FMath::Max<uint32>(1,(SizeY >> MipIndex) / GPixelFormats[Format].BlockSizeY);
-				uint32 NumLayers = FMath::Max<uint32>(1,ArraySize);
-
-				MipOffset               += NumBlocksX * NumBlocksY * NumLayers * GPixelFormats[Format].BlockBytes;
-
 			}
 		}
 	}
