MSAA for GearVR
Adds two new console vars r.MobileOnChipMSAA=0 by default as MSAA code still crashes on some non-GearVR Android devices, and gearvr.EnableMSAA=1 by default, which enables r.MobileOnChipMSAA when running on GearVR.
Jira:UEMOB-4

diff --git a/Engine/Plugins/Runtime/GearVR/Source/GearVR/Private/GearVR.cpp b/Engine/Plugins/Runtime/GearVR/Source/GearVR/Private/GearVR.cpp
--- a/Engine/Plugins/Runtime/GearVR/Source/GearVR/Private/GearVR.cpp
+++ b/Engine/Plugins/Runtime/GearVR/Source/GearVR/Private/GearVR.cpp
@@ -26,6 +26,10 @@ extern bool AndroidThunkCpp_IsGearVRApplication();
 // GearVR Plugin Implementation
 //---------------------------------------------------
 
+#if GEARVR_SUPPORTED_PLATFORMS
+static TAutoConsoleVariable<int32> CVarGearVREnableMSAA(TEXT("gearvr.EnableMSAA"), 1, TEXT("Enable MSAA when rendering on GearVR"));
+#endif
+
 class FGearVRPlugin : public IGearVRPlugin
 {
 	/** IHeadMountedDisplayModule implementation */
@@ -723,6 +727,13 @@ void FGearVR::Startup()
 	LoadFromIni();
 	SaveSystemValues();
 
+	if(CVarGearVREnableMSAA.GetValueOnAnyThread())
+	{
+		static IConsoleVariable* CVarMobileOnChipMSAA = IConsoleManager::Get().FindConsoleVariable(TEXT("r.MobileOnChipMSAA"));
+		UE_LOG(LogHMD, Log, TEXT("Enabling r.MobileOnChipMSAA, previous value %d"), CVarMobileOnChipMSAA->GetInt());
+		CVarMobileOnChipMSAA->Set(1);
+	}
+
 	UE_LOG(LogHMD, Log, TEXT("GearVR has started"));
 }
 
diff --git a/Engine/Source/Runtime/Core/Private/HAL/ConsoleManager.cpp b/Engine/Source/Runtime/Core/Private/HAL/ConsoleManager.cpp
--- a/Engine/Source/Runtime/Core/Private/HAL/ConsoleManager.cpp
+++ b/Engine/Source/Runtime/Core/Private/HAL/ConsoleManager.cpp
@@ -1908,6 +1908,14 @@ static TAutoConsoleVariable<float> CVarMobileContentScaleFactor(
 	TEXT("Content scale multiplier (equates to iOS's contentScaleFactor to support Retina displays"),
 	ECVF_Default);
 
+static TAutoConsoleVariable<int32> CVarMobileOnChipMSAA(
+	TEXT("r.MobileOnChipMSAA"),
+	0,
+	TEXT("Whether to enable on-chip MSAA for tile based mobile GPUs")
+	TEXT("0: disabed (default)\n")
+	TEXT("1: enabled\n"),
+	ECVF_Default);
+
 // this cvar can be removed in shipping to not compile shaders for development (faster)
 static TAutoConsoleVariable<int32> CVarCompileShadersForDevelopment(
 	TEXT("r.CompileShadersForDevelopment"),
diff --git a/Engine/Source/Runtime/OpenGLDrv/Private/OpenGLRenderTarget.cpp b/Engine/Source/Runtime/OpenGLDrv/Private/OpenGLRenderTarget.cpp
--- a/Engine/Source/Runtime/OpenGLDrv/Private/OpenGLRenderTarget.cpp
+++ b/Engine/Source/Runtime/OpenGLDrv/Private/OpenGLRenderTarget.cpp
@@ -116,6 +116,10 @@ GLuint FOpenGLDynamicRHI::GetOpenGLFramebuffer(uint32 NumSimultaneousRenderTarge
 		return GL_NONE;
 	}
 
+#if PLATFORM_ANDROID
+	static const auto CVarMobileOnChipMSAA = IConsoleManager::Get().FindTConsoleVariableDataInt(TEXT("r.MobileOnChipMSAA"));
+#endif
+
 	// Not found. Preparing new one.
 	GLuint Framebuffer;
 	glGenFramebuffers(1, &Framebuffer);
@@ -139,7 +143,17 @@ GLuint FOpenGLDynamicRHI::GetOpenGLFramebuffer(uint32 NumSimultaneousRenderTarge
 			{
 			case GL_TEXTURE_2D:
 			case GL_TEXTURE_2D_MULTISAMPLE:
-				FOpenGL::FramebufferTexture2D(GL_FRAMEBUFFER, GL_COLOR_ATTACHMENT0 + RenderTargetIndex, RenderTarget->Target, RenderTarget->Resource, MipmapLevels[RenderTargetIndex]);
+#if PLATFORM_ANDROID
+				if (FOpenGL::SupportsMultisampledRenderToTexture() && CVarMobileOnChipMSAA->GetValueOnRenderThread())
+				{
+					glFramebufferTexture2DMultisampleEXT(GL_FRAMEBUFFER, GL_COLOR_ATTACHMENT0 + RenderTargetIndex, RenderTarget->Target, RenderTarget->Resource, MipmapLevels[RenderTargetIndex], 2);
+					VERIFY_GL(glFramebufferTexture2DMultisampleEXT);
+				}
+				else
+#endif
+				{
+					FOpenGL::FramebufferTexture2D(GL_FRAMEBUFFER, GL_COLOR_ATTACHMENT0 + RenderTargetIndex, RenderTarget->Target, RenderTarget->Resource, MipmapLevels[RenderTargetIndex]);
+				}
 				break;
 			case GL_TEXTURE_3D:
 			case GL_TEXTURE_2D_ARRAY:
@@ -160,7 +174,17 @@ GLuint FOpenGLDynamicRHI::GetOpenGLFramebuffer(uint32 NumSimultaneousRenderTarge
 			case GL_TEXTURE_2D:
 			case GL_TEXTURE_2D_MULTISAMPLE:
 				check( ArrayIndices[RenderTargetIndex] == 0);
-				FOpenGL::FramebufferTexture2D(GL_FRAMEBUFFER, GL_COLOR_ATTACHMENT0 + RenderTargetIndex, RenderTarget->Target, RenderTarget->Resource, MipmapLevels[RenderTargetIndex]);
+#if PLATFORM_ANDROID
+				if (FOpenGL::SupportsMultisampledRenderToTexture() && CVarMobileOnChipMSAA->GetValueOnRenderThread())
+				{
+					glFramebufferTexture2DMultisampleEXT(GL_FRAMEBUFFER, GL_COLOR_ATTACHMENT0 + RenderTargetIndex, RenderTarget->Target, RenderTarget->Resource, MipmapLevels[RenderTargetIndex], 2);
+					VERIFY_GL(glFramebufferTexture2DMultisampleEXT);
+				}
+				else
+#endif
+				{
+					FOpenGL::FramebufferTexture2D(GL_FRAMEBUFFER, GL_COLOR_ATTACHMENT0 + RenderTargetIndex, RenderTarget->Target, RenderTarget->Resource, MipmapLevels[RenderTargetIndex]);
+				}
 				break;
 			case GL_TEXTURE_3D:
 				FOpenGL::FramebufferTexture3D(GL_FRAMEBUFFER,GL_COLOR_ATTACHMENT0 + RenderTargetIndex, RenderTarget->Target, RenderTarget->Resource, MipmapLevels[RenderTargetIndex], ArrayIndices[RenderTargetIndex]);
@@ -188,14 +212,31 @@ GLuint FOpenGLDynamicRHI::GetOpenGLFramebuffer(uint32 NumSimultaneousRenderTarge
 		{
 		case GL_TEXTURE_2D:
 		case GL_TEXTURE_2D_MULTISAMPLE:
-			if (!FOpenGL::SupportsCombinedDepthStencilAttachment() && DepthStencilTarget->Attachment == GL_DEPTH_STENCIL_ATTACHMENT)
+#if PLATFORM_ANDROID
+			if (FOpenGL::SupportsMultisampledRenderToTexture() && CVarMobileOnChipMSAA->GetValueOnRenderThread())
 			{
-				FOpenGL::FramebufferTexture2D(GL_FRAMEBUFFER, GL_DEPTH_ATTACHMENT, DepthStencilTarget->Target, DepthStencilTarget->Resource, 0);
-				FOpenGL::FramebufferTexture2D(GL_FRAMEBUFFER, GL_STENCIL_ATTACHMENT, DepthStencilTarget->Target, DepthStencilTarget->Resource, 0);
+				FOpenGLTexture2D* DepthStencilTarget2D = (FOpenGLTexture2D*)DepthStencilTarget;
+				GLuint DepthBuffer;
+				glGenRenderbuffers(1, &DepthBuffer);
+				glBindRenderbuffer(GL_RENDERBUFFER, DepthBuffer);
+				glRenderbufferStorageMultisampleEXT(GL_RENDERBUFFER, 2, FOpenGL::SupportsPackedDepthStencil() ? GL_DEPTH24_STENCIL8 : GL_DEPTH_COMPONENT24, DepthStencilTarget2D->GetSizeX(), DepthStencilTarget2D->GetSizeY());
+				VERIFY_GL(glRenderbufferStorageMultisampleEXT);
+				glBindRenderbuffer(GL_RENDERBUFFER, 0);
+				glFramebufferRenderbuffer(GL_FRAMEBUFFER, GL_DEPTH_ATTACHMENT, GL_RENDERBUFFER, DepthBuffer);
+				VERIFY_GL(glFramebufferRenderbuffer);
 			}
 			else
+#endif
 			{
-				FOpenGL::FramebufferTexture2D(GL_FRAMEBUFFER, DepthStencilTarget->Attachment, DepthStencilTarget->Target, DepthStencilTarget->Resource, 0);
+				if (!FOpenGL::SupportsCombinedDepthStencilAttachment() && DepthStencilTarget->Attachment == GL_DEPTH_STENCIL_ATTACHMENT)
+				{
+					FOpenGL::FramebufferTexture2D(GL_FRAMEBUFFER, GL_DEPTH_ATTACHMENT, DepthStencilTarget->Target, DepthStencilTarget->Resource, 0);
+					FOpenGL::FramebufferTexture2D(GL_FRAMEBUFFER, GL_STENCIL_ATTACHMENT, DepthStencilTarget->Target, DepthStencilTarget->Resource, 0);
+				}
+				else
+				{
+					FOpenGL::FramebufferTexture2D(GL_FRAMEBUFFER, DepthStencilTarget->Attachment, DepthStencilTarget->Target, DepthStencilTarget->Resource, 0);
+				}
 			}
 			break;
 		case GL_TEXTURE_3D:
