Disable GL_AMD_debug_output for OpenGL ES2 and ES3.1
#PR #1493
#ue4

diff --git a/Engine/Source/Runtime/OpenGLDrv/Public/OpenGLES2.h b/Engine/Source/Runtime/OpenGLDrv/Public/OpenGLES2.h
--- a/Engine/Source/Runtime/OpenGLDrv/Public/OpenGLES2.h
+++ b/Engine/Source/Runtime/OpenGLDrv/Public/OpenGLES2.h
@@ -20,6 +20,9 @@ typedef GLfloat GLdouble;
 #include "OpenGL.h"
 #include "OpenGLUtil.h"		// for VERIFY_GL
 
+#ifdef GL_AMD_debug_output
+	#undef GL_AMD_debug_output
+#endif
 
 // Redefine to disable support for pixel buffer objects
 #ifdef UGL_SUPPORTS_PIXELBUFFERS
diff --git a/Engine/Source/Runtime/OpenGLDrv/Public/OpenGLES31.h b/Engine/Source/Runtime/OpenGLDrv/Public/OpenGLES31.h
--- a/Engine/Source/Runtime/OpenGLDrv/Public/OpenGLES31.h
+++ b/Engine/Source/Runtime/OpenGLDrv/Public/OpenGLES31.h
@@ -11,6 +11,10 @@
 
 #include "OpenGL.h"
 
+#ifdef GL_AMD_debug_output
+	#undef GL_AMD_debug_output
+#endif
+
 /** Unreal tokens that maps to different OpenGL tokens by platform. */
 #undef UGL_ANY_SAMPLES_PASSED
 #define UGL_ANY_SAMPLES_PASSED	GL_ANY_SAMPLES_PASSED_EXT
@@ -22,6 +26,7 @@
 #define USE_OPENGL_NAME_CACHE 1
 #define OPENGL_NAME_CACHE_SIZE 1024
 
+
 struct FOpenGLES31 : public FOpenGLBase
 {
 	static FORCEINLINE bool SupportsVertexArrayObjects()				{ return bSupportsVertexArrayObjects || !bES2Fallback; }
