Fix OpenGL ES 3.1 compile errors from adding MSAA to Android
#PR #1497
#ue4
#android

diff --git a/Engine/Source/Runtime/OpenGLDrv/Private/Android/AndroidES31OpenGL.cpp b/Engine/Source/Runtime/OpenGLDrv/Private/Android/AndroidES31OpenGL.cpp
--- a/Engine/Source/Runtime/OpenGLDrv/Private/Android/AndroidES31OpenGL.cpp
+++ b/Engine/Source/Runtime/OpenGLDrv/Private/Android/AndroidES31OpenGL.cpp
@@ -175,6 +175,10 @@ void FPlatformOpenGLDevice::Init()
 
 	}
 
+	// For MSAA
+	glFramebufferTexture2DMultisampleEXT = (PFNGLFRAMEBUFFERTEXTURE2DMULTISAMPLEEXTPROC)((void*)eglGetProcAddress("glFramebufferTexture2DMultisampleEXT"));
+	glRenderbufferStorageMultisampleEXT = (PFNGLRENDERBUFFERSTORAGEMULTISAMPLEEXTPROC)((void*)eglGetProcAddress("glRenderbufferStorageMultisampleEXT"));
+
 	if (!bFoundAllEntryPoints)
 	{
 		UE_LOG(LogRHI, Warning, TEXT("Failed to acquire all optional OpenGL entrypoints, may fall back to OpenGL ES 2.0"));
diff --git a/Engine/Source/Runtime/OpenGLDrv/Private/Android/AndroidES31OpenGL.h b/Engine/Source/Runtime/OpenGLDrv/Private/Android/AndroidES31OpenGL.h
--- a/Engine/Source/Runtime/OpenGLDrv/Private/Android/AndroidES31OpenGL.h
+++ b/Engine/Source/Runtime/OpenGLDrv/Private/Android/AndroidES31OpenGL.h
@@ -227,6 +227,8 @@ GL_APICALL void GL_APIENTRY glGetVertexAttribLui64vNV(GLuint index, GLenum pname
 	EnumMacro(PFNGLUNMAPBUFFEROESPROC, glUnmapBufferOES) \
 	EnumMacro(PFNGLQUERYCOUNTEREXTPROC, glQueryCounterEXT) \
 	EnumMacro(PFNGLGETQUERYOBJECTUI64VEXTPROC, glGetQueryObjectui64vEXT) \
+	EnumMacro(PFNGLFRAMEBUFFERTEXTURE2DMULTISAMPLEEXTPROC, glFramebufferTexture2DMultisampleEXT) \
+	EnumMacro(PFNGLRENDERBUFFERSTORAGEMULTISAMPLEEXTPROC, glRenderbufferStorageMultisampleEXT) \
 	EnumMacro(PFNGLGETTEXTUREHANDLENVPROC, glGetTextureHandleNV)\
 	EnumMacro(PFNGLGETTEXTURESAMPLERHANDLENVPROC, glGetTextureSamplerHandleNV)\
 	EnumMacro(PFNGLMAKETEXTUREHANDLERESIDENTNVPROC, glMakeTextureHandleResidentNV)\
diff --git a/Engine/Source/Runtime/OpenGLDrv/Private/Android/AndroidOpenGL.cpp b/Engine/Source/Runtime/OpenGLDrv/Private/Android/AndroidOpenGL.cpp
--- a/Engine/Source/Runtime/OpenGLDrv/Private/Android/AndroidOpenGL.cpp
+++ b/Engine/Source/Runtime/OpenGLDrv/Private/Android/AndroidOpenGL.cpp
@@ -28,8 +28,8 @@ PFNGLGETQUERYOBJECTUI64VEXTPROC			glGetQueryObjectui64vEXT = NULL;
 // Offscreen MSAA rendering
 PFNBLITFRAMEBUFFERNVPROC				glBlitFramebufferNV = NULL;
 PFNGLDISCARDFRAMEBUFFEREXTPROC			glDiscardFramebufferEXT = NULL;
-PFNGLFRAMEBUFFERTEXTURE2DMULTISAMPLEEXT	glFramebufferTexture2DMultisampleEXT = NULL;
-PFNGLRENDERBUFFERSTORAGEMULTISAMPLEEXT	glRenderbufferStorageMultisampleEXT = NULL;
+PFNGLFRAMEBUFFERTEXTURE2DMULTISAMPLEEXTPROC	glFramebufferTexture2DMultisampleEXT = NULL;
+PFNGLRENDERBUFFERSTORAGEMULTISAMPLEEXTPROC	glRenderbufferStorageMultisampleEXT = NULL;
 
 PFNGLPUSHGROUPMARKEREXTPROC				glPushGroupMarkerEXT = NULL;
 PFNGLPOPGROUPMARKEREXTPROC				glPopGroupMarkerEXT = NULL;
@@ -342,8 +342,8 @@ void FAndroidOpenGL::ProcessExtensions(const FString& ExtensionsString)
 	}
 
 	glDiscardFramebufferEXT = (PFNGLDISCARDFRAMEBUFFEREXTPROC)((void*)eglGetProcAddress("glDiscardFramebufferEXT"));
-	glFramebufferTexture2DMultisampleEXT = (PFNGLFRAMEBUFFERTEXTURE2DMULTISAMPLEEXT)((void*)eglGetProcAddress("glFramebufferTexture2DMultisampleEXT"));
-	glRenderbufferStorageMultisampleEXT = (PFNGLRENDERBUFFERSTORAGEMULTISAMPLEEXT)((void*)eglGetProcAddress("glRenderbufferStorageMultisampleEXT"));
+	glFramebufferTexture2DMultisampleEXT = (PFNGLFRAMEBUFFERTEXTURE2DMULTISAMPLEEXTPROC)((void*)eglGetProcAddress("glFramebufferTexture2DMultisampleEXT"));
+	glRenderbufferStorageMultisampleEXT = (PFNGLRENDERBUFFERSTORAGEMULTISAMPLEEXTPROC)((void*)eglGetProcAddress("glRenderbufferStorageMultisampleEXT"));
 	glPushGroupMarkerEXT = (PFNGLPUSHGROUPMARKEREXTPROC)((void*)eglGetProcAddress("glPushGroupMarkerEXT"));
 	glPopGroupMarkerEXT = (PFNGLPOPGROUPMARKEREXTPROC)((void*)eglGetProcAddress("glPopGroupMarkerEXT"));
 	glLabelObjectEXT = (PFNGLLABELOBJECTEXTPROC)((void*)eglGetProcAddress("glLabelObjectEXT"));
diff --git a/Engine/Source/Runtime/OpenGLDrv/Private/Android/AndroidOpenGL.h b/Engine/Source/Runtime/OpenGLDrv/Private/Android/AndroidOpenGL.h
--- a/Engine/Source/Runtime/OpenGLDrv/Private/Android/AndroidOpenGL.h
+++ b/Engine/Source/Runtime/OpenGLDrv/Private/Android/AndroidOpenGL.h
@@ -76,8 +76,8 @@ typedef void (GL_APIENTRYP PFNGLPUSHGROUPMARKEREXTPROC) (GLsizei length, const G
 typedef void (GL_APIENTRYP PFNGLLABELOBJECTEXTPROC) (GLenum type, GLuint object, GLsizei length, const GLchar *label);
 typedef void (GL_APIENTRYP PFNGLGETOBJECTLABELEXTPROC) (GLenum type, GLuint object, GLsizei bufSize, GLsizei *length, GLchar *label);
 typedef void (GL_APIENTRYP PFNGLPOPGROUPMARKEREXTPROC) (void);
-typedef void (GL_APIENTRYP PFNGLFRAMEBUFFERTEXTURE2DMULTISAMPLEEXT) (GLenum target, GLenum attachment, GLenum textarget, GLuint texture, GLint level, GLsizei samples);
-typedef void (GL_APIENTRYP PFNGLRENDERBUFFERSTORAGEMULTISAMPLEEXT) (GLenum target, GLsizei samples, GLenum internalformat, GLsizei width, GLsizei height);
+typedef void (GL_APIENTRYP PFNGLFRAMEBUFFERTEXTURE2DMULTISAMPLEEXTPROC) (GLenum target, GLenum attachment, GLenum textarget, GLuint texture, GLint level, GLsizei samples);
+typedef void (GL_APIENTRYP PFNGLRENDERBUFFERSTORAGEMULTISAMPLEEXTPROC) (GLenum target, GLsizei samples, GLenum internalformat, GLsizei width, GLsizei height);
 /** from ES 3.0 but can be called on certain Adreno devices */
 typedef void (GL_APIENTRYP PFNGLTEXSTORAGE2DPROC) (GLenum target, GLsizei levels, GLenum internalformat, GLsizei width, GLsizei height);
 
@@ -94,8 +94,8 @@ extern PFNGLGETQUERYOBJECTUI64VEXTPROC	glGetQueryObjectui64vEXT;
 extern PFNGLMAPBUFFEROESPROC			glMapBufferOES;
 extern PFNGLUNMAPBUFFEROESPROC			glUnmapBufferOES;
 extern PFNGLDISCARDFRAMEBUFFEREXTPROC 	glDiscardFramebufferEXT ;
-extern PFNGLFRAMEBUFFERTEXTURE2DMULTISAMPLEEXT	glFramebufferTexture2DMultisampleEXT;
-extern PFNGLRENDERBUFFERSTORAGEMULTISAMPLEEXT	glRenderbufferStorageMultisampleEXT;
+extern PFNGLFRAMEBUFFERTEXTURE2DMULTISAMPLEEXTPROC	glFramebufferTexture2DMultisampleEXT;
+extern PFNGLRENDERBUFFERSTORAGEMULTISAMPLEEXTPROC	glRenderbufferStorageMultisampleEXT;
 extern PFNGLPUSHGROUPMARKEREXTPROC		glPushGroupMarkerEXT;
 extern PFNGLLABELOBJECTEXTPROC			glLabelObjectEXT;
 extern PFNGLGETOBJECTLABELEXTPROC		glGetObjectLabelEXT;
diff --git a/Engine/Source/Runtime/OpenGLDrv/Private/OpenGLES31.cpp b/Engine/Source/Runtime/OpenGLDrv/Private/OpenGLES31.cpp
--- a/Engine/Source/Runtime/OpenGLDrv/Private/OpenGLES31.cpp
+++ b/Engine/Source/Runtime/OpenGLDrv/Private/OpenGLES31.cpp
@@ -73,6 +73,9 @@ bool FOpenGLES31::bSupportsShaderFramebufferFetch = false;
 /** GL_ARM_shader_framebuffer_fetch_depth_stencil */
 bool FOpenGLES31::bSupportsShaderDepthStencilFetch = false;
 
+/** GL_EXT_multisampled_render_to_texture */
+bool FOpenGLES31::bSupportsMultisampledRenderToTexture = false;
+
 /** GL_EXT_sRGB */
 bool FOpenGLES31::bSupportsSGRB = false;
 
@@ -281,6 +284,7 @@ void FOpenGLES31::ProcessExtensions( const FString& ExtensionsString )
 	bSupportsNvImageFormats = ExtensionsString.Contains(TEXT("GL_NV_image_formats"));
 	bSupportsShaderFramebufferFetch = ExtensionsString.Contains(TEXT("GL_EXT_shader_framebuffer_fetch")) || ExtensionsString.Contains(TEXT("GL_NV_shader_framebuffer_fetch")) || ExtensionsString.Contains(TEXT("GL_ARM_shader_framebuffer_fetch"));
 	bSupportsShaderDepthStencilFetch = ExtensionsString.Contains(TEXT("GL_ARM_shader_framebuffer_fetch_depth_stencil"));
+	bSupportsMultisampledRenderToTexture = ExtensionsString.Contains(TEXT("GL_EXT_multisampled_render_to_texture"));
 	// @todo es3: SRGB support does not work with our texture format setup (ES2 docs indicate that internalFormat and format must match, but they don't at all with sRGB enabled)
 	//             One possible solution us to use GLFormat.InternalFormat[bSRGB] instead of GLFormat.Format
 	bSupportsSGRB = false;//ExtensionsString.Contains(TEXT("GL_EXT_sRGB"));
diff --git a/Engine/Source/Runtime/OpenGLDrv/Public/OpenGLES31.h b/Engine/Source/Runtime/OpenGLDrv/Public/OpenGLES31.h
--- a/Engine/Source/Runtime/OpenGLDrv/Public/OpenGLES31.h
+++ b/Engine/Source/Runtime/OpenGLDrv/Public/OpenGLES31.h
@@ -71,6 +71,7 @@ struct FOpenGLES31 : public FOpenGLBase
 	static FORCEINLINE bool SupportsR11G11B10F()						{ return bSupportsNvImageFormats && !bES2Fallback; }
 	static FORCEINLINE bool SupportsShaderFramebufferFetch()			{ return bSupportsShaderFramebufferFetch; }
 	static FORCEINLINE bool SupportsShaderDepthStencilFetch()			{ return bSupportsShaderDepthStencilFetch; }
+	static FORCEINLINE bool SupportsMultisampledRenderToTexture()		{ return bSupportsMultisampledRenderToTexture; }
 	static FORCEINLINE bool SupportsVertexArrayBGRA()					{ return false; }
 	static FORCEINLINE bool SupportsBGRA8888()							{ return bSupportsBGRA8888; }
 	static FORCEINLINE bool SupportsSRGB()								{ return bSupportsSGRB || !bES2Fallback; }
@@ -904,6 +905,9 @@ protected:
 	/** GL_ARM_shader_framebuffer_fetch_depth_stencil */
 	static bool bSupportsShaderDepthStencilFetch;
 
+	/** GL_EXT_MULTISAMPLED_RENDER_TO_TEXTURE */
+	static bool bSupportsMultisampledRenderToTexture;
+
 	/** GL_FRAGMENT_SHADER, GL_LOW_FLOAT */
 	static int ShaderLowPrecision;
 
