Improve compressed 2D texture streaming on Mac OpenGL by retaining the pixel buffer objects for 2D compressed textures, this removes calls to glGenBuffers, glBufferData and glGetCompressedTexImage that were previously used to retreive the image data from the old texture. Now all that needs to be done is to push the existing data into the new texture. This reduces hitching during streaming and the effect is most pronounced on Nvidia Mac GL drivers which perform a GL command stream flush in glGetCompressedTexImage. As this trades memory for performance it is restricted to PLATFORM_MAC as we can safely assume that Macs will have sufficient memory to cope but other GL devices may not.

diff --git a/Engine/Source/Runtime/OpenGLDrv/Private/OpenGLTexture.cpp b/Engine/Source/Runtime/OpenGLDrv/Private/OpenGLTexture.cpp
--- a/Engine/Source/Runtime/OpenGLDrv/Private/OpenGLTexture.cpp
+++ b/Engine/Source/Runtime/OpenGLDrv/Private/OpenGLTexture.cpp
@@ -1059,7 +1059,10 @@ void TOpenGLTexture<RHIResourceType>::Unlock(uint32 MipIndex,uint32 ArrayIndex)
 		}
 
 		//need to free PBO if we aren't keeping shadow copies
-		PixelBuffers[BufferIndex] = NULL;
+		if(!PLATFORM_MAC || !GLFormat.bCompressed || Target != GL_TEXTURE_2D)
+		{
+			PixelBuffers[BufferIndex] = NULL;
+		}
 	}
 	else if(!bUseClientStorage || !ClientStorageBuffers[BufferIndex].bReadOnly)
 	{
@@ -1248,6 +1251,15 @@ void TOpenGLTexture<RHIResourceType>::CloneViaPBO( TOpenGLTexture<RHIResourceTyp
 
 			const uint32 MipBytes = NumBlocksX * NumBlocksY * BlockBytes;
 			const int32 BufferIndex = DstMipIndex * (bCubemap ? 6 : 1) * this->GetEffectiveSizeZ() + ArrayIndex;
+			const int32 SrcBufferIndex = SrcMipIndex * (Src->bCubemap ? 6 : 1) * Src->GetEffectiveSizeZ() + ArrayIndex;
+			
+			// Retain the existing PBO for this texture data - as it is compressed it won't change
+			if(PLATFORM_MAC && GLFormat.bCompressed && Target == GL_TEXTURE_2D)
+			{
+				PixelBuffers[BufferIndex] = Src->PixelBuffers[SrcBufferIndex];
+				check(PixelBuffers[BufferIndex]->GetSize() == MipBytes);
+				check(!PixelBuffers[BufferIndex]->IsLocked());
+			}
 			
 			// Standard path with a PBO mirroring ever slice of a texture to allow multiple simulataneous maps
 			if (!IsValidRef(PixelBuffers[BufferIndex]))
@@ -1261,6 +1273,7 @@ void TOpenGLTexture<RHIResourceType>::CloneViaPBO( TOpenGLTexture<RHIResourceTyp
 			
 			// Transfer data from texture to pixel buffer.
 			// This may be further optimized by caching information if surface content was changed since last lock.
+			if(!PLATFORM_MAC || !GLFormat.bCompressed || !IsValidRef(Src->PixelBuffers[SrcBufferIndex]))
 			{
 				// Use a texture stage that's not likely to be used for draws, to avoid waiting
 				OpenGLRHI->CachedSetupTextureStage(ContextState, FOpenGL::GetMaxCombinedTextureImageUnits() - 1, Src->Target, Src->Resource, -1, this->GetNumMips());
@@ -1391,8 +1404,11 @@ void TOpenGLTexture<RHIResourceType>::CloneViaPBO( TOpenGLTexture<RHIResourceTyp
 				}
 			}
 			
-			// need to free PBO if we aren't keeping shadow copies
-			PixelBuffers[BufferIndex] = NULL;
+			if(!PLATFORM_MAC || !GLFormat.bCompressed || Target != GL_TEXTURE_2D)
+			{
+				// need to free PBO if we aren't keeping shadow copies
+				PixelBuffers[BufferIndex] = NULL;
+			}
 			
 			// No need to restore texture stage; leave it like this,
 			// and the next draw will take care of cleaning it up; or
