On OS X 10.10.5 or later disable OpenGL.UseMapBuffer when using MTGL to reduce the number of context synchronisation points & on OS X 10.11.0 or later re-enable GPU timestamp emulation on Nvidia GPUs as it now appears to work consistently.

diff --git a/Engine/Source/Runtime/OpenGLDrv/Private/Mac/MacOpenGLContext.cpp b/Engine/Source/Runtime/OpenGLDrv/Private/Mac/MacOpenGLContext.cpp
--- a/Engine/Source/Runtime/OpenGLDrv/Private/Mac/MacOpenGLContext.cpp
+++ b/Engine/Source/Runtime/OpenGLDrv/Private/Mac/MacOpenGLContext.cpp
@@ -165,6 +165,13 @@ static NSOpenGLContext* CreateContext( NSOpenGLContext* SharedContext )
 		CGLEnable((CGLContextObj)[Context CGLContextObj], kCGLCEMPEngine);
 		
 		// @todo: We should disable OpenGL.UseMapBuffer for improved performance under MTGL, but radr://17662549 prevents it when using GL_TEXTURE_BUFFER's for skinning.
+		// Disable OpenGL.UseMapBuffer when using MTGL to reduce the number of context synchronisation points.
+		// All calls to glMapBuffer/Range will stall the MTGL thread, even with the unsynchronized bit set, so we want to avoid it.
+		static auto CVar = IConsoleManager::Get().FindConsoleVariable(TEXT("OpenGL.UseMapBuffer"));
+		if(FPlatformMisc::MacOSXVersionCompare(10,10,5) >= 0 && CVar && CVar->GetInt() == 1)
+		{
+			CVar->Set(0);
+		}
 	}
 
 	return Context;
diff --git a/Engine/Source/Runtime/OpenGLDrv/Private/Mac/MacOpenGLQuery.cpp b/Engine/Source/Runtime/OpenGLDrv/Private/Mac/MacOpenGLQuery.cpp
--- a/Engine/Source/Runtime/OpenGLDrv/Private/Mac/MacOpenGLQuery.cpp
+++ b/Engine/Source/Runtime/OpenGLDrv/Private/Mac/MacOpenGLQuery.cpp
@@ -322,7 +322,7 @@ FMacOpenGLQueryEmu::FMacOpenGLQueryEmu(FPlatformOpenGLContext* InContext)
 	
 	// Only use the timestamp emulation in a non-shipping build - end-users shouldn't care about this profiling feature.
 #if (!UE_BUILD_SHIPPING)
-	GIsEmulatingTimestamp = ((FPlatformMisc::MacOSXVersionCompare(10,10,0) >= 0 && !IsRHIDeviceNVIDIA()) || GMacIsEmulatingTimestamp == 2) && !FParse::Param(FCommandLine::Get(), TEXT("DisableMacGPUTimestamp"));
+	GIsEmulatingTimestamp = ((FPlatformMisc::MacOSXVersionCompare(10,10,0) >= 0 && !IsRHIDeviceNVIDIA()) || FPlatformMisc::MacOSXVersionCompare(10,11,0) >= 0 || GMacIsEmulatingTimestamp == 2) && !FParse::Param(FCommandLine::Get(), TEXT("DisableMacGPUTimestamp"));
 #endif
 	
 	if ( GIsEmulatingTimestamp )
