When orphaning an OpenGL buffer during a write-lock that will use glBufferData to update the contents the initial discard can set the buffer size to 0 and avoid a redundant attempt to allocate backing storage. This is a marginal performance improvement on AMD's OS X drivers which don't optimise this away and could plausibly benefit other drivers too.

diff --git a/Engine/Source/Runtime/OpenGLDrv/Public/OpenGLResources.h b/Engine/Source/Runtime/OpenGLDrv/Public/OpenGLResources.h
--- a/Engine/Source/Runtime/OpenGLDrv/Public/OpenGLResources.h
+++ b/Engine/Source/Runtime/OpenGLDrv/Public/OpenGLResources.h
@@ -192,18 +192,28 @@ public:
 		bIsLocked = true;
 		bIsLockReadOnly = bReadOnly;
 		uint8 *Data = NULL;
-
+		
+		// Discard if the input size is the same as the backing store size, regardless of the input argument, as orphaning the backing store will typically be faster.
+		bDiscard = bDiscard || (!bReadOnly && InSize == RealSize);
+		
+		// Map buffer is faster in some circumstances and slower in others, decide when to use it carefully.
+		bool const bCanUseMapBuffer = FOpenGL::SupportsMapBuffer() && BaseType::GLSupportsType();
+		bool const bUseMapBuffer = bCanUseMapBuffer && (bReadOnly || OpenGLConsoleVariables::bUseMapBuffer);
+		
 		// If we're able to discard the current data, do so right away
+		// If we can then we should orphan the buffer name & reallocate the backing store only once as calls to glBufferData may do so even when the size is the same.
+		uint32 DiscardSize = (bDiscard && !bUseMapBuffer && InSize == RealSize) ? 0 : RealSize;
+		
 		// Don't call BufferData if Bindless is on, as bindless texture buffers make buffers immutable
 		if ( bDiscard && !OpenGLConsoleVariables::bBindlessTexture )
 		{
 			if (BaseType::GLSupportsType())
 			{
-				glBufferData( Type, RealSize, NULL, GetAccess());
+				glBufferData( Type, DiscardSize, NULL, GetAccess());
 			}
 		}
 
-		if ( FOpenGL::SupportsMapBuffer() && BaseType::GLSupportsType() && (OpenGLConsoleVariables::bUseMapBuffer || bReadOnly))
+		if ( bUseMapBuffer)
 		{
 			FOpenGL::EResourceLockMode LockMode = bReadOnly ? FOpenGL::RLM_ReadOnly : FOpenGL::RLM_WriteOnly;
 			Data = static_cast<uint8*>( FOpenGL::MapBufferRange( Type, InOffset, InSize, LockMode ) );
@@ -243,14 +253,27 @@ public:
 		bIsLockReadOnly = false;
 		uint8 *Data = NULL;
 
+		// Discard if the input size is the same as the backing store size, regardless of the input argument, as orphaning the backing store will typically be faster.
+		bDiscard = bDiscard || InSize == RealSize;
+		
+		// Map buffer is faster in some circumstances and slower in others, decide when to use it carefully.
+		bool const bCanUseMapBuffer = FOpenGL::SupportsMapBuffer() && BaseType::GLSupportsType();
+		bool const bUseMapBuffer = bCanUseMapBuffer && OpenGLConsoleVariables::bUseMapBuffer;
+		
 		// If we're able to discard the current data, do so right away
+		// If we can then we should orphan the buffer name & reallocate the backing store only once as calls to glBufferData may do so even when the size is the same.
+		uint32 DiscardSize = (bDiscard && !bUseMapBuffer && InSize == RealSize) ? 0 : RealSize;
+		
 		// Don't call BufferData if Bindless is on, as bindless texture buffers make buffers immutable
-		if (bDiscard && !OpenGLConsoleVariables::bBindlessTexture)
+		if ( bDiscard && !OpenGLConsoleVariables::bBindlessTexture )
 		{
-			glBufferData( Type, RealSize, NULL, GetAccess());
+			if (BaseType::GLSupportsType())
+			{
+				glBufferData( Type, DiscardSize, NULL, GetAccess());
+			}
 		}
-
-		if ( FOpenGL::SupportsMapBuffer() && OpenGLConsoleVariables::bUseMapBuffer)
+		
+		if ( bUseMapBuffer)
 		{
 			FOpenGL::EResourceLockMode LockMode = bDiscard ? FOpenGL::RLM_WriteOnly : FOpenGL::RLM_WriteOnlyUnsynchronized;
 			Data = static_cast<uint8*>( FOpenGL::MapBufferRange( Type, InOffset, InSize, LockMode ) );
