Fix a bug introduced in 2686069 that caused a crash on buffer unlock on Windows/Linux GL - we can't do optimised discard logic when RESTRICT_SUBDATA_SIZE  is enabled.

diff --git a/Engine/Source/Runtime/OpenGLDrv/Public/OpenGLResources.h b/Engine/Source/Runtime/OpenGLDrv/Public/OpenGLResources.h
--- a/Engine/Source/Runtime/OpenGLDrv/Public/OpenGLResources.h
+++ b/Engine/Source/Runtime/OpenGLDrv/Public/OpenGLResources.h
@@ -202,7 +202,7 @@ public:
 		
 		// If we're able to discard the current data, do so right away
 		// If we can then we should orphan the buffer name & reallocate the backing store only once as calls to glBufferData may do so even when the size is the same.
-		uint32 DiscardSize = (bDiscard && !bUseMapBuffer && InSize == RealSize) ? 0 : RealSize;
+		uint32 DiscardSize = (bDiscard && !bUseMapBuffer && InSize == RealSize && !RESTRICT_SUBDATA_SIZE) ? 0 : RealSize;
 		
 		// Don't call BufferData if Bindless is on, as bindless texture buffers make buffers immutable
 		if ( bDiscard && !OpenGLConsoleVariables::bBindlessTexture )
@@ -262,7 +262,7 @@ public:
 		
 		// If we're able to discard the current data, do so right away
 		// If we can then we should orphan the buffer name & reallocate the backing store only once as calls to glBufferData may do so even when the size is the same.
-		uint32 DiscardSize = (bDiscard && !bUseMapBuffer && InSize == RealSize) ? 0 : RealSize;
+		uint32 DiscardSize = (bDiscard && !bUseMapBuffer && InSize == RealSize && !RESTRICT_SUBDATA_SIZE) ? 0 : RealSize;
 		
 		// Don't call BufferData if Bindless is on, as bindless texture buffers make buffers immutable
 		if ( bDiscard && !OpenGLConsoleVariables::bBindlessTexture )
