Always define RESTRICT_SUBDATA_SIZE in OpenGLResources.h to correctly discard depending on which path we have to take.

diff --git a/Engine/Source/Runtime/OpenGLDrv/Public/OpenGLResources.h b/Engine/Source/Runtime/OpenGLDrv/Public/OpenGLResources.h
--- a/Engine/Source/Runtime/OpenGLDrv/Public/OpenGLResources.h
+++ b/Engine/Source/Runtime/OpenGLDrv/Public/OpenGLResources.h
@@ -33,6 +33,8 @@ namespace OpenGLConsoleVariables
 
 #if PLATFORM_WINDOWS || PLATFORM_ANDROIDES31
 #define RESTRICT_SUBDATA_SIZE 1
+#else 
+#define RESTRICT_SUBDATA_SIZE 0
 #endif
 
 void IncrementBufferMemory(GLenum Type, bool bStructuredBuffer, uint32 NumBytes);
