Linux: fix crash on invalid windows.

- PR #1589: "Linux SDL2 X11 OpenGL related crash fix "contributed by [REDACTED_PERSON] (UE-21414).

diff --git a/Engine/Source/Runtime/OpenGLDrv/Private/Linux/OpenGLLinux.cpp b/Engine/Source/Runtime/OpenGLDrv/Private/Linux/OpenGLLinux.cpp
--- a/Engine/Source/Runtime/OpenGLDrv/Private/Linux/OpenGLLinux.cpp
+++ b/Engine/Source/Runtime/OpenGLDrv/Private/Linux/OpenGLLinux.cpp
@@ -399,6 +399,12 @@ bool PlatformBlitToViewport(FPlatformOpenGLDevice* Device,
 		GLenum BlitFilter;
 		GLint DestX0, DestY0, DestX1, DestY1;
 
+		if ( WinH == 0 || WinW == 0 )
+		{
+			// Nothing to blit
+			return false;
+		}
+
 		if ( ( WinW == BackbufferSizeX ) && ( WinH == BackbufferSizeY ) )
 		{
 			// We match up. We're probably in windowed mode, or an exact
