Fix android startup crash in shader hack testing

diff --git a/Engine/Source/Runtime/OpenGLDrv/Private/OpenGLDevice.cpp b/Engine/Source/Runtime/OpenGLDrv/Private/OpenGLDevice.cpp
--- a/Engine/Source/Runtime/OpenGLDrv/Private/OpenGLDevice.cpp
+++ b/Engine/Source/Runtime/OpenGLDrv/Private/OpenGLDevice.cpp
@@ -1127,7 +1127,7 @@ static void CheckTextureCubeLodSupport()
 		FOpenGL::bIsCheckingShaderCompilerHacks = true;
 
 		// This code creates a sample program and finds out which hacks are required to compile it
-		const ANSICHAR* TestFragmentProgram = "\n"
+		static const ANSICHAR TestFragmentProgram[] = "\n"
 			"#version 100\n"
 			"#ifndef DONTEMITEXTENSIONSHADERTEXTURELODENABLE\n"
 			"#extension GL_EXT_shader_texture_lod : enable\n"
@@ -1145,21 +1145,24 @@ static void CheckTextureCubeLodSupport()
 			"	gl_FragColor = textureCubeLodEXT(Texture,TexCoord, 4.0);\n"
 			"}\n";
 
-		/*TArray<uint8> Code;
-		Code.Add( strlen( TestFragmentProgram ) + 1 );
-		FMemory::Memcpy( Code.GetData(), TestFragmentProgram, Code.Num() );*/
+		FShaderCode ShaderCode;
+		{
+			TArray<uint8>& Code = ShaderCode.GetWriteAccess();
+			Code.AddUninitialized(sizeof(TestFragmentProgram));
+			FMemory::Memcpy(Code.GetData(), TestFragmentProgram, sizeof(TestFragmentProgram));
+		}
 
 		FOpenGL::bRequiresDontEmitPrecisionForTextureSamplers = false;
 		FOpenGL::bRequiresTextureCubeLodEXTToTextureCubeLodDefine = false;
 
-		FOpenGLCodeHeader Header;
+		FOpenGLCodeHeader Header{};
 		Header.FrequencyMarker = 0x5053;
 		Header.GlslMarker = 0x474c534c;
 		TArray<uint8> Code;
 		{
 			FMemoryWriter Writer(Code);
 			Writer << Header;
-			Writer.Serialize((void*)(TestFragmentProgram), strlen(TestFragmentProgram) + 1);
+			Writer << ShaderCode;
 			Writer.Close();
 		}
 		// try to compile without any hacks
