fix for HTML5 compile failure

diff --git a/Engine/Source/Runtime/OpenGLDrv/Private/OpenGLDevice.cpp b/Engine/Source/Runtime/OpenGLDrv/Private/OpenGLDevice.cpp
--- a/Engine/Source/Runtime/OpenGLDrv/Private/OpenGLDevice.cpp
+++ b/Engine/Source/Runtime/OpenGLDrv/Private/OpenGLDevice.cpp
@@ -852,9 +852,11 @@ static void InitRHICapabilitiesForGL()
 		{
 			SetupTextureFormat( PF_G16,				FOpenGLTextureFormat( GL_R16,					GL_R16,					GL_RED,			GL_UNSIGNED_SHORT,					false,	false));
 			SetupTextureFormat( PF_R32_FLOAT,		FOpenGLTextureFormat( GL_R32F,					GL_R32F,				GL_RED,			GL_FLOAT,							false,	false));
+#ifdef GL_RG_EXT
 			SetupTextureFormat( PF_G16R16F,			FOpenGLTextureFormat( GL_RG16F,					GL_RG16F,				GL_RG_EXT,		GL_HALF_FLOAT,						false,	false));
 			SetupTextureFormat( PF_G16R16F_FILTER,	FOpenGLTextureFormat( GL_RG16F,					GL_RG16F,				GL_RG_EXT,		GL_HALF_FLOAT,						false,	false));
 			SetupTextureFormat( PF_G32R32F,			FOpenGLTextureFormat( GL_RG32F,					GL_RG32F,				GL_RG_EXT,		GL_FLOAT,							false,	false));
+#endif
 #ifdef GL_UNSIGNED_INT_2_10_10_10_REV
 			SetupTextureFormat( PF_A2B10G10R10,		FOpenGLTextureFormat( GL_RGB10_A2,				GL_RGB10_A2,			GL_RGBA,		GL_UNSIGNED_INT_2_10_10_10_REV,		false,	false));
 #endif
