Copying //UE4/Ocean-Staging to //UE4/Dev-Main (Source: //Ocean/Main @ 3207682)

==========================
MAJOR FEATURES + CHANGES
==========================

Change 3202313 on 2016/11/17 by [REDACTED_PERSON]

	#ocn - Fix portrait and customizer lighting on iOS.

	- Converted the key light from directional to point.
	- Turn post processing back on.
	- Fixed an engine bug.

Change 3197000 on 2016/11/14 by [REDACTED_PERSON]

	Fixed a tool tip typo.

Change 3195359 on 2016/11/11 by [REDACTED_PERSON]

	Hotfix from CL 3192058 for UE-37121 - UMG Widget Hierarchy will now no longer autoexpand when moving or renaming widgets in the hierarchy view

Change 3187821 on 2016/11/04 by [REDACTED_PERSON]

	Fixed Parse1ToNSectionOfStrings, which wasn't correctly handling merged ini files. It now matches the start of the KeyOne and KeyN strings, so you can add numbers to the end of your keys to maintain proper order.

	Example:

	[SectionToParse]
	.Map1=MapA
	.Package1=Pak1
	.Map2=MapB
	.Package2=Pak2
	.Map3=MapC
	.Package3=Pak1
	.Package3=Pak2
	.Package3=Pak3

	GConfig->Parse1ToNSectionOfStrings(TEXT("SectionToParse"), TEXT("Map"), TEXT("Package"), OutputStrings, GGameIni);

Change 3178358 on 2016/10/28 by [REDACTED_PERSON]

	Hotfix from CL 3178267 for UE-37946 -
	Sound notify does not play in Persona if Follow is enabled

Change 3174220 on 2016/10/25 by [REDACTED_PERSON]

	Fix for TextureRenderTarget2D link errors.

Change 3174064 on 2016/10/25 by [REDACTED_PERSON]

	TextureRenderTarget2D max size.

Change 3174050 on 2016/10/25 by [REDACTED_PERSON]

	Hotfix from CL 3132846 for UE-35907: Slate Remote Enable Remove Server option in Project Settings -> Plugins causes editor to freeze and a memory leak

Change 3166783 on 2016/10/18 by [REDACTED_PERSON]

	#ocn - Hotfix from CL 3166734 for shader memory and load time reduction on mobile.  Removes the HDR_LINEAR_32 output format.  Reduce our shader memory footprint from 83 MB to 60 MB and cuts shader load time from 20s to 15s.

Change 3166751 on 2016/10/18 by [REDACTED_PERSON]

	Back out changelist 3166743

Change 3166743 on 2016/10/18 by [REDACTED_PERSON]

	Hotfix from CL 3166734 for shader memory and load time reduction on mobile.  Removes the HDR_LINEAR_32 output format.  Reduce our shader memory footprint from 83 MB to 60 MB and cuts shader load time from 20s to 15s.

Change 3163792 on 2016/10/14 by [REDACTED_PERSON]

	Hotfix from CL 3162988 for
	UE-37313
	Create Mesh Socket is missing

Change 3162582 on 2016/10/13 by [REDACTED_PERSON]

	Changed Flurry version back to 6.7.0 to fix cook errors

Change 3161957 on 2016/10/13 by [REDACTED_PERSON]

	#ocn #jira Ocn-5314 Updated our analytics solution to implement both AnalyticsET and Flurry seperately rather than use multicast

Change 3161761 on 2016/10/13 by [REDACTED_PERSON]

	Hotfix from CL 3159904 for UE-37112 - PIE loses focus while toggling Play / Pause

Change 3158496 on 2016/10/11 by [REDACTED_PERSON]

	Fix for Metal frame capture and xcode fps display on iOS.  Set GMetalCommandQueueSize to 0 for iOS.

Change 3154101 on 2016/10/06 by [REDACTED_PERSON]

	Merging //UE4/Ocean-Staging to //Ocean/Main (hotfix from 3143791 for ios compile failures on Xcode 8)

Change 3147743 on 2016/09/30 by [REDACTED_PERSON]

	#ocn - Merged in a fix from dev-platform for a crash on iOS in scene capture rendering.

	Merging

	//UE4/Dev-Platform/Engine/Source/Runtime/Renderer/Private/SceneCaptureRendering.cpp

	to //Ocean/Main/Engine/Source/Runtime/Renderer/Private/SceneCaptureRendering.cpp

Change 3147169 on 2016/09/30 by [REDACTED_PERSON]

	#ocn - Merging CL#3131515 from wex to get scroll speed control over UI widgets.

	//WEX/Main/...

	to //Ocean/Main/...

Change 3145971 on 2016/09/29 by [REDACTED_PERSON]

	Hotfix from CL 3135993 for UE-36276 - PIE does not Resume normally when Toggle Play / Pause is bound to a key

Change 3142158 on 2016/09/27 by [REDACTED_PERSON]

	Merging //UE4/Ocean-Staging to //Ocean/Main (Source: //UE4/Main @ CL 3141459)

diff --git a/Engine/Shaders/Common.usf b/Engine/Shaders/Common.usf
--- a/Engine/Shaders/Common.usf
+++ b/Engine/Shaders/Common.usf
@@ -926,14 +926,8 @@ MaterialFloat3 RGBTDecode8BPC(MaterialFloat4 RGBT, MaterialFloat Range)
 float GetHDR32bppEncodeMode()
 {
 #if ES2_PROFILE
-	#if COMPILER_GLSL_ES2
-		#if IOS || WEBGL
-			#if USE_32BPP_HDR
-				return HDR_ENCODE_MOSAIC;
-			#endif
-		#else // ANDROID
-			return intrinsic_GetHDR32bppEncodeModeES2();
-		#endif
+	#if COMPILER_GLSL_ES2 && !IOS && !WEBGL // ANDROID
+		return intrinsic_GetHDR32bppEncodeModeES2();
 	#endif
 	#if MOBILE_EMULATION
 		// To enable editor runtime change without recompile, PC always eats the encode cost.
diff --git a/Engine/Source/Editor/UMGEditor/Private/Hierarchy/SHierarchyView.cpp b/Engine/Source/Editor/UMGEditor/Private/Hierarchy/SHierarchyView.cpp
--- a/Engine/Source/Editor/UMGEditor/Private/Hierarchy/SHierarchyView.cpp
+++ b/Engine/Source/Editor/UMGEditor/Private/Hierarchy/SHierarchyView.cpp
@@ -82,6 +82,7 @@ void SHierarchyView::Construct(const FArguments& InArgs, TSharedPtr<FWidgetBluep
 	BlueprintEditor.Pin()->OnSelectedWidgetsChanged.AddRaw(this, &SHierarchyView::OnEditorSelectionChanged);
 
 	bRefreshRequested = true;
+	bExpandAllNodes = true;
 }
 
 SHierarchyView::~SHierarchyView()
@@ -105,6 +106,11 @@ void SHierarchyView::Tick( const FGeometry& AllottedGeometry, const double InCur
 {
 	if ( bRebuildTreeRequested || bRefreshRequested )
 	{
+		if (!bExpandAllNodes)
+		{
+			FindExpandedItemNames();
+		}
+
 		if ( bRebuildTreeRequested )
 		{
 			RebuildTreeView();
@@ -118,6 +124,9 @@ void SHierarchyView::Tick( const FGeometry& AllottedGeometry, const double InCur
 
 		bRefreshRequested = false;
 		bRebuildTreeRequested = false;
+		bExpandAllNodes = false;
+
+		ExpandedItemNames.Empty();
 	}
 }
 
@@ -172,6 +181,7 @@ void SHierarchyView::TransformWidgetToString(TSharedPtr<FHierarchyModel> Item, O
 void SHierarchyView::OnSearchChanged(const FText& InFilterText)
 {
 	bRefreshRequested = true;
+	bExpandAllNodes = InFilterText.IsEmpty();
 	FilterHandler->SetIsEnabled(!InFilterText.IsEmpty());
 	SearchBoxWidgetFilter->SetRawFilterText(InFilterText);
 	SearchBoxPtr->SetError(SearchBoxWidgetFilter->GetFilterErrorText());
@@ -318,22 +328,67 @@ void SHierarchyView::OnObjectsReplaced(const TMap<UObject*, UObject*>& Replaceme
 
 void SHierarchyView::RestoreExpandedItems()
 {
+	EExpandBehavior ExpandBehavior = bExpandAllNodes ? EExpandBehavior::AlwaysExpand : EExpandBehavior::RestoreFromPrevious;
+
 	for ( TSharedPtr<FHierarchyModel>& Model : RootWidgets )
 	{
-		RecursiveExpand(Model);
+		RecursiveExpand(Model, ExpandBehavior);
+	}
+}
+
+void SHierarchyView::FindExpandedItemNames()
+{
+	ExpandedItemNames.Empty();
+
+	if (WidgetTreeView.IsValid())
+	{
+		TSet< TSharedPtr<FHierarchyModel> > ExpandedItems;
+		WidgetTreeView->GetExpandedItems(ExpandedItems);
+
+		for (TSharedPtr<FHierarchyModel> Item : ExpandedItems)
+		{
+			if (Item.IsValid())
+			{
+				ExpandedItemNames.Add(Item->GetUniqueName());
+			}
+		}
 	}
 }
 
-void SHierarchyView::RecursiveExpand(TSharedPtr<FHierarchyModel>& Model, bool InShouldExpandItem)
+void SHierarchyView::RecursiveExpand(TSharedPtr<FHierarchyModel>& Model, EExpandBehavior ExpandBehavior)
 {
-	WidgetTreeView->SetItemExpansion(Model, InShouldExpandItem);
+	bool bShouldExpandItem = true;
+
+	switch (ExpandBehavior)
+	{
+		case EExpandBehavior::NeverExpand:
+		{
+			bShouldExpandItem = false;
+		}
+		break;
+
+		case EExpandBehavior::RestoreFromPrevious:
+		{
+			bShouldExpandItem = ExpandedItemNames.Contains(Model->GetUniqueName());
+		}
+		break;
+
+		case EExpandBehavior::AlwaysExpand:
+		default:
+		{
+			bShouldExpandItem = true;
+		}
+		break;
+	}
+
+	WidgetTreeView->SetItemExpansion(Model, bShouldExpandItem);
 
 	TArray< TSharedPtr<FHierarchyModel> > Children;
 	Model->GatherChildren(Children);
 
 	for (TSharedPtr<FHierarchyModel>& ChildModel : Children)
 	{
-		RecursiveExpand(ChildModel, InShouldExpandItem);
+		RecursiveExpand(ChildModel, ExpandBehavior);
 	}
 }
 
@@ -372,7 +427,7 @@ void SHierarchyView::SetItemExpansionRecursive(TSharedPtr<FHierarchyModel> Model
 {
 	if (Model.IsValid())
 	{
-		RecursiveExpand(Model, bInExpansionState);
+		RecursiveExpand(Model, bInExpansionState ? EExpandBehavior::AlwaysExpand : EExpandBehavior::NeverExpand);
 	}
 }
 
diff --git a/Engine/Source/Editor/UMGEditor/Private/Hierarchy/SHierarchyView.h b/Engine/Source/Editor/UMGEditor/Private/Hierarchy/SHierarchyView.h
--- a/Engine/Source/Editor/UMGEditor/Private/Hierarchy/SHierarchyView.h
+++ b/Engine/Source/Editor/UMGEditor/Private/Hierarchy/SHierarchyView.h
@@ -95,8 +95,15 @@ private:
 	/** Restores the state of expanded items based on the saved expanded item state, then clears the expanded state cache. */
 	void RestoreExpandedItems();
 
+	enum class EExpandBehavior : uint8
+	{
+		NeverExpand,
+		AlwaysExpand,
+		RestoreFromPrevious,
+	};
+
 	/** Recursively expands the models based on the expansion set. */
-	void RecursiveExpand(TSharedPtr<FHierarchyModel>& Model, bool InShouldExpandItem = true);
+	void RecursiveExpand(TSharedPtr<FHierarchyModel>& Model, EExpandBehavior ExpandBehavior = EExpandBehavior::AlwaysExpand);
 
 	/**  */
 	void RestoreSelectedItems();
@@ -107,6 +114,9 @@ private:
 	/** Handler for recursively expanding/collapsing items */
 	void SetItemExpansionRecursive(TSharedPtr<FHierarchyModel> Model, bool bInExpansionState);
 
+	/** Find and store the names of all currently expanded nodes in the hierarchy view. Should only be called when rebuilding the tree */
+	void FindExpandedItemNames();
+
 private:
 
 	/** Cached pointer to the blueprint editor that owns this tree. */
@@ -130,6 +140,9 @@ private:
 	/** The widget hierarchy slate treeview widget */
 	TSharedPtr< STreeView< TSharedPtr<FHierarchyModel> > > WidgetTreeView;
 
+	/** The unique names of all nodes expanded in the tree view */
+	TSet< FName > ExpandedItemNames;
+
 	/** The search box used to update the filter text */
 	TSharedPtr<class SSearchBox> SearchBoxPtr;
 
@@ -144,4 +157,7 @@ private:
 
 	/** Flag to ignore selections while the hierarchy view is updating the selection. */
 	bool bIsUpdatingSelection;
+
+	/** Should all nodes in the tree be expanded? */
+	bool bExpandAllNodes;
 };
diff --git a/Engine/Source/Runtime/Apple/MetalRHI/Private/MetalContext.cpp b/Engine/Source/Runtime/Apple/MetalRHI/Private/MetalContext.cpp
--- a/Engine/Source/Runtime/Apple/MetalRHI/Private/MetalContext.cpp
+++ b/Engine/Source/Runtime/Apple/MetalRHI/Private/MetalContext.cpp
@@ -32,7 +32,7 @@ static FAutoConsoleVariableRef CVarMetalSupportsIntermediateBackBuffer(
 static int32 GMetalCommandQueueSize = 5120; // This number is large due to texture streaming - currently each texture is its own command-buffer.
 // The whole MetalRHI needs to be changed to use MTLHeaps/MTLFences & reworked so that operations with the same synchronisation requirements are collapsed into a single blit command-encoder/buffer.
 #else
-static int32 GMetalCommandQueueSize = 64;
+static int32 GMetalCommandQueueSize = 0;
 #endif
 static FAutoConsoleVariableRef CVarMetalCommandQueueSize(
 	TEXT("rhi.Metal.CommandQueueSize"),
diff --git a/Engine/Source/Runtime/Core/Private/Misc/ConfigCacheIni.cpp b/Engine/Source/Runtime/Core/Private/Misc/ConfigCacheIni.cpp
--- a/Engine/Source/Runtime/Core/Private/Misc/ConfigCacheIni.cpp
+++ b/Engine/Source/Runtime/Core/Private/Misc/ConfigCacheIni.cpp
@@ -1607,18 +1607,18 @@ bool FConfigCacheIni::AreFileOperationsDisabled()
 }
 
 /**
- * Prases apart an ini section that contains a list of 1-to-N mappings of names in the following format
+ * Parses apart an ini section that contains a list of 1-to-N mappings of names in the following format
  *	 [PerMapPackages]
- *	 MapName=Map1
- *	 Package=PackageA
- *	 Package=PackageB
- *	 MapName=Map2
- *	 Package=PackageC
- *	 Package=PackageD
+ *	 .MapName1=Map1
+ *	 .Package1=PackageA
+ *	 .Package1=PackageB
+ *	 .MapName2=Map2
+ *	 .Package2=PackageC
+ *	 .Package2=PackageD
  * 
  * @param Section Name of section to look in
- * @param KeyOne Key to use for the 1 in the 1-to-N (MapName in the above example)
- * @param KeyN Key to use for the N in the 1-to-N (Package in the above example)
+ * @param KeyOne Key to use for the 1 in the 1-to-N (MapName in the above example - the number suffix gets ignored but helps to keep ordering)
+ * @param KeyN Key to use for the N in the 1-to-N (Package in the above example - the number suffix gets ignored but helps to keep ordering)
  * @param OutMap Map containing parsed results
  * @param Filename Filename to use to find the section
  *
@@ -1644,7 +1644,7 @@ void FConfigCacheIni::Parse1ToNSectionOfNames(const TCHAR* Section, const TCHAR*
 	for( FConfigSectionMap::TIterator It(*ConfigSection); It; ++It )
 	{
 		// is the current key the 1 key?
-		if (It.Key() == KeyOne)
+		if (It.Key().ToString().StartsWith(KeyOne))
 		{
 			const FName KeyName(*It.Value().GetValue());
 
@@ -1658,7 +1658,7 @@ void FConfigCacheIni::Parse1ToNSectionOfNames(const TCHAR* Section, const TCHAR*
 			}
 		}
 		// is the current key the N key?
-		else if (It.Key() == KeyN && WorkingList != nullptr)
+		else if (It.Key().ToString().StartsWith(KeyN) && WorkingList != nullptr)
 		{
 			// if so, add it to the N list for the current 1 key
 			WorkingList->Add(FName(*It.Value().GetValue()));
@@ -1672,18 +1672,18 @@ void FConfigCacheIni::Parse1ToNSectionOfNames(const TCHAR* Section, const TCHAR*
 }
 
 /**
- * Prases apart an ini section that contains a list of 1-to-N mappings of strings in the following format
+ * Parses apart an ini section that contains a list of 1-to-N mappings of strings in the following format
  *	 [PerMapPackages]
- *	 MapName=Map1
- *	 Package=PackageA
- *	 Package=PackageB
- *	 MapName=Map2
- *	 Package=PackageC
- *	 Package=PackageD
+ *	 .MapName1=Map1
+ *	 .Package1=PackageA
+ *	 .Package1=PackageB
+ *	 .MapName2=Map2
+ *	 .Package2=PackageC
+ *	 .Package2=PackageD
  * 
  * @param Section Name of section to look in
- * @param KeyOne Key to use for the 1 in the 1-to-N (MapName in the above example)
- * @param KeyN Key to use for the N in the 1-to-N (Package in the above example)
+ * @param KeyOne Key to use for the 1 in the 1-to-N (MapName in the above example - the number suffix gets ignored but helps to keep ordering)
+ * @param KeyN Key to use for the N in the 1-to-N (Package in the above example - the number suffix gets ignored but helps to keep ordering)
  * @param OutMap Map containing parsed results
  * @param Filename Filename to use to find the section
  *
@@ -1709,7 +1709,7 @@ void FConfigCacheIni::Parse1ToNSectionOfStrings(const TCHAR* Section, const TCHA
 	for( FConfigSectionMap::TIterator It(*ConfigSection); It; ++It )
 	{
 		// is the current key the 1 key?
-		if (It.Key() == KeyOne)
+		if (It.Key().ToString().StartsWith(KeyOne))
 		{
 			// look for existing set in the map
 			WorkingList = OutMap.Find(It.Value().GetValue());
@@ -1721,7 +1721,7 @@ void FConfigCacheIni::Parse1ToNSectionOfStrings(const TCHAR* Section, const TCHA
 			}
 		}
 		// is the current key the N key?
-		else if (It.Key() == KeyN && WorkingList != nullptr)
+		else if (It.Key().ToString().StartsWith(KeyN) && WorkingList != nullptr)
 		{
 			// if so, add it to the N list for the current 1 key
 			WorkingList->Add(It.Value().GetValue());
diff --git a/Engine/Source/Runtime/Engine/Classes/Engine/TextureRenderTarget2D.h b/Engine/Source/Runtime/Engine/Classes/Engine/TextureRenderTarget2D.h
--- a/Engine/Source/Runtime/Engine/Classes/Engine/TextureRenderTarget2D.h
+++ b/Engine/Source/Runtime/Engine/Classes/Engine/TextureRenderTarget2D.h
@@ -12,6 +12,9 @@ class FTextureResource;
 class UTexture2D;
 struct FPropertyChangedEvent;
 
+extern ENGINE_API int32 GTextureRenderTarget2DMaxSizeX;
+extern ENGINE_API int32 GTextureRenderTarget2DMaxSizeY;
+
 /**
  * TextureRenderTarget2D
  *
diff --git a/Engine/Source/Runtime/Engine/Private/TextureRenderTarget2D.cpp b/Engine/Source/Runtime/Engine/Private/TextureRenderTarget2D.cpp
--- a/Engine/Source/Runtime/Engine/Private/TextureRenderTarget2D.cpp
+++ b/Engine/Source/Runtime/Engine/Private/TextureRenderTarget2D.cpp
@@ -12,6 +12,9 @@
 #include "DeviceProfiles/DeviceProfile.h"
 #include "DeviceProfiles/DeviceProfileManager.h"
 
+int32 GTextureRenderTarget2DMaxSizeX = 999999999;
+int32 GTextureRenderTarget2DMaxSizeY = 999999999;
+
 /*-----------------------------------------------------------------------------
 	UTextureRenderTarget2D
 -----------------------------------------------------------------------------*/
@@ -161,6 +164,9 @@ void UTextureRenderTarget2D::PostLoad()
 		SizeY = FMath::Min<int32>(SizeY, GSystemResolution.ResY);
 	}
 
+	SizeX = FMath::Min<int32>(SizeX, GTextureRenderTarget2DMaxSizeX);
+	SizeY = FMath::Min<int32>(SizeY, GTextureRenderTarget2DMaxSizeY);
+	
 	Super::PostLoad();
 }
 
diff --git a/Engine/Source/Runtime/OpenGLDrv/Private/OpenGLDevice.cpp b/Engine/Source/Runtime/OpenGLDrv/Private/OpenGLDevice.cpp
--- a/Engine/Source/Runtime/OpenGLDrv/Private/OpenGLDevice.cpp
+++ b/Engine/Source/Runtime/OpenGLDrv/Private/OpenGLDevice.cpp
@@ -19,6 +19,7 @@
 #include "ShaderCache.h"
 #include "OpenGLDrv.h"
 #include "OpenGLDrvPrivate.h"
+#include "SceneUtils.h"
 
 #include "HardwareInfo.h"
 
@@ -200,11 +201,11 @@ void FOpenGLDynamicRHI::RHIEndScene()
 
 bool GDisableOpenGLDebugOutput = false;
 
-// workaround for HTML5. 
+// workaround for HTML5.
 #if PLATFORM_HTML5
 #undef GL_ARB_debug_output
 #undef GL_KHR_debug
-#endif 
+#endif
 
 #if defined(GL_ARB_debug_output) || defined(GL_KHR_debug)
 /**
@@ -248,7 +249,7 @@ static const TCHAR* GetOpenGLDebugTypeStringARB(GLenum Type)
 	{
 		return TypeStrings[Type - GL_DEBUG_TYPE_ERROR_ARB];
 	}
-#ifdef GL_KHR_debug	
+#ifdef GL_KHR_debug
 	{
 		static const TCHAR* DebugTypeStrings[] =
 		{
@@ -405,7 +406,7 @@ static void APIENTRY OpenGLDebugMessageCallbackAMD(
 #if !NO_LOGGING
 	const TCHAR* CategoryStr = GetOpenGLDebugCategoryStringAMD(Category);
 	const TCHAR* SeverityStr = GetOpenGLDebugSeverityStringAMD(Severity);
-	
+
 	ELogVerbosity::Type Verbosity = ELogVerbosity::Warning;
 	if (Severity == GL_DEBUG_SEVERITY_HIGH_AMD)
 	{
@@ -442,7 +443,7 @@ static inline void SetupTextureFormat( EPixelFormat Format, const FOpenGLTexture
 }
 
 
-void InitDebugContext() 
+void InitDebugContext()
 {
 	// Set the debug output callback if the driver supports it.
 	VERIFY_GL(__FUNCTION__);
@@ -504,14 +505,14 @@ void InitDebugContext()
 }
 
 TAutoConsoleVariable<FString> CVarOpenGLStripExtensions(
-	TEXT("r.OpenGL.StripExtensions"), 
-	TEXT(""), 
+	TEXT("r.OpenGL.StripExtensions"),
+	TEXT(""),
 	TEXT("List of comma separated OpenGL extensions to strip from a driver reported extensions string"),
 	ECVF_ReadOnly);
 
 TAutoConsoleVariable<FString> CVarOpenGLAddExtensions(
-	TEXT("r.OpenGL.AddExtensions"), 
-	TEXT(""), 
+	TEXT("r.OpenGL.AddExtensions"),
+	TEXT(""),
 	TEXT("List of comma separated OpenGL extensions to add to a driver reported extensions string"),
 	ECVF_ReadOnly);
 
@@ -544,7 +545,7 @@ void ApplyExtensionsOverrides(FString& ExtensionsString)
 			ExtName = ExtName.Trim().TrimTrailing();
 			if (!ExtensionsString.Contains(ExtName))
 			{
-				ExtensionsString.Append(TEXT(" ")); // extensions delimiter 
+				ExtensionsString.Append(TEXT(" ")); // extensions delimiter
 				ExtensionsString.Append(ExtName);
 				UE_LOG(LogRHI, Log, TEXT("Added extension: %s"), *ExtName);
 			}
@@ -563,7 +564,7 @@ static void InitRHICapabilitiesForGL()
 	GTexturePoolSize = 0;
 	GPoolSizeVRAMPercentage = 0;
 #if PLATFORM_WINDOWS || PLATFORM_LINUX
-	GConfig->GetInt( TEXT( "TextureStreaming" ), TEXT( "PoolSizeVRAMPercentage" ), GPoolSizeVRAMPercentage, GEngineIni );	
+	GConfig->GetInt( TEXT( "TextureStreaming" ), TEXT( "PoolSizeVRAMPercentage" ), GPoolSizeVRAMPercentage, GEngineIni );
 #endif
 
 	// GL vendor and version information.
@@ -651,7 +652,7 @@ static void InitRHICapabilitiesForGL()
 #endif
 #if GL_MAX_3D_TEXTURE_SIZE
 	LOG_AND_GET_GL_INT_TEMP(GL_MAX_3D_TEXTURE_SIZE, 0);
-#endif 
+#endif
 	LOG_AND_GET_GL_INT_TEMP(GL_MAX_RENDERBUFFER_SIZE, 0);
 	LOG_AND_GET_GL_INT_TEMP(GL_MAX_TEXTURE_IMAGE_UNITS, 0);
 	if (FOpenGL::SupportsDrawBuffers())
@@ -771,23 +772,25 @@ static void InitRHICapabilitiesForGL()
 	GSupportsDepthBoundsTest = FOpenGL::SupportsDepthBoundsTest();
 
 	GSupportsRenderTargetFormat_PF_FloatRGBA = FOpenGL::SupportsColorBufferHalfFloat();
-	
+
 	GSupportsMultipleRenderTargets = FOpenGL::SupportsMultipleRenderTargets();
 	GSupportsWideMRT = FOpenGL::SupportsWideMRT();
 	GSupportsTexture3D = FOpenGL::SupportsTexture3D();
 	GSupportsMobileMultiView = FOpenGL::SupportsMobileMultiView();
 	GSupportsResourceView = FOpenGL::SupportsResourceView();
-		
+
 	GSupportsShaderFramebufferFetch = FOpenGL::SupportsShaderFramebufferFetch();
 	GSupportsShaderDepthStencilFetch = FOpenGL::SupportsShaderDepthStencilFetch();
 	GMaxShadowDepthBufferSizeX = FMath::Min<int32>(Value_GL_MAX_RENDERBUFFER_SIZE, 4096); // Limit to the D3D11 max.
 	GMaxShadowDepthBufferSizeY = FMath::Min<int32>(Value_GL_MAX_RENDERBUFFER_SIZE, 4096);
 	GHardwareHiddenSurfaceRemoval = FOpenGL::HasHardwareHiddenSurfaceRemoval();
-	GRHISupportsInstancing = FOpenGL::SupportsInstancing(); // HTML5 does not support it. Android supports it with OpenGL ES3.0+ 
+	GRHISupportsInstancing = FOpenGL::SupportsInstancing(); // HTML5 does not support it. Android supports it with OpenGL ES3.0+
 	GSupportsTimestampRenderQueries = FOpenGL::SupportsTimestampQueries();
 
 	GSupportsHDR32bppEncodeModeIntrinsic = FOpenGL::SupportsHDR32bppEncodeModeIntrinsic();
 
+	checkf(!IsMobileHDR32bpp() || GSupportsHDR32bppEncodeModeIntrinsic || IsPCPlatform(GMaxRHIShaderPlatform), TEXT("Current platform does not support 32bpp HDR but IsMobileHDR32bpp() returned true"));
+
 	GShaderPlatformForFeatureLevel[ERHIFeatureLevel::ES2] = (GMaxRHIFeatureLevel == ERHIFeatureLevel::ES2) ? GMaxRHIShaderPlatform : SP_OPENGL_PCES2;
 	GShaderPlatformForFeatureLevel[ERHIFeatureLevel::ES3_1] = (GMaxRHIFeatureLevel == ERHIFeatureLevel::ES3_1) ? GMaxRHIShaderPlatform : SP_OPENGL_PCES3_1;
 	GShaderPlatformForFeatureLevel[ERHIFeatureLevel::SM4] = PLATFORM_MAC ? SP_OPENGL_SM4_MAC : SP_OPENGL_SM4;
@@ -1033,7 +1036,7 @@ FOpenGLDynamicRHI::FOpenGLDynamicRHI()
 ,	PlatformDevice(NULL)
 ,	GPUProfilingData(this)
 {
-	// This should be called once at the start 
+	// This should be called once at the start
 	check( IsInGameThread() );
 	check( !GIsThreadedRendering );
 
@@ -1116,10 +1119,10 @@ static bool VerifyCompiledShader(GLuint Shader, const ANSICHAR* GlslCode, bool I
 				Temp += Converted.Len();
 			}
 
-		}	
+		}
 #endif
 		UE_LOG(LogRHI,Warning,TEXT("Failed to compile shader. Compile log:\n%s"), ANSI_TO_TCHAR(CompileLog));
-		
+
 		if (LogLength > 1)
 		{
 			FMemory::Free(CompileLog);
@@ -1321,12 +1324,12 @@ static void CheckTextureCubeLodSupport()
 			FOpenGL::bIsCheckingShaderCompilerHacks = false;
 			return;
 		}
-		
+
 		FOpenGL::bRequiresDontEmitPrecisionForTextureSamplers = true;
 		FOpenGL::bRequiresTextureCubeLodEXTToTextureCubeLodDefine = false;
 
 		// second most number of devices fall into this hack category
-		// try to compile without using precision for texture samplers 
+		// try to compile without using precision for texture samplers
 		// Samsung Galaxy Express	Samsung Galaxy S3	Samsung Galaxy S3 mini	Samsung Galaxy Tab GT-P1000	Samsung Galaxy Tab 2
 		PixelShader = (FOpenGLPixelShader*)(RHICreatePixelShader(Code).GetReference());
 
@@ -1377,7 +1380,7 @@ void FOpenGLDynamicRHI::Init()
 {
 	check(!GIsRHIInitialized);
 	VERIFY_GL_SCOPE();
-	
+
 	FOpenGLProgramBinaryCache::Initialize();
 #if PLATFORM_DESKTOP
 	FShaderCache::InitShaderCache(SCO_Default, FOpenGL::GetMaxTextureImageUnits());
@@ -1441,7 +1444,7 @@ void FOpenGLDynamicRHI::Init()
 
 	CheckTextureCubeLodSupport();
 	CheckVaryingLimit();
-	
+
 #if PLATFORM_DESKTOP
 	FShaderCache::LoadBinaryCache();
 #endif
@@ -1590,7 +1593,7 @@ void FOpenGLDynamicRHI::InvalidateQueries( void )
 
 	{
 		FScopeLock Lock(&TimerQueriesListCriticalSection);
-		
+
 		for( int32 Index = 0; Index < TimerQueries.Num(); ++Index )
 		{
 			TimerQueries[Index]->bInvalidResource = true;
diff --git a/Engine/Source/Runtime/Renderer/Private/MobileBasePassRendering.cpp b/Engine/Source/Runtime/Renderer/Private/MobileBasePassRendering.cpp
--- a/Engine/Source/Runtime/Renderer/Private/MobileBasePassRendering.cpp
+++ b/Engine/Source/Runtime/Renderer/Private/MobileBasePassRendering.cpp
@@ -26,16 +26,12 @@ static TAutoConsoleVariable<int32> CVarMobileDisableVertexFog(
 
 #define IMPLEMENT_MOBILE_SHADING_BASEPASS_LIGHTMAPPED_PIXEL_SHADER_TYPE(LightMapPolicyType,LightMapPolicyName,NumDynamicPointLights) \
 	typedef TMobileBasePassPS< LightMapPolicyType, LDR_GAMMA_32, false, NumDynamicPointLights > TMobileBasePassPS##LightMapPolicyName##NumDynamicPointLights##LDRGamma32; \
-	typedef TMobileBasePassPS< LightMapPolicyType, HDR_LINEAR_32, false, NumDynamicPointLights > TMobileBasePassPS##LightMapPolicyName##NumDynamicPointLights##HDRLinear32; \
 	typedef TMobileBasePassPS< LightMapPolicyType, HDR_LINEAR_64, false, NumDynamicPointLights > TMobileBasePassPS##LightMapPolicyName##NumDynamicPointLights##HDRLinear64; \
 	typedef TMobileBasePassPS< LightMapPolicyType, LDR_GAMMA_32, true, NumDynamicPointLights > TMobileBasePassPS##LightMapPolicyName##NumDynamicPointLights##LDRGamma32##Skylight; \
-	typedef TMobileBasePassPS< LightMapPolicyType, HDR_LINEAR_32, true, NumDynamicPointLights > TMobileBasePassPS##LightMapPolicyName##NumDynamicPointLights##HDRLinear32##Skylight; \
 	typedef TMobileBasePassPS< LightMapPolicyType, HDR_LINEAR_64, true, NumDynamicPointLights > TMobileBasePassPS##LightMapPolicyName##NumDynamicPointLights##HDRLinear64##Skylight; \
 	IMPLEMENT_MATERIAL_SHADER_TYPE(template<>, TMobileBasePassPS##LightMapPolicyName##NumDynamicPointLights##LDRGamma32, TEXT("MobileBasePassPixelShader"), TEXT("Main"), SF_Pixel); \
-	IMPLEMENT_MATERIAL_SHADER_TYPE(template<>, TMobileBasePassPS##LightMapPolicyName##NumDynamicPointLights##HDRLinear32, TEXT("MobileBasePassPixelShader"), TEXT("Main"), SF_Pixel); \
 	IMPLEMENT_MATERIAL_SHADER_TYPE(template<>, TMobileBasePassPS##LightMapPolicyName##NumDynamicPointLights##HDRLinear64, TEXT("MobileBasePassPixelShader"), TEXT("Main"), SF_Pixel); \
 	IMPLEMENT_MATERIAL_SHADER_TYPE(template<>, TMobileBasePassPS##LightMapPolicyName##NumDynamicPointLights##LDRGamma32##Skylight, TEXT("MobileBasePassPixelShader"), TEXT("Main"), SF_Pixel); \
-	IMPLEMENT_MATERIAL_SHADER_TYPE(template<>, TMobileBasePassPS##LightMapPolicyName##NumDynamicPointLights##HDRLinear32##Skylight, TEXT("MobileBasePassPixelShader"), TEXT("Main"), SF_Pixel); \
 	IMPLEMENT_MATERIAL_SHADER_TYPE(template<>, TMobileBasePassPS##LightMapPolicyName##NumDynamicPointLights##HDRLinear64##Skylight, TEXT("MobileBasePassPixelShader"), TEXT("Main"), SF_Pixel);
 
 static_assert(MAX_BASEPASS_DYNAMIC_POINT_LIGHTS == 4, "If you change MAX_BASEPASS_DYNAMIC_POINT_LIGHTS, you need to add shader types below");
diff --git a/Engine/Source/Runtime/Renderer/Private/MobileBasePassRendering.h b/Engine/Source/Runtime/Renderer/Private/MobileBasePassRendering.h
--- a/Engine/Source/Runtime/Renderer/Private/MobileBasePassRendering.h
+++ b/Engine/Source/Runtime/Renderer/Private/MobileBasePassRendering.h
@@ -32,7 +32,6 @@ class FScene;
 enum EOutputFormat
 {
 	LDR_GAMMA_32,
-	HDR_LINEAR_32,
 	HDR_LINEAR_64,
 };
 
@@ -55,13 +54,10 @@ static bool ShouldCacheShaderByPlatformAndOutputFormat(EShaderPlatform Platform,
 {
 	bool bSupportsMobileHDR = IsMobileHDR();
 	bool bShaderUsesLDR = (OutputFormat == LDR_GAMMA_32);
-	bool bShaderUsesHDR = !bShaderUsesLDR;
-	// Android ES2 uses intrinsic_GetHDR32bppEncodeModeES2 so doesn't need a HDR_LINEAR_32 permutation
-	bool bIsAndroid32bpp = (OutputFormat == HDR_LINEAR_32) && (Platform == SP_OPENGL_ES2_ANDROID || Platform == SP_OPENGL_ES3_1_ANDROID);
 
 	// only cache this shader if the LDR/HDR output matches what we currently support.  IsMobileHDR can't change, so we don't need
-	// the LDR shaders if we are doing HDR, and vice-versa.	Android doesn't need HDR_LINEAR_32 as it
-	return (bShaderUsesLDR && !bSupportsMobileHDR) || (bShaderUsesHDR && bSupportsMobileHDR && !bIsAndroid32bpp);
+	// the LDR shaders if we are doing HDR, and vice-versa.
+	return (bShaderUsesLDR && !bSupportsMobileHDR) || (!bShaderUsesLDR && bSupportsMobileHDR);
 }
 
 /**
@@ -441,7 +437,6 @@ public:
 	{		
 		TMobileBasePassPSBaseType<LightMapPolicyType, NumDynamicPointLights>::ModifyCompilationEnvironment(Platform, Material, OutEnvironment);
 		OutEnvironment.SetDefine(TEXT("ENABLE_SKY_LIGHT"), bEnableSkyLight);
-		OutEnvironment.SetDefine(TEXT("USE_32BPP_HDR"), OutputFormat == HDR_LINEAR_32);
 		OutEnvironment.SetDefine(TEXT("OUTPUT_GAMMA_SPACE"), OutputFormat == LDR_GAMMA_32);
 	}
 	
@@ -470,20 +465,7 @@ struct GetMobileBasePassShaders
 	TMobileBasePassPSPolicyParamType<typename LightMapPolicyType::PixelParametersType, NumDynamicPointLights>*& PixelShader
 	)
 	{
-		if (IsMobileHDR32bpp() && !GSupportsHDR32bppEncodeModeIntrinsic)
-		{
-			VertexShader = Material.GetShader<TMobileBasePassVS<LightMapPolicyType, HDR_LINEAR_64> >(VertexFactoryType);
-
-			if (bEnableSkyLight)
-			{
-				PixelShader = Material.GetShader< TMobileBasePassPS<LightMapPolicyType, HDR_LINEAR_32, true, NumDynamicPointLights> >(VertexFactoryType);
-			}
-			else
-			{
-				PixelShader = Material.GetShader< TMobileBasePassPS<LightMapPolicyType, HDR_LINEAR_32, false, NumDynamicPointLights> >(VertexFactoryType);
-			}
-		}
-		else if (IsMobileHDR())
+		if (IsMobileHDR())
 		{
 			VertexShader = Material.GetShader<TMobileBasePassVS<LightMapPolicyType, HDR_LINEAR_64> >(VertexFactoryType);
 
@@ -535,20 +517,7 @@ void GetUniformMobileBasePassShaders(
 	TMobileBasePassPSPolicyParamType<FUniformLightMapPolicyShaderParametersType, NumDynamicPointLights>*& PixelShader
 	)
 {
-	if (IsMobileHDR32bpp() && !GSupportsHDR32bppEncodeModeIntrinsic)
-	{
-		VertexShader = Material.GetShader<TMobileBasePassVS<TUniformLightMapPolicy<Policy>, HDR_LINEAR_64> >(VertexFactoryType);
-
-		if (bEnableSkyLight)
-		{
-			PixelShader = Material.GetShader< TMobileBasePassPS<TUniformLightMapPolicy<Policy>, HDR_LINEAR_32, true, NumDynamicPointLights> >(VertexFactoryType);
-		}
-		else
-		{
-			PixelShader = Material.GetShader< TMobileBasePassPS<TUniformLightMapPolicy<Policy>, HDR_LINEAR_32, false, NumDynamicPointLights> >(VertexFactoryType);
-		}
-	}
-	else if (IsMobileHDR())
+	if (IsMobileHDR())
 	{
 		VertexShader = Material.GetShader<TMobileBasePassVS<TUniformLightMapPolicy<Policy>, HDR_LINEAR_64> >(VertexFactoryType);
 
diff --git a/Engine/Source/Runtime/Renderer/Private/PostProcess/PostProcessTonemap.cpp b/Engine/Source/Runtime/Renderer/Private/PostProcess/PostProcessTonemap.cpp
--- a/Engine/Source/Runtime/Renderer/Private/PostProcess/PostProcessTonemap.cpp
+++ b/Engine/Source/Runtime/Renderer/Private/PostProcess/PostProcessTonemap.cpp
@@ -1426,16 +1426,18 @@ class FPostProcessTonemapPS_ES2 : public FGlobalShader
 
 	static bool ShouldCache(EShaderPlatform Platform)
 	{
-		// This is only used on ES2.
-		// TODO: Make this only compile on PC/Mobile (and not console).
-		return !IsConsolePlatform(Platform);
+		const uint32 ConfigBitmask = TonemapperConfBitmaskMobile[ConfigIndex];
+
+		// Only cache for ES2/3.1 shader platforms, and only compile 32bpp shaders for Android or PC emulation
+		return IsMobilePlatform(Platform) && 
+			(!TonemapperIsDefined(ConfigBitmask, Tonemapper32BPPHDR) || Platform == SP_OPENGL_ES2_ANDROID || (IsES2Platform(Platform) && IsPCPlatform(Platform)));
 	}
 
 	static void ModifyCompilationEnvironment(EShaderPlatform Platform, FShaderCompilerEnvironment& OutEnvironment)
 	{
 		FGlobalShader::ModifyCompilationEnvironment(Platform, OutEnvironment);
 
-		uint32 ConfigBitmask = TonemapperConfBitmaskMobile[ConfigIndex];
+		const uint32 ConfigBitmask = TonemapperConfBitmaskMobile[ConfigIndex];
 
 		OutEnvironment.SetDefine(TEXT("USE_GAMMA_ONLY"),         TonemapperIsDefined(ConfigBitmask, TonemapperGammaOnly));
 		OutEnvironment.SetDefine(TEXT("USE_COLOR_MATRIX"),       TonemapperIsDefined(ConfigBitmask, TonemapperColorMatrix));
diff --git a/Engine/Source/Runtime/Slate/Private/Widgets/Layout/SScrollBox.cpp b/Engine/Source/Runtime/Slate/Private/Widgets/Layout/SScrollBox.cpp
--- a/Engine/Source/Runtime/Slate/Private/Widgets/Layout/SScrollBox.cpp
+++ b/Engine/Source/Runtime/Slate/Private/Widgets/Layout/SScrollBox.cpp
@@ -762,7 +762,7 @@ FReply SScrollBox::OnMouseWheel( const FGeometry& MyGeometry, const FPointerEven
 		// Make sure scroll velocity is cleared so it doesn't fight with the mouse wheel input
 		InertialScrollManager.ClearScrollVelocity();
 
-		const bool bScrollWasHandled = this->ScrollBy(MyGeometry, -MouseEvent.GetWheelDelta()*WheelScrollAmount, EAllowOverscroll::No);
+		const bool bScrollWasHandled = this->ScrollBy(MyGeometry, -MouseEvent.GetWheelDelta() * GetGlobalScrollAmount(), EAllowOverscroll::No);
 
 		if ( bScrollWasHandled && !bIsScrollingActiveTimerRegistered )
 		{
diff --git a/Engine/Source/Runtime/Slate/Private/Widgets/Text/SMultiLineEditableText.cpp b/Engine/Source/Runtime/Slate/Private/Widgets/Text/SMultiLineEditableText.cpp
--- a/Engine/Source/Runtime/Slate/Private/Widgets/Text/SMultiLineEditableText.cpp
+++ b/Engine/Source/Runtime/Slate/Private/Widgets/Text/SMultiLineEditableText.cpp
@@ -580,7 +580,7 @@ FReply SMultiLineEditableText::OnMouseWheel( const FGeometry& MyGeometry, const
 {
 	if (VScrollBar.IsValid() && VScrollBar->IsNeeded())
 	{
-		const float ScrollAmount = -MouseEvent.GetWheelDelta() * WheelScrollAmount;
+		const float ScrollAmount = -MouseEvent.GetWheelDelta() * GetGlobalScrollAmount();
 
 		const FVector2D PreviousScrollOffset = EditableTextLayout->GetScrollOffset();
 		
diff --git a/Engine/Source/Runtime/Slate/Private/Widgets/Views/STableViewBase.cpp b/Engine/Source/Runtime/Slate/Private/Widgets/Views/STableViewBase.cpp
--- a/Engine/Source/Runtime/Slate/Private/Widgets/Views/STableViewBase.cpp
+++ b/Engine/Source/Runtime/Slate/Private/Widgets/Views/STableViewBase.cpp
@@ -658,7 +658,7 @@ STableViewBase::STableViewBase( ETableViewMode::Type InTableViewMode )
 	, SelectionMode( ESelectionMode::Multi )
 	, SoftwareCursorPosition( ForceInitToZero )
 	, bShowSoftwareCursor( false )
-	, WheelScrollMultiplier( WheelScrollAmount )
+	, WheelScrollMultiplier(GetGlobalScrollAmount())
 	, bIsScrollingActiveTimerRegistered( false )
 	, Overscroll()
 	, AllowOverscroll(EAllowOverscroll::Yes)
diff --git a/Engine/Source/Runtime/Slate/Public/Widgets/Views/SListView.h b/Engine/Source/Runtime/Slate/Public/Widgets/Views/SListView.h
--- a/Engine/Source/Runtime/Slate/Public/Widgets/Views/SListView.h
+++ b/Engine/Source/Runtime/Slate/Public/Widgets/Views/SListView.h
@@ -74,7 +74,7 @@ public:
 		, _ExternalScrollbar()
 		, _AllowOverscroll(EAllowOverscroll::Yes)
 		, _ConsumeMouseWheel( EConsumeMouseWheel::WhenScrollingPossible )
-		, _WheelScrollMultiplier( WheelScrollAmount )
+		, _WheelScrollMultiplier(GetGlobalScrollAmount())
 		, _HandleGamepadEvents( true )
 		{ }
 
diff --git a/Engine/Source/Runtime/Slate/Public/Widgets/Views/STileView.h b/Engine/Source/Runtime/Slate/Public/Widgets/Views/STileView.h
--- a/Engine/Source/Runtime/Slate/Public/Widgets/Views/STileView.h
+++ b/Engine/Source/Runtime/Slate/Public/Widgets/Views/STileView.h
@@ -52,7 +52,7 @@ public:
 		, _ScrollbarVisibility(EVisibility::Visible)
 		, _AllowOverscroll(EAllowOverscroll::Yes)
 		, _ConsumeMouseWheel( EConsumeMouseWheel::WhenScrollingPossible )
-		, _WheelScrollMultiplier( WheelScrollAmount )
+		, _WheelScrollMultiplier(GetGlobalScrollAmount())
 		, _HandleGamepadEvents( true )
 		{}
 
diff --git a/Engine/Source/Runtime/Slate/Public/Widgets/Views/STreeView.h b/Engine/Source/Runtime/Slate/Public/Widgets/Views/STreeView.h
--- a/Engine/Source/Runtime/Slate/Public/Widgets/Views/STreeView.h
+++ b/Engine/Source/Runtime/Slate/Public/Widgets/Views/STreeView.h
@@ -103,7 +103,7 @@ public:
 		, _ExternalScrollbar()
 		, _ConsumeMouseWheel( EConsumeMouseWheel::WhenScrollingPossible )
 		, _AllowOverscroll(EAllowOverscroll::Yes)
-		, _WheelScrollMultiplier( WheelScrollAmount )
+		, _WheelScrollMultiplier(GetGlobalScrollAmount())
 		{}
 
 		SLATE_EVENT( FOnGenerateRow, OnGenerateRow )
diff --git a/Engine/Source/Runtime/SlateCore/Private/SlateCoreClasses.cpp b/Engine/Source/Runtime/SlateCore/Private/SlateCoreClasses.cpp
--- a/Engine/Source/Runtime/SlateCore/Private/SlateCoreClasses.cpp
+++ b/Engine/Source/Runtime/SlateCore/Private/SlateCoreClasses.cpp
@@ -1,8 +1,14 @@
 // Copyright 1998-2016 Epic Games, Inc. All Rights Reserved.
 
 #include "CoreMinimal.h"
+#include "Types/SlateConstants.h"
 #include "Styling/SlateWidgetStyle.h"
 
+/** How much to scroll for each click of the mouse wheel (in Slate Screen Units). */
+TAutoConsoleVariable<float> GlobalScrollAmount(
+	TEXT("Slate.GlobalScrollAmount"),
+	32.0f,
+	TEXT("How much to scroll for each click of the mouse wheel (in Slate Screen Units)."));
 
 FSlateWidgetStyle::FSlateWidgetStyle( )
 { }
diff --git a/Engine/Source/Runtime/SlateCore/Public/Types/SlateConstants.h b/Engine/Source/Runtime/SlateCore/Public/Types/SlateConstants.h
--- a/Engine/Source/Runtime/SlateCore/Public/Types/SlateConstants.h
+++ b/Engine/Source/Runtime/SlateCore/Public/Types/SlateConstants.h
@@ -3,6 +3,8 @@
 #pragma once
 
 #include "CoreMinimal.h"
+#include "HAL/IConsoleManager.h"
 
 /** How much to scroll for each click of the mouse wheel (in Slate Screen Units). */
-const float WheelScrollAmount = 32.0f;
+extern SLATECORE_API TAutoConsoleVariable<float> GlobalScrollAmount;
+inline float GetGlobalScrollAmount() { return GlobalScrollAmount.GetValueOnAnyThread(); }
