Fix egl issues for Android ARM64 (same as tested in CL3348071)
#ue4
#android

diff --git a/Engine/Source/Runtime/OpenGLDrv/Private/Android/AndroidESDeferredOpenGL.cpp b/Engine/Source/Runtime/OpenGLDrv/Private/Android/AndroidESDeferredOpenGL.cpp
--- a/Engine/Source/Runtime/OpenGLDrv/Private/Android/AndroidESDeferredOpenGL.cpp
+++ b/Engine/Source/Runtime/OpenGLDrv/Private/Android/AndroidESDeferredOpenGL.cpp
@@ -32,6 +32,11 @@ ENUM_GL_ENTRYPOINTS_CORE(DEFINE_GL_ENTRYPOINTS) \
 ENUM_GL_ENTRYPOINTS_MANUAL(DEFINE_GL_ENTRYPOINTS) \
 ENUM_GL_ENTRYPOINTS_OPTIONAL(DEFINE_GL_ENTRYPOINTS)
 
+PFNEGLGETSYSTEMTIMENVPROC eglGetSystemTimeNV_p = NULL;
+PFNEGLCREATESYNCKHRPROC eglCreateSyncKHR_p = NULL;
+PFNEGLDESTROYSYNCKHRPROC eglDestroySyncKHR_p = NULL;
+PFNEGLCLIENTWAITSYNCKHRPROC eglClientWaitSyncKHR_p = NULL;
+
 ///////////////////////////////////////////////////////////////////////////////
 
 bool FAndroidESDeferredOpenGL::bSupportsBindlessTexture = false;
@@ -167,6 +172,11 @@ void FPlatformOpenGLDevice::Init()
 	ENUM_GL_ENTRYPOINTS_MANUAL(GET_GL_ENTRYPOINTS);
 	ENUM_GL_ENTRYPOINTS_OPTIONAL(GET_GL_ENTRYPOINTS);
 
+	eglGetSystemTimeNV_p = (PFNEGLGETSYSTEMTIMENVPROC)((void*)eglGetProcAddress("eglGetSystemTimeNV"));
+	eglCreateSyncKHR_p = (PFNEGLCREATESYNCKHRPROC)((void*)eglGetProcAddress("eglCreateSyncKHR"));
+	eglDestroySyncKHR_p = (PFNEGLDESTROYSYNCKHRPROC)((void*)eglGetProcAddress("eglDestroySyncKHR"));
+	eglClientWaitSyncKHR_p = (PFNEGLCLIENTWAITSYNCKHRPROC)((void*)eglGetProcAddress("eglClientWaitSyncKHR"));
+
 	// Check that all of the required entry points have been initialized
 	bool bFoundAllEntryPoints = true;
 #define CHECK_GL_ENTRYPOINTS(Type,Func) if (Func == NULL) { bFoundAllEntryPoints = false; UE_LOG(LogRHI, Warning, TEXT("Failed to find entry point for %s"), TEXT(#Func)); }
diff --git a/Engine/Source/Runtime/OpenGLDrv/Private/Android/AndroidESDeferredOpenGL.h b/Engine/Source/Runtime/OpenGLDrv/Private/Android/AndroidESDeferredOpenGL.h
--- a/Engine/Source/Runtime/OpenGLDrv/Private/Android/AndroidESDeferredOpenGL.h
+++ b/Engine/Source/Runtime/OpenGLDrv/Private/Android/AndroidESDeferredOpenGL.h
@@ -239,13 +239,17 @@ typedef void (GL_APIENTRYP PFNGLFRAMEBUFFERTEXTUREMULTISAMPLEMULTIVIEWOVRPROC) (
 	EnumMacro(PFNGLMAKETEXTUREHANDLERESIDENTNVPROC, glMakeTextureHandleResidentNV)\
 	EnumMacro(PFNGLUNIFORMHANDLEUI64NVPROC, glUniformHandleui64NV)\
 	EnumMacro(PFNGLMAKETEXTUREHANDLENONRESIDENTNVPROC, glMakeTextureHandleNonResidentNV)\
-	EnumMacro(PFNEGLGETSYSTEMTIMENVPROC, eglGetSystemTimeNV)\
-	EnumMacro(PFNEGLCREATESYNCKHRPROC, eglCreateSyncKHR)\
-	EnumMacro(PFNEGLDESTROYSYNCKHRPROC, eglDestroySyncKHR)\
-	EnumMacro(PFNEGLCLIENTWAITSYNCKHRPROC, eglClientWaitSyncKHR)\
 	EnumMacro(PFNGLFRAMEBUFFERTEXTUREMULTIVIEWOVRPROC, glFramebufferTextureMultiviewOVR)\
 	EnumMacro(PFNGLFRAMEBUFFERTEXTUREMULTISAMPLEMULTIVIEWOVRPROC, glFramebufferTextureMultisampleMultiviewOVR)
 
+extern "C"
+{
+	extern PFNEGLGETSYSTEMTIMENVPROC eglGetSystemTimeNV_p;
+	extern PFNEGLCREATESYNCKHRPROC eglCreateSyncKHR_p;
+	extern PFNEGLDESTROYSYNCKHRPROC eglDestroySyncKHR_p;
+	extern PFNEGLCLIENTWAITSYNCKHRPROC eglClientWaitSyncKHR_p;
+}
+
 // List of all OpenGL entry points
 #define ENUM_GL_ENTRYPOINTS_ALL(EnumMacro) \
 	ENUM_GL_ENTRYPOINTS_CORE(EnumMacro) \
@@ -366,7 +370,7 @@ struct FAndroidESDeferredOpenGL : public FOpenGLESDeferred
 			if (GUseThreadedRendering)
 			{
 				//handle error here
-				EGLBoolean Result = eglDestroySyncKHR(AndroidEGL::GetInstance()->GetDisplay(), Sync.EGL);
+				EGLBoolean Result = eglDestroySyncKHR_p(AndroidEGL::GetInstance()->GetDisplay(), Sync.EGL);
 				if (Result == EGL_FALSE)
 				{
 					//handle error here
@@ -384,7 +388,7 @@ struct FAndroidESDeferredOpenGL : public FOpenGLESDeferred
 		if (IsES2())
 		{
 			check(Condition == GL_SYNC_GPU_COMMANDS_COMPLETE && Flags == 0);
-			return UGLsync(GUseThreadedRendering ? eglCreateSyncKHR(AndroidEGL::GetInstance()->GetDisplay(), EGL_SYNC_FENCE_KHR, NULL) : 0);
+			return UGLsync(GUseThreadedRendering ? eglCreateSyncKHR_p(AndroidEGL::GetInstance()->GetDisplay(), EGL_SYNC_FENCE_KHR, NULL) : 0);
 		}
 		else
 		{
@@ -418,7 +422,7 @@ struct FAndroidESDeferredOpenGL : public FOpenGLESDeferred
 			if (GUseThreadedRendering)
 			{
 				// check( Flags == GL_SYNC_FLUSH_COMMANDS_BIT );
-				GLenum Result = eglClientWaitSyncKHR(AndroidEGL::GetInstance()->GetDisplay(), Sync.EGL, EGL_SYNC_FLUSH_COMMANDS_BIT_KHR, Timeout);
+				GLenum Result = eglClientWaitSyncKHR_p(AndroidEGL::GetInstance()->GetDisplay(), Sync.EGL, EGL_SYNC_FLUSH_COMMANDS_BIT_KHR, Timeout);
 				switch (Result)
 				{
 					case EGL_TIMEOUT_EXPIRED_KHR:		return FR_TimeoutExpired;
diff --git a/Engine/Source/Runtime/OpenGLDrv/Private/Android/AndroidOpenGL.cpp b/Engine/Source/Runtime/OpenGLDrv/Private/Android/AndroidOpenGL.cpp
--- a/Engine/Source/Runtime/OpenGLDrv/Private/Android/AndroidOpenGL.cpp
+++ b/Engine/Source/Runtime/OpenGLDrv/Private/Android/AndroidOpenGL.cpp
@@ -11,10 +11,10 @@
 #include "AndroidWindow.h"
 #include "AndroidOpenGLPrivate.h"
 
-PFNEGLGETSYSTEMTIMENVPROC eglGetSystemTimeNV;
-PFNEGLCREATESYNCKHRPROC eglCreateSyncKHR = NULL;
-PFNEGLDESTROYSYNCKHRPROC eglDestroySyncKHR = NULL;
-PFNEGLCLIENTWAITSYNCKHRPROC eglClientWaitSyncKHR = NULL;
+PFNEGLGETSYSTEMTIMENVPROC eglGetSystemTimeNV_p = NULL;
+PFNEGLCREATESYNCKHRPROC eglCreateSyncKHR_p = NULL;
+PFNEGLDESTROYSYNCKHRPROC eglDestroySyncKHR_p = NULL;
+PFNEGLCLIENTWAITSYNCKHRPROC eglClientWaitSyncKHR_p = NULL;
 
 // Occlusion Queries
 PFNGLGENQUERIESEXTPROC 					glGenQueriesEXT = NULL;
@@ -276,10 +276,10 @@ bool PlatformContextIsCurrent( uint64 QueryContext )
 
 void FPlatformOpenGLDevice::LoadEXT()
 {
-	eglGetSystemTimeNV = (PFNEGLGETSYSTEMTIMENVPROC)((void*)eglGetProcAddress("eglGetSystemTimeNV"));
-	eglCreateSyncKHR = (PFNEGLCREATESYNCKHRPROC)((void*)eglGetProcAddress("eglCreateSyncKHR"));
-	eglDestroySyncKHR = (PFNEGLDESTROYSYNCKHRPROC)((void*)eglGetProcAddress("eglDestroySyncKHR"));
-	eglClientWaitSyncKHR = (PFNEGLCLIENTWAITSYNCKHRPROC)((void*)eglGetProcAddress("eglClientWaitSyncKHR"));
+	eglGetSystemTimeNV_p = (PFNEGLGETSYSTEMTIMENVPROC)((void*)eglGetProcAddress("eglGetSystemTimeNV"));
+	eglCreateSyncKHR_p = (PFNEGLCREATESYNCKHRPROC)((void*)eglGetProcAddress("eglCreateSyncKHR"));
+	eglDestroySyncKHR_p = (PFNEGLDESTROYSYNCKHRPROC)((void*)eglGetProcAddress("eglDestroySyncKHR"));
+	eglClientWaitSyncKHR_p = (PFNEGLCLIENTWAITSYNCKHRPROC)((void*)eglGetProcAddress("eglClientWaitSyncKHR"));
 
 	glDebugMessageControlKHR = (PFNGLDEBUGMESSAGECONTROLKHRPROC)((void*)eglGetProcAddress("glDebugMessageControlKHR"));
 
diff --git a/Engine/Source/Runtime/OpenGLDrv/Private/Android/AndroidOpenGL.h b/Engine/Source/Runtime/OpenGLDrv/Private/Android/AndroidOpenGL.h
--- a/Engine/Source/Runtime/OpenGLDrv/Private/Android/AndroidOpenGL.h
+++ b/Engine/Source/Runtime/OpenGLDrv/Private/Android/AndroidOpenGL.h
@@ -152,10 +152,10 @@ extern PFNGLVERTEXATTRIBIPOINTERPROC	glVertexAttribIPointer;
 
 extern "C"
 {
-	extern PFNEGLGETSYSTEMTIMENVPROC eglGetSystemTimeNV;
-	extern PFNEGLCREATESYNCKHRPROC eglCreateSyncKHR;
-	extern PFNEGLDESTROYSYNCKHRPROC eglDestroySyncKHR;
-	extern PFNEGLCLIENTWAITSYNCKHRPROC eglClientWaitSyncKHR;
+	extern PFNEGLGETSYSTEMTIMENVPROC eglGetSystemTimeNV_p;
+	extern PFNEGLCREATESYNCKHRPROC eglCreateSyncKHR_p;
+	extern PFNEGLDESTROYSYNCKHRPROC eglDestroySyncKHR_p;
+	extern PFNEGLCLIENTWAITSYNCKHRPROC eglClientWaitSyncKHR_p;
 }
 
 struct FAndroidOpenGL : public FOpenGLES2
@@ -205,7 +205,7 @@ struct FAndroidOpenGL : public FOpenGLES2
 		if (GUseThreadedRendering)
 		{
 			//handle error here
-			EGLBoolean Result = eglDestroySyncKHR( AndroidEGL::GetInstance()->GetDisplay(), Sync );
+			EGLBoolean Result = eglDestroySyncKHR_p( AndroidEGL::GetInstance()->GetDisplay(), Sync );
 			if(Result == EGL_FALSE)
 			{
 				//handle error here
@@ -216,7 +216,7 @@ struct FAndroidOpenGL : public FOpenGLES2
 	static FORCEINLINE UGLsync FenceSync(GLenum Condition, GLbitfield Flags)
 	{
 		check(Condition == GL_SYNC_GPU_COMMANDS_COMPLETE && Flags == 0);
-		return GUseThreadedRendering ? eglCreateSyncKHR( AndroidEGL::GetInstance()->GetDisplay(), EGL_SYNC_FENCE_KHR, NULL ) : 0;
+		return GUseThreadedRendering ? eglCreateSyncKHR_p( AndroidEGL::GetInstance()->GetDisplay(), EGL_SYNC_FENCE_KHR, NULL ) : 0;
 	}
 	
 	static FORCEINLINE bool IsSync(UGLsync Sync)
@@ -236,7 +236,7 @@ struct FAndroidOpenGL : public FOpenGLES2
 		if (GUseThreadedRendering)
 		{
 			// check( Flags == GL_SYNC_FLUSH_COMMANDS_BIT );
-			GLenum Result = eglClientWaitSyncKHR( AndroidEGL::GetInstance()->GetDisplay(), Sync, EGL_SYNC_FLUSH_COMMANDS_BIT_KHR, Timeout );
+			GLenum Result = eglClientWaitSyncKHR_p( AndroidEGL::GetInstance()->GetDisplay(), Sync, EGL_SYNC_FLUSH_COMMANDS_BIT_KHR, Timeout );
 			switch (Result)
 			{
 			case EGL_TIMEOUT_EXPIRED_KHR:		return FR_TimeoutExpired;
