More CIS fixing (missed a file in previous commit).

diff --git a/Engine/Source/Runtime/OpenGLDrv/Private/Linux/OpenGLLinux.cpp b/Engine/Source/Runtime/OpenGLDrv/Private/Linux/OpenGLLinux.cpp
--- a/Engine/Source/Runtime/OpenGLDrv/Private/Linux/OpenGLLinux.cpp
+++ b/Engine/Source/Runtime/OpenGLDrv/Private/Linux/OpenGLLinux.cpp
@@ -7,6 +7,7 @@
 #include "Linux/OpenGLLinux.h"
 #include "Misc/ScopeLock.h"
 #include "OpenGLDrv.h"
+#include "SDL.h"
 #include "OpenGLDrvPrivate.h"
 #include "ComponentReregisterContext.h"
 
