Fix VS2017 compile warning due to shadowed variables.

diff --git a/Engine/Source/Runtime/OpenGLDrv/Private/OpenGLTexture.cpp b/Engine/Source/Runtime/OpenGLDrv/Private/OpenGLTexture.cpp
--- a/Engine/Source/Runtime/OpenGLDrv/Private/OpenGLTexture.cpp
+++ b/Engine/Source/Runtime/OpenGLDrv/Private/OpenGLTexture.cpp
@@ -1216,7 +1216,7 @@ void TOpenGLTexture<RHIResourceType>::Unlock(uint32 MipIndex,uint32 ArrayIndex)
 }
 
 template<typename RHIResourceType>
-void TOpenGLTexture<RHIResourceType>::CloneViaCopyImage( TOpenGLTexture<RHIResourceType>* Src, uint32 NumMips, int32 SrcOffset, int32 DstOffset)
+void TOpenGLTexture<RHIResourceType>::CloneViaCopyImage( TOpenGLTexture<RHIResourceType>* Src, uint32 InNumMips, int32 SrcOffset, int32 DstOffset)
 {
 	VERIFY_GL_SCOPE();
 	
@@ -1225,7 +1225,7 @@ void TOpenGLTexture<RHIResourceType>::CloneViaCopyImage( TOpenGLTexture<RHIResou
 	for (uint32 ArrayIndex = 0; ArrayIndex < this->GetEffectiveSizeZ(); ArrayIndex++)
 	{
 		// use the Copy Image functionality to copy mip level by mip level
-		for(uint32 MipIndex = 0;MipIndex < NumMips;++MipIndex)
+		for(uint32 MipIndex = 0;MipIndex < InNumMips;++MipIndex)
 		{
 			// Calculate the dimensions of the mip-map.
 			const uint32 DstMipIndex = MipIndex + DstOffset;
@@ -1244,7 +1244,7 @@ void TOpenGLTexture<RHIResourceType>::CloneViaCopyImage( TOpenGLTexture<RHIResou
 }
 
 template<typename RHIResourceType>
-void TOpenGLTexture<RHIResourceType>::CloneViaPBO( TOpenGLTexture<RHIResourceType>* Src, uint32 NumMips, int32 SrcOffset, int32 DstOffset)
+void TOpenGLTexture<RHIResourceType>::CloneViaPBO( TOpenGLTexture<RHIResourceType>* Src, uint32 InNumMips, int32 SrcOffset, int32 DstOffset)
 {
 	VERIFY_GL_SCOPE();
 	
@@ -1271,7 +1271,7 @@ void TOpenGLTexture<RHIResourceType>::CloneViaPBO( TOpenGLTexture<RHIResourceTyp
 	for (uint32 ArrayIndex = 0; ArrayIndex < this->GetEffectiveSizeZ(); ArrayIndex++)
 	{
 		// use PBO functionality to copy mip level by mip level
-		for(uint32 MipIndex = 0;MipIndex < NumMips;++MipIndex)
+		for(uint32 MipIndex = 0;MipIndex < InNumMips;++MipIndex)
 		{
 			// Actual mip levels
 			const uint32 DstMipIndex = MipIndex + DstOffset;
