Remaining conflicts from //UE4/Main merge.

diff --git a/Engine/Source/Runtime/OpenGLDrv/Private/Linux/OpenGLLinux.cpp b/Engine/Source/Runtime/OpenGLDrv/Private/Linux/OpenGLLinux.cpp
--- a/Engine/Source/Runtime/OpenGLDrv/Private/Linux/OpenGLLinux.cpp
+++ b/Engine/Source/Runtime/OpenGLDrv/Private/Linux/OpenGLLinux.cpp
@@ -1,7 +1,7 @@
 // Copyright 1998-2017 Epic Games, Inc. All Rights Reserved.
 
 /*=============================================================================
-	OpenGLWindowsLoader.cpp: Manual loading of OpenGL functions from DLL.
+	OpenGLLinux.cpp: OpenGL context management on Linux
 =============================================================================*/
 
 #include "Linux/OpenGLLinux.h"
@@ -27,19 +27,19 @@ typedef SDL_GLContext	SDL_HGLContext;
 /*------------------------------------------------------------------------------
 	OpenGL context management.
 ------------------------------------------------------------------------------*/
-static void _ContextMakeCurrent( SDL_HWindow hWnd, SDL_HGLContext hGLDC )
+void Linux_ContextMakeCurrent( SDL_HWindow hWnd, SDL_HGLContext hGLDC )
 {
 	GLint Result = SDL_GL_MakeCurrent( hWnd, hGLDC );
 	if (Result != 0)
 	{
 		// this is a warning and not error, since Slate sometimes destroys windows before
 		// releasing RHI resources associated with them. This code can result in leaks - proper resolution is tracked as UE-7388
-		FString SdlError(ANSI_TO_TCHAR(SDL_GetError()));
+		FString SdlError(UTF8_TO_TCHAR(SDL_GetError()));
 		UE_LOG(LogLinux, Warning, TEXT("SDL_GL_MakeCurrent() failed, SDL error: '%s'"), *SdlError );
 	}
 }
 
-static SDL_HGLContext _GetCurrentContext()
+SDL_HGLContext Linux_GetCurrentContext()
 {
 	return SDL_GL_GetCurrentContext();
 }
@@ -76,7 +76,7 @@ public:
 			}
 			// no need to glFlush() on Windows, it does flush by itself before switching contexts
 
-			_ContextMakeCurrent( Context->hWnd, Context->hGLContext );
+			Linux_ContextMakeCurrent( Context->hWnd, Context->hGLContext );
 		}
 	}
 
@@ -85,13 +85,13 @@ public:
 		if (!bSameDCAndContext)
 		{
 			glFlush();	// not needed on Windows, it does flush by itself before switching contexts
-			if	( hPreGLContext )
+			if (hPreGLContext)
 			{
-				_ContextMakeCurrent( hPreWnd, hPreGLContext );
+				Linux_ContextMakeCurrent( hPreWnd, hPreGLContext );
 			}
 			else
 			{
-				_ContextMakeCurrent( NULL, NULL );
+				Linux_ContextMakeCurrent( NULL, NULL );
 			}
 		}
 
@@ -105,12 +105,12 @@ private:
 
 
 
-void _DeleteQueriesForCurrentContext( SDL_HGLContext hGLContext );
+void Linux_DeleteQueriesForCurrentContext( SDL_HGLContext hGLContext );
 
 /**
  * Create a dummy window used to construct OpenGL contexts.
  */
-static void _PlatformCreateDummyGLWindow( FPlatformOpenGLContext *OutContext )
+void Linux_PlatformCreateDummyGLWindow( FPlatformOpenGLContext *OutContext )
 {
 	static bool bInitializedWindowClass = false;
 
@@ -136,17 +136,17 @@ static void _PlatformCreateDummyGLWindow( FPlatformOpenGLContext *OutContext )
  * Determine OpenGL Context version based on command line arguments
  */
 
-static bool IsOpenGL3Forced()
+bool IsOpenGL3Forced()
 {
 	return FParse::Param(FCommandLine::Get(),TEXT("opengl3"));
 }
 
-static bool IsOpenGL4Forced()
+bool IsOpenGL4Forced()
 {
 	return FParse::Param(FCommandLine::Get(),TEXT("opengl4"));
 }
 
-static void PlatformOpenGLVersionFromCommandLine(int& OutMajorVersion, int& OutMinorVersion)
+void PlatformOpenGLVersionFromCommandLine(int& OutMajorVersion, int& OutMinorVersion)
 {
 	bool bGL3 = IsOpenGL3Forced();
 	bool bGL4 = IsOpenGL4Forced();
@@ -179,7 +179,7 @@ static void PlatformOpenGLVersionFromCommandLine(int& OutMajorVersion, int& OutM
 /**
  * Enable/Disable debug context from the commandline
  */
-static bool _PlatformOpenGLDebugCtx()
+bool Linux_PlatformOpenGLDebugCtx()
 {
 #if UE_BUILD_DEBUG
 	return ! FParse::Param(FCommandLine::Get(),TEXT("openglNoDebug"));
@@ -191,7 +191,7 @@ static bool _PlatformOpenGLDebugCtx()
 /**
  * Create a core profile OpenGL context.
  */
-static void _PlatformCreateOpenGLContextCore(FPlatformOpenGLContext* OutContext)
+void Linux_PlatformCreateOpenGLContextCore(FPlatformOpenGLContext* OutContext)
 {
 	check( OutContext );
 
@@ -212,7 +212,7 @@ static void _PlatformCreateOpenGLContextCore(FPlatformOpenGLContext* OutContext)
 		SDL_GL_GetAttribute( SDL_GL_CONTEXT_MAJOR_VERSION, &OpenGLMajorVersion );
 		SDL_GL_GetAttribute( SDL_GL_CONTEXT_MINOR_VERSION, &OpenGLMinorVersion );
 		
-		UE_LOG(LogInit, Error, TEXT("_PlatformCreateOpenGLContextCore - Could not create OpenGL %d.%d context, SDL error: '%s'"),
+		UE_LOG(LogInit, Error, TEXT("Linux_PlatformCreateOpenGLContextCore - Could not create OpenGL %d.%d context, SDL error: '%s'"),
 				OpenGLMajorVersion, OpenGLMinorVersion,
 				*SdlError
 			);
@@ -245,8 +245,8 @@ struct FPlatformOpenGLDevice
 
 		verifyf(SDL_GL_SetAttribute(SDL_GL_SHARE_WITH_CURRENT_CONTEXT, 0) == 0, TEXT("SDL_GL_SetAttribute(SDL_GL_SHARE_WITH_CURRENT_CONTEXT, 0) failed: %s\n."), UTF8_TO_TCHAR(SDL_GetError()));
 
-		_PlatformCreateDummyGLWindow( &SharedContext );
-		_PlatformCreateOpenGLContextCore( &SharedContext );
+		Linux_PlatformCreateDummyGLWindow( &SharedContext );
+		Linux_PlatformCreateOpenGLContextCore( &SharedContext );
 
 		if (SharedContext.hGLContext == nullptr)
 		{
@@ -268,12 +268,13 @@ struct FPlatformOpenGLDevice
 
 		verifyf(SDL_GL_SetAttribute(SDL_GL_SHARE_WITH_CURRENT_CONTEXT, 1) == 0, TEXT("SDL_GL_SetAttribute(SDL_GL_SHARE_WITH_CURRENT_CONTEXT, 1) failed: %s\n."), UTF8_TO_TCHAR(SDL_GetError()));
 
-		_ContextMakeCurrent( SharedContext.hWnd, SharedContext.hGLContext );
+		Linux_ContextMakeCurrent( SharedContext.hWnd, SharedContext.hGLContext );
+
+		Linux_PlatformCreateDummyGLWindow( &RenderingContext );
+		Linux_PlatformCreateOpenGLContextCore( &RenderingContext );
 
-		_PlatformCreateDummyGLWindow( &RenderingContext );
-		_PlatformCreateOpenGLContextCore( &RenderingContext );
+		check( RenderingContext.hGLContext );
 
-		check	( RenderingContext.hGLContext )
 		{
 			FScopeContext ScopeContext( &RenderingContext );
 			InitDebugContext();
@@ -287,7 +288,7 @@ struct FPlatformOpenGLDevice
 	{
 		check( NumUsedContexts==0 );
 
-		_ContextMakeCurrent( NULL, NULL );
+		Linux_ContextMakeCurrent( NULL, NULL );
 
 		OnQueryInvalidation();
 		PlatformReleaseOpenGLContext( this,&RenderingContext );
@@ -327,7 +328,7 @@ FPlatformOpenGLContext* PlatformCreateOpenGLContext(FPlatformOpenGLDevice* Devic
 	{
 		FScopeContext Scope(&(Device->SharedContext));
 		verifyf(SDL_GL_SetAttribute(SDL_GL_SHARE_WITH_CURRENT_CONTEXT, 1) == 0, TEXT("SDL_GL_SetAttribute(SDL_GL_SHARE_WITH_CURRENT_CONTEXT, 1) failed: %s\n."), UTF8_TO_TCHAR(SDL_GetError()));
-		_PlatformCreateOpenGLContextCore( Context );
+		Linux_PlatformCreateOpenGLContextCore( Context );
 	}
 
 	check( Context->hGLContext );
@@ -351,7 +352,7 @@ void PlatformReleaseOpenGLContext( FPlatformOpenGLDevice* Device, FPlatformOpenG
 		{
 			FScopeContext ScopeContext( Context );
 
-			_DeleteQueriesForCurrentContext( Context->hGLContext );
+			Linux_DeleteQueriesForCurrentContext( Context->hGLContext );
 			glBindVertexArray(0);
 			glDeleteVertexArrays(1, &Context->VertexArrayObject);
 
@@ -585,12 +586,12 @@ void PlatformRenderingContextSetup(FPlatformOpenGLDevice* Device)
 {
 	check( Device && Device->RenderingContext.hWnd && Device->RenderingContext.hGLContext );
 
-	if ( _GetCurrentContext() )
+	if ( Linux_GetCurrentContext() )
 	{
 		glFlush();
 	}
 
-	_ContextMakeCurrent( Device->RenderingContext.hWnd, Device->RenderingContext.hGLContext );
+	Linux_ContextMakeCurrent( Device->RenderingContext.hWnd, Device->RenderingContext.hGLContext );
 }
 
 void PlatformSharedContextSetup(FPlatformOpenGLDevice* Device)
@@ -598,84 +599,30 @@ void PlatformSharedContextSetup(FPlatformOpenGLDevice* Device)
 	check( Device && Device->SharedContext.hWnd && Device->SharedContext.hGLContext );
 
 	// no need to glFlush() on Windows, it does flush by itself before switching contexts
-	if ( _GetCurrentContext() )
+	if ( Linux_GetCurrentContext() )
 	{
 		glFlush();
 	}
 
-	_ContextMakeCurrent( Device->SharedContext.hWnd, Device->SharedContext.hGLContext );
+	Linux_ContextMakeCurrent( Device->SharedContext.hWnd, Device->SharedContext.hGLContext );
 }
 
 
 void PlatformNULLContextSetup()
 {
-	if ( _GetCurrentContext() )
+	if ( Linux_GetCurrentContext() )
 	{
 		glFlush();
 	}
 
-	_ContextMakeCurrent( NULL, NULL );
+	Linux_ContextMakeCurrent( NULL, NULL );
 }
 
 
-#if 0
-static void __GetBestFullscreenResolution( SDL_HWindow hWnd, int32 *pWidth, int32 *pHeight )
-{
-	uint32 InitializedMode = false;
-	uint32 BestWidth = 0;
-	uint32 BestHeight = 0;
-	uint32 ModeIndex = 0;
-
-	int32 dsp_idx = SDL_GetWindowDisplayIndex( hWnd );
-	if ( dsp_idx < 0 )
-	{	dsp_idx = 0;	}
-
-	SDL_DisplayMode dsp_mode;
-	FMemory::Memzero( &dsp_mode, sizeof(SDL_DisplayMode) );
-
-	while ( !SDL_GetDisplayMode( dsp_idx, ModeIndex++, &dsp_mode ) )
-	{
-		bool IsEqualOrBetterWidth  = FMath::Abs((int32)dsp_mode.w - (int32)(*pWidth))  <= FMath::Abs((int32)BestWidth  - (int32)(*pWidth ));
-		bool IsEqualOrBetterHeight = FMath::Abs((int32)dsp_mode.h - (int32)(*pHeight)) <= FMath::Abs((int32)BestHeight - (int32)(*pHeight));
-		if	(!InitializedMode || (IsEqualOrBetterWidth && IsEqualOrBetterHeight))
-		{
-			BestWidth = dsp_mode.w;
-			BestHeight = dsp_mode.h;
-			InitializedMode = true;
-		}
-	}
-
-	check(InitializedMode);
-
-	*pWidth  = BestWidth;
-	*pHeight = BestHeight;
-}
-
-
-static void __SetBestFullscreenDisplayMode( SDL_HWindow hWnd, int32 *pWidth, int32 *pHeight )
-{
-	SDL_DisplayMode dsp_mode;
-
-	dsp_mode.w = *pWidth;
-	dsp_mode.h = *pHeight;
-	dsp_mode.format = SDL_PIXELFORMAT_ARGB8888;
-	dsp_mode.refresh_rate = 60;
-	dsp_mode.driverdata = NULL;
-
-	__GetBestFullscreenResolution( hWnd, &dsp_mode.w, &dsp_mode.h );
-	SDL_SetWindowDisplayMode( hWnd, &dsp_mode );
-
-	*pWidth  = dsp_mode.w;
-	*pHeight = dsp_mode.h;
-}
-#endif
-
 
 /**
  * Resize the GL context.
  */
-static int gResizeC = 0;
-
 void PlatformResizeGLContext(	FPlatformOpenGLDevice* Device,
 								FPlatformOpenGLContext* Context,
 								uint32 SizeX, uint32 SizeY,
@@ -684,80 +631,38 @@ void PlatformResizeGLContext(	FPlatformOpenGLDevice* Device,
 								GLenum BackBufferTarget,
 								GLuint BackBufferResource)
 {
-//	printf( "--------------------------------------------------------\n" );
-//	printf("Ariel - PlatformResizeGLContext - Start - %03d\n", gResizeC );
-
+	FScopeLock ScopeLock(Device->ContextUsageGuard);
 	{
-		FScopeLock ScopeLock(Device->ContextUsageGuard);
-
-#if 0
-		SDL_HWindow h_wnd = Context->hWnd;
-
-		int32 closest_w = SizeX;
-		int32 closest_h = SizeY;
-
-		if	( bFullscreen )
-		{
-
-			SDL_SetWindowPosition( h_wnd, 0, 0 );
-			__SetBestFullscreenDisplayMode( h_wnd, &closest_w, &closest_h );
-
-			SDL_SetWindowBordered( h_wnd, SDL_FALSE );
-			SDL_SetWindowFullscreen( h_wnd, SDL_WINDOW_FULLSCREEN );
+		FScopeContext ScopeContext(Context);
 
-		}
-		else 
-		if	( bWasFullscreen )
+		if (Context->ViewportFramebuffer == 0)
 		{
-			SDL_SetWindowFullscreen( h_wnd, 0 );
-			SDL_SetWindowBordered( h_wnd, SDL_TRUE );
-			SDL_SetWindowSize( h_wnd, SizeX, SizeY );
+			glGenFramebuffers(1, &Context->ViewportFramebuffer);
 		}
-#endif
-
-		{
-			FScopeContext ScopeContext(Context);
-
-			if (Context->ViewportFramebuffer == 0)
-			{
-				glGenFramebuffers(1, &Context->ViewportFramebuffer);
-			}
-			glBindFramebuffer(GL_FRAMEBUFFER, Context->ViewportFramebuffer);
-			glFramebufferTexture2D(GL_FRAMEBUFFER, GL_COLOR_ATTACHMENT0, BackBufferTarget, BackBufferResource, 0);
-#if UE_BUILD_DEBUG
-			glReadBuffer(GL_COLOR_ATTACHMENT0);
-			glDrawBuffer(GL_COLOR_ATTACHMENT0);
-#endif
-			FOpenGL::CheckFrameBuffer();
+		glBindFramebuffer(GL_FRAMEBUFFER, Context->ViewportFramebuffer);
+		glFramebufferTexture2D(GL_FRAMEBUFFER, GL_COLOR_ATTACHMENT0, BackBufferTarget, BackBufferResource, 0);
+		FOpenGL::CheckFrameBuffer();
 
-			glViewport( 0, 0, SizeX, SizeY );			
-		//	printf( "      - NewVPW = %d, NewVPH = %d\n", SizeX, SizeY );
+		glViewport( 0, 0, SizeX, SizeY );			
 
-			static GLfloat ZeroColor[4] = { 0.0f, 0.0f, 0.0f, 0.0f };
-			glClearBufferfv(GL_COLOR, 0, ZeroColor );
-		}
+		static GLfloat ZeroColor[4] = { 0.0f, 0.0f, 0.0f, 0.0f };
+		glClearBufferfv(GL_COLOR, 0, ZeroColor );
+	}
 
 
-		//	bool b_need_reattach = bFullscreen || bWasFullscreen;
-		if	( bFullscreen )
-		{
-			//	Deregister all components
-			//	this line will fix the geometry missing and color distortion bug on Linux/Nvidia machine
-			//	it detach all the components and re attach all the components
-			FGlobalComponentReregisterContext RecreateComponents;
-		}
-		else
-		if	( bWasFullscreen )
-		{
-			FGlobalComponentReregisterContext RecreateComponents;
-		}
+	if (bFullscreen)
+	{
+		//	Deregister all components
+		//	this line will fix the geometry missing and color distortion bug on Linux/Nvidia machine
+		//	it detach all the components and re attach all the components
+		FGlobalComponentReregisterContext RecreateComponents;
+	}
+	else if	(bWasFullscreen)
+	{
+		FGlobalComponentReregisterContext RecreateComponents;
 	}
-
-//	printf("Ariel - PlatformResizeGLContext - End - %03d\n", gResizeC );
-	gResizeC++;
 }
 
-//		int dip_num = SDL_GetNumVideoDisplays();
 void PlatformGetSupportedResolution(uint32 &Width, uint32 &Height)
 {
 	uint32 InitializedMode = false;
@@ -765,17 +670,17 @@ void PlatformGetSupportedResolution(uint32 &Width, uint32 &Height)
 	uint32 BestHeight = 0;
 	uint32 ModeIndex = 0;
 
-	SDL_DisplayMode dsp_mode;
-	FMemory::Memzero( &dsp_mode, sizeof(SDL_DisplayMode) );
+	SDL_DisplayMode DisplayMode;
+	FMemory::Memzero( &DisplayMode, sizeof(DisplayMode) );
 
-	while ( !SDL_GetDisplayMode( 0, ModeIndex++, &dsp_mode ) )
+	while ( !SDL_GetDisplayMode( 0, ModeIndex++, &DisplayMode ) )
 	{
-		bool IsEqualOrBetterWidth = FMath::Abs((int32)dsp_mode.w - (int32)Width) <= FMath::Abs((int32)BestWidth - (int32)Width);
-		bool IsEqualOrBetterHeight = FMath::Abs((int32)dsp_mode.h - (int32)Height) <= FMath::Abs((int32)BestHeight - (int32)Height);
+		bool IsEqualOrBetterWidth = FMath::Abs((int32)DisplayMode.w - (int32)Width) <= FMath::Abs((int32)BestWidth - (int32)Width);
+		bool IsEqualOrBetterHeight = FMath::Abs((int32)DisplayMode.h - (int32)Height) <= FMath::Abs((int32)BestHeight - (int32)Height);
 		if	(!InitializedMode || (IsEqualOrBetterWidth && IsEqualOrBetterHeight))
 		{
-			BestWidth = dsp_mode.w;
-			BestHeight = dsp_mode.h;
+			BestWidth = DisplayMode.w;
+			BestHeight = DisplayMode.h;
 			InitializedMode = true;
 		}
 	}
@@ -811,7 +716,7 @@ bool PlatformGetAvailableResolutions( FScreenResolutionArray& Resolutions, bool
 	SDL_DisplayMode DisplayMode;
 	FMemory::Memzero(&DisplayMode, sizeof(SDL_DisplayMode));
 
-	while ( !SDL_GetDisplayMode( 0, ModeIndex++, &DisplayMode ) )
+	while (!SDL_GetDisplayMode( 0, ModeIndex++, &DisplayMode ))
 	{
 		if (((int32)DisplayMode.w >= MinAllowableResolutionX) &&
 			((int32)DisplayMode.w <= MaxAllowableResolutionX) &&
@@ -823,8 +728,7 @@ bool PlatformGetAvailableResolutions( FScreenResolutionArray& Resolutions, bool
 			if (bIgnoreRefreshRate == false)
 			{
 				if (((int32)DisplayMode.refresh_rate < MinAllowableRefreshRate) ||
-					((int32)DisplayMode.refresh_rate > MaxAllowableRefreshRate)
-					)
+					((int32)DisplayMode.refresh_rate > MaxAllowableRefreshRate))
 				{
 					continue;
 				}
@@ -863,7 +767,6 @@ bool PlatformGetAvailableResolutions( FScreenResolutionArray& Resolutions, bool
 
 void PlatformRestoreDesktopDisplayMode()
 {
-//	printf( "ArielPlatform - RestoreDesktopDisplayMode\n" );
 }
 
 
@@ -896,7 +799,7 @@ bool PlatformInitOpenGL()
 
 		int DebugFlag = 0;
 
-		if ( _PlatformOpenGLDebugCtx() )
+		if (Linux_PlatformOpenGLDebugCtx())
 		{
 			DebugFlag = SDL_GL_CONTEXT_DEBUG_FLAG;
 		}
@@ -932,13 +835,13 @@ bool PlatformInitOpenGL()
 
 		// Create a dummy context to verify opengl support.
 		FPlatformOpenGLContext DummyContext;
-		_PlatformCreateDummyGLWindow(&DummyContext);
-		_PlatformCreateOpenGLContextCore(&DummyContext );	
+		Linux_PlatformCreateDummyGLWindow(&DummyContext);
+		Linux_PlatformCreateOpenGLContextCore(&DummyContext );	
 
 		if (DummyContext.hGLContext)
 		{
 			bOpenGLSupported = true;
-			_ContextMakeCurrent(DummyContext.hWnd, DummyContext.hGLContext);
+			Linux_ContextMakeCurrent(DummyContext.hWnd, DummyContext.hGLContext);
 		}
 		else
 		{
@@ -962,7 +865,7 @@ bool PlatformInitOpenGL()
 		// The dummy context can now be released.
 		if (DummyContext.hGLContext)
 		{
-			_ContextMakeCurrent(NULL,NULL);
+			Linux_ContextMakeCurrent(NULL,NULL);
 			SDL_GL_DeleteContext(DummyContext.hGLContext);
 		}
 		check(DummyContext.bReleaseWindowOnDestroy);
@@ -976,7 +879,7 @@ bool PlatformInitOpenGL()
 
 bool PlatformOpenGLContextValid()
 {
-	return ( _GetCurrentContext() != NULL );
+	return ( Linux_GetCurrentContext() != NULL );
 }
 
 int32 PlatformGlGetError()
@@ -986,17 +889,17 @@ int32 PlatformGlGetError()
 
 EOpenGLCurrentContext PlatformOpenGLCurrentContext( FPlatformOpenGLDevice* Device )
 {
-	SDL_HGLContext hGLContext = _GetCurrentContext();
+	SDL_HGLContext hGLContext = Linux_GetCurrentContext();
 
-	if	( hGLContext == Device->RenderingContext.hGLContext )	// most common case
+	if (LIKELY(hGLContext == Device->RenderingContext.hGLContext))	// most common case
 	{
 		return CONTEXT_Rendering;
 	}
-	else if	( hGLContext == Device->SharedContext.hGLContext )
+	else if	(hGLContext == Device->SharedContext.hGLContext)
 	{
 		return CONTEXT_Shared;
 	}
-	else if	( hGLContext )
+	else if	(hGLContext)
 	{
 		return CONTEXT_Other;
 	}
@@ -1045,7 +948,7 @@ void PlatformGetNewRenderQuery( GLuint* OutQuery, uint64* OutQueryContext )
 		check( OutQuery && OutQueryContext );
 #endif
 
-		SDL_HGLContext hGLContext = _GetCurrentContext();
+		SDL_HGLContext hGLContext = Linux_GetCurrentContext();
 		check( hGLContext );
 
 		GLuint NewQuery = 0;
@@ -1078,7 +981,7 @@ void PlatformGetNewRenderQuery( GLuint* OutQuery, uint64* OutQueryContext )
 
 void PlatformReleaseRenderQuery( GLuint Query, uint64 QueryContext )
 {
-	SDL_HGLContext hGLContext = _GetCurrentContext();
+	SDL_HGLContext hGLContext = Linux_GetCurrentContext();
 	if( (uint64)hGLContext == QueryContext )
 	{
 		FOpenGL::DeleteQueries(1, &Query );
@@ -1099,34 +1002,31 @@ void PlatformReleaseRenderQuery( GLuint Query, uint64 QueryContext )
 
 bool PlatformContextIsCurrent( uint64 QueryContext )
 {
-	return (uint64)_GetCurrentContext() == QueryContext;
+	return (uint64)Linux_GetCurrentContext() == QueryContext;
 }
 
 FRHITexture* PlatformCreateBuiltinBackBuffer( FOpenGLDynamicRHI* OpenGLRHI, uint32 SizeX, uint32 SizeY )
 {
-	return NULL;
+	return nullptr;
 }
 
-void _DeleteQueriesForCurrentContext( SDL_HGLContext hGLContext )
+void Linux_DeleteQueriesForCurrentContext( SDL_HGLContext hGLContext )
 {
 	if( !ReleasedQueriesGuard )
 	{
 		ReleasedQueriesGuard = new FCriticalSection;
 	}
 
+	FScopeLock Lock(ReleasedQueriesGuard);
+	for( int32 Index = 0; Index < ReleasedQueries.Num(); ++Index )
 	{
-		FScopeLock Lock(ReleasedQueriesGuard);
-		for( int32 Index = 0; Index < ReleasedQueries.Num(); ++Index )
+		if( ReleasedQueries[Index].hGLContext == hGLContext )
 		{
-			if( ReleasedQueries[Index].hGLContext == hGLContext )
-			{
-				FOpenGL::DeleteQueries(1,&ReleasedQueries[Index].Query);
-				ReleasedQueries.RemoveAtSwap(Index);
-				--Index;
-			}
+			FOpenGL::DeleteQueries(1,&ReleasedQueries[Index].Query);
+			ReleasedQueries.RemoveAtSwap(Index);
+			--Index;
 		}
 	}
-
 }
 
 void FLinuxOpenGL::ProcessExtensions( const FString& ExtensionsString )
@@ -1138,13 +1038,11 @@ void FLinuxOpenGL::ProcessExtensions( const FString& ExtensionsString )
 	if ( VendorName.Contains(TEXT("ATI ")) )
 	{
 		// Workaround for AMD driver not handling GL_SRGB8_ALPHA8 in glTexStorage2D() properly (gets treated as non-sRGB)
-		glTexStorage1D = NULL;
-		glTexStorage2D = NULL;
-		glTexStorage3D = NULL;
+		// FIXME: obsolete ? this was the case in <= 2014
+		glTexStorage1D = nullptr;
+		glTexStorage2D = nullptr;
+		glTexStorage3D = nullptr;
 
 		FOpenGLBase::bSupportsCopyImage  = false;
 	}
 }
-
-
-
diff --git a/Engine/Source/Runtime/Renderer/Private/LightPropagationVolume.cpp b/Engine/Source/Runtime/Renderer/Private/LightPropagationVolume.cpp
--- a/Engine/Source/Runtime/Renderer/Private/LightPropagationVolume.cpp
+++ b/Engine/Source/Runtime/Renderer/Private/LightPropagationVolume.cpp
@@ -898,7 +898,7 @@ void FLightPropagationVolume::InitSettings(FRHICommandListImmediate& RHICmdList,
 			LPV_GRIDRES,
 			PF_FloatRGBA,
 			FClearValueBinding::None,
-			TexCreate_HideInVisualizeTexture | TexCreate_FastVRAM,
+			TexCreate_HideInVisualizeTexture | GFastVRamConfig.LPV,
 			TexCreate_ShaderResource | TexCreate_UAV,
 			false,
 			1));
@@ -932,7 +932,7 @@ void FLightPropagationVolume::InitSettings(FRHICommandListImmediate& RHICmdList,
 				LPV_GRIDRES,
 				PF_G8,
 				FClearValueBinding::None,
-				TexCreate_HideInVisualizeTexture | TexCreate_FastVRAM,
+				TexCreate_HideInVisualizeTexture | GFastVRamConfig.LPV,
 				TexCreate_ShaderResource | TexCreate_UAV,
 				false,
 				1));
