Fix OpenGL ES deferred compile errors.
(mirrored from //UE4/Dev-Mobile CL3620341)
#ue4
#android

diff --git a/Engine/Source/Runtime/OpenGLDrv/Private/OpenGLShaders.cpp b/Engine/Source/Runtime/OpenGLDrv/Private/OpenGLShaders.cpp
--- a/Engine/Source/Runtime/OpenGLDrv/Private/OpenGLShaders.cpp
+++ b/Engine/Source/Runtime/OpenGLDrv/Private/OpenGLShaders.cpp
@@ -695,9 +695,9 @@ void OPENGLDRV_API GLSLToDeviceCompatibleGLSL(FAnsiCharArray& GlslCodeOriginal,
 	bool bUseES30ShadingLanguage = Capabilities.bUseES30ShadingLanguage;
 
 #if PLATFORM_ANDROID
-	FAndroidOpenGL::EImageExternalType ImageExternalType = FAndroidOpenGL::GetImageExternalType();
+	FOpenGL::EImageExternalType ImageExternalType = FOpenGL::GetImageExternalType();
 
-	if (ImageExternalType == FAndroidOpenGL::EImageExternalType::ImageExternal100)
+	if (ImageExternalType == FOpenGL::EImageExternalType::ImageExternal100)
 	{
 		bUseES30ShadingLanguage = false;
 	}
@@ -753,15 +753,15 @@ void OPENGLDRV_API GLSLToDeviceCompatibleGLSL(FAnsiCharArray& GlslCodeOriginal,
 #if PLATFORM_ANDROID
 			switch (ImageExternalType)
 			{
-				case FAndroidOpenGL::EImageExternalType::ImageExternal100:
+				case FOpenGL::EImageExternalType::ImageExternal100:
 					AppendCString(GlslCode, "#extension GL_OES_EGL_image_external : require\n");
 					break;
 
-				case FAndroidOpenGL::EImageExternalType::ImageExternal300:
+				case FOpenGL::EImageExternalType::ImageExternal300:
 					AppendCString(GlslCode, "#extension GL_OES_EGL_image_external : require\n");
 					break;
 
-				case FAndroidOpenGL::EImageExternalType::ImageExternalESSL300:
+				case FOpenGL::EImageExternalType::ImageExternalESSL300:
 					// GL_OES_EGL_image_external_essl3 is only compatible with ES 3.x
 					AppendCString(GlslCode, "#extension GL_OES_EGL_image_external_essl3 : require\n");
 					break;
