Fix compile error packaging for Android with Support OpenGL ES Deferred Renderer enabled

diff --git a/Engine/Source/Runtime/OpenGLDrv/Private/Android/AndroidESDeferredOpenGL.h b/Engine/Source/Runtime/OpenGLDrv/Private/Android/AndroidESDeferredOpenGL.h
--- a/Engine/Source/Runtime/OpenGLDrv/Private/Android/AndroidESDeferredOpenGL.h
+++ b/Engine/Source/Runtime/OpenGLDrv/Private/Android/AndroidESDeferredOpenGL.h
@@ -46,6 +46,7 @@ typedef FEGLGLSync UGLsync;
 #define GLdouble		GLfloat
 #define GL_BGRA			GL_BGRA_EXT
 #define GL_UNSIGNED_INT_8_8_8_8_REV	GL_UNSIGNED_BYTE
+#define GL_UNSIGNED_INT_8_8_8_8	0x8035
 
 #ifndef EGL_KHR_create_context
 #define EGL_KHR_create_context 1
