Copying //Ocean/Main to //UE4/Dev-Main (Source: //Ocean/Main @ 3755227)

========================
MAJOR FEATURES + CHANGES
========================

Change 3754304 by [REDACTED_PERSON]

	#ue4: Fix for crash that occurred when seeking past the end of an archive (#jira OCN-8990).

Change 3748500 by josh.jensen

	Fixed the round() function hack shader to determine whether the shader needs mediump or highp

	The shader fails to compile on some Android devices, because the precision keyword follows the round() definitions.

Change 3747950 by [REDACTED_PERSON]

	Exposed the virtual keyboard dismiss action to UMG.

Change 3744282 by [REDACTED_PERSON]

	Fix for failed engine check caused by GSystemResolution getting out of sync with actual device resolution / orientation.

Change 3739717 by [REDACTED_PERSON]

	Fix renderer compile errors in Mac nonunity build that started with Ocean's last integration from Main.

Change 3738008 by [REDACTED_PERSON]

	changing lockfile still waiting warning to a regular log message

Change 3729582 by [REDACTED_PERSON]

	Fix for large Metal texture memory leak when opening card packs. (And any other time the size of the static texture heaps is exceeded.) Don't release a dynamic heap until it is completely drained.

Change 3722096 by josh.jensen

	Prevented an unnecessary red<->blue swap in FVulkanDynamicRHI::RHIReadSurfaceData() when reading back a surface in VK_FORMAT_R16G16B16A16_SFLOAT format

Change 3718183 by josh.jensen

	Fixed ES2 issue reading from a render target in float format for BGRA8888 when the glReadPixels() call fails

Change 3690830 by [REDACTED_PERSON]

	Default r.MetalHullParameterSize and r.MetalDomainParameterSize to 0.  Saves 20 MB on iOS.

Change 3670184 by [REDACTED_PERSON]

	Added ENGINE_API to UTexture2D::GetMipData().

	Function was accessible is single executable stand-alone builds, but not in Editor for use when running PIE.

Change 3668970 by josh.jensen

	Provided an implementation of round() under GLES2 when the shader compiler does not provide it

Change 3624855 by josh.jensen

	Fixed issue where FOpenGLDynamicRHI::RHISetRenderTargetsAndClear() would crash when the render target did not have a color binding

	This reproduced in a number of circumstances, but it was most common with tonemap postprocessing on Android for a particular title.

Change 3693739 by [REDACTED_PERSON]

	Moved the hull and domain shader parameter size change from BaseDeviceProfiles to DefaultDeviceProfiles.

Change 3655434 by [REDACTED_PERSON]

	Adding a function to set the delay of a particle system dynamically

Change 3605036 by [REDACTED_PERSON]

	Track opened pak files. Console command "DumpOpenedFiles" print the list.

diff --git a/Engine/Source/Developer/StandaloneRenderer/Private/OpenGL/SlateOpenGLRenderer.h b/Engine/Source/Developer/StandaloneRenderer/Private/OpenGL/SlateOpenGLRenderer.h
--- a/Engine/Source/Developer/StandaloneRenderer/Private/OpenGL/SlateOpenGLRenderer.h
+++ b/Engine/Source/Developer/StandaloneRenderer/Private/OpenGL/SlateOpenGLRenderer.h
@@ -154,6 +154,7 @@ public:
 	virtual void CreateViewport( const TSharedRef<SWindow> InWindow ) override;
 	virtual void RequestResize( const TSharedPtr<SWindow>& InWindow, uint32 NewSizeX, uint32 NewSizeY ) override;
 	virtual void UpdateFullscreenState( const TSharedRef<SWindow> InWindow, uint32 OverrideResX, uint32 OverrideResY ) override;
+	virtual void SetSystemResolution(uint32 Width, uint32 Height) override {}
 	virtual void RestoreSystemResolution(const TSharedRef<SWindow> InWindow) override {}
 	virtual void ReleaseDynamicResource( const FSlateBrush& Brush ) override;
 	virtual bool GenerateDynamicImageResource(FName ResourceName, uint32 Width, uint32 Height, const TArray< uint8 >& Bytes) override;
diff --git a/Engine/Source/Developer/StandaloneRenderer/Private/Windows/D3D/SlateD3DRenderer.h b/Engine/Source/Developer/StandaloneRenderer/Private/Windows/D3D/SlateD3DRenderer.h
--- a/Engine/Source/Developer/StandaloneRenderer/Private/Windows/D3D/SlateD3DRenderer.h
+++ b/Engine/Source/Developer/StandaloneRenderer/Private/Windows/D3D/SlateD3DRenderer.h
@@ -58,6 +58,7 @@ public:
 	virtual void CreateViewport( const TSharedRef<SWindow> InWindow ) override;
 	virtual void RequestResize( const TSharedPtr<SWindow>& InWindow, uint32 NewSizeX, uint32 NewSizeY ) override;
 	virtual void UpdateFullscreenState( const TSharedRef<SWindow> InWindow, uint32 OverrideResX, uint32 OverrideResY ) override;
+	virtual void SetSystemResolution(uint32 Width, uint32 Height) override {}
 	virtual void RestoreSystemResolution(const TSharedRef<SWindow> InWindow) override {}
 	virtual void ReleaseDynamicResource( const FSlateBrush& Brush ) override;
 	virtual bool GenerateDynamicImageResource(FName ResourceName, uint32 Width, uint32 Height, const TArray< uint8 >& Bytes) override;
diff --git a/Engine/Source/Programs/AutomationTool/BuildGraph/LockFile.cs b/Engine/Source/Programs/AutomationTool/BuildGraph/LockFile.cs
--- a/Engine/Source/Programs/AutomationTool/BuildGraph/LockFile.cs
+++ b/Engine/Source/Programs/AutomationTool/BuildGraph/LockFile.cs
@@ -62,7 +62,7 @@ namespace AutomationTool
                     }
                     else if ((Iterations % 30) == 0)
                     {
-                        CommandUtils.LogWarning("Still waiting for lock file '{0}' after {1} seconds.", LockFilePath, CurrentTime.Subtract(StartTime).TotalSeconds);
+                        CommandUtils.Log("Still waiting for lock file '{0}' after {1} seconds.", LockFilePath, CurrentTime.Subtract(StartTime).TotalSeconds);
                     }
 
                     // Wait for a while before retrying.
diff --git a/Engine/Source/Runtime/Apple/MetalRHI/Private/MetalHeap.cpp b/Engine/Source/Runtime/Apple/MetalRHI/Private/MetalHeap.cpp
--- a/Engine/Source/Runtime/Apple/MetalRHI/Private/MetalHeap.cpp
+++ b/Engine/Source/Runtime/Apple/MetalRHI/Private/MetalHeap.cpp
@@ -1573,7 +1573,7 @@ void FMetalHeap::Drain(FMetalDeviceContext& Context, bool const bForce)
 						
 						// There may be heaps with no resources currently in-use but the memory has yet to be released, we can dispose of these heaps
 						TSet<NSObject<MTLTexture>*>& Textures = TextureResources.FindChecked(DynamicTextureHeaps[i][j][k][l].Key);
-						if (((bForce || NumFrames > EMetalHeapConstantsCullAfterFrames || j == EMetalHeapBufferExactSize) && Textures.Num() == 0) || (DynamicTextureHeaps[i][j][k][l].Key.usedSize == 0))
+						if (((bForce || NumFrames > EMetalHeapConstantsCullAfterFrames || j == EMetalHeapBufferExactSize) && Textures.Num() == 0) && (DynamicTextureHeaps[i][j][k][l].Key.usedSize == 0) && (((FMTLHeap*)DynamicTextureHeaps[i][j][k][l].Key).poolSize == 0))
 						{
 							TotalTextureMemory -= DynamicTextureHeaps[i][j][k][l].Key.size;
 							DEC_DWORD_STAT_FName(NumTextureHeapStats[k]);
diff --git a/Engine/Source/Runtime/Apple/MetalRHI/Private/MetalRHIPrivate.h b/Engine/Source/Runtime/Apple/MetalRHI/Private/MetalRHIPrivate.h
--- a/Engine/Source/Runtime/Apple/MetalRHI/Private/MetalRHIPrivate.h
+++ b/Engine/Source/Runtime/Apple/MetalRHI/Private/MetalRHIPrivate.h
@@ -50,7 +50,6 @@ enum EMTLTextureType
 #define METAL_SUPPORTS_INDIRECT_ARGUMENT_BUFFERS (!PLATFORM_MAC && __IPHONE_OS_VERSION_MAX_ALLOWED >= 110000) && (__clang_major__ >= 9)
 #define METAL_SUPPORTS_CAPTURE_MANAGER (PLATFORM_MAC && __MAC_OS_X_VERSION_MAX_ALLOWED >= 101300) || (!PLATFORM_MAC && __IPHONE_OS_VERSION_MAX_ALLOWED >= 110000) && (__clang_major__ >= 9)
 #define METAL_SUPPORTS_TILE_SHADERS (!PLATFORM_MAC && !PLATFORM_TVOS && __IPHONE_OS_VERSION_MAX_ALLOWED >= 110000) && (__clang_major__ >= 9)
-
 // In addition to compile-time SDK checks we also need a way to check if these are available on runtime
 extern bool GMetalSupportsIndirectArgumentBuffers;
 extern bool GMetalSupportsCaptureManager;
@@ -76,6 +75,8 @@ extern FMetalBufferFormat GMetalBufferFormats[PF_MAX];
 #define METAL_DEBUG_OPTION(Code)
 #endif
 
+extern bool GMetalSupportsTileShaders;
+
 #define SHOULD_TRACK_OBJECTS (UE_BUILD_DEBUG)
 
 /** Set to 1 to enable GPU events in Xcode frame debugger */
diff --git a/Engine/Source/Runtime/Core/Private/Math/SHMath.cpp b/Engine/Source/Runtime/Core/Private/Math/SHMath.cpp
--- a/Engine/Source/Runtime/Core/Private/Math/SHMath.cpp
+++ b/Engine/Source/Runtime/Core/Private/Math/SHMath.cpp
@@ -9,8 +9,6 @@
 float NormalizationConstants[9];
 int32 BasisL[9];
 int32 BasisM[9];
-template<> CORE_API const float TSHVector<2>::ConstantBasisIntegral = 2.0f * FMath::Sqrt(PI);
-template<> CORE_API const float TSHVector<3>::ConstantBasisIntegral = 2.0f * FMath::Sqrt(PI);
 
 /** Computes a factorial. */
 static int32 Factorial(int32 A)
diff --git a/Engine/Source/Runtime/Core/Public/Math/SHMath.h b/Engine/Source/Runtime/Core/Public/Math/SHMath.h
--- a/Engine/Source/Runtime/Core/Public/Math/SHMath.h
+++ b/Engine/Source/Runtime/Core/Public/Math/SHMath.h
@@ -37,7 +37,7 @@ public:
 	float V[NumTotalFloats];
 
 	/** The integral of the constant SH basis. */
-	CORE_API static const float ConstantBasisIntegral;
+	static constexpr float ConstantBasisIntegral = 3.5449077018110320545963349666823; // 2 * Sqrt(PI)
 
 	/** Default constructor. */
 	TSHVector()
diff --git a/Engine/Source/Runtime/CoreUObject/Private/Serialization/BulkData.cpp b/Engine/Source/Runtime/CoreUObject/Private/Serialization/BulkData.cpp
--- a/Engine/Source/Runtime/CoreUObject/Private/Serialization/BulkData.cpp
+++ b/Engine/Source/Runtime/CoreUObject/Private/Serialization/BulkData.cpp
@@ -1406,7 +1406,7 @@ void FUntypedBulkData::LoadDataIntoMemory( void* Dest )
 
 #else
 	bool bWasLoadedSuccessfully = false;
-	if ((IsInGameThread() || IsInAsyncLoadingThread()) && Package.IsValid() && Package->LinkerLoad && Package->LinkerLoad->GetOwnerThreadId() == FPlatformTLS::GetCurrentThreadId() && ((BulkDataFlags & BULKDATA_PayloadInSeperateFile) == 0))
+	if (IsInAsyncLoadingThread() && Package.IsValid() && Package->LinkerLoad && Package->LinkerLoad->GetOwnerThreadId() == FPlatformTLS::GetCurrentThreadId() && ((BulkDataFlags & BULKDATA_PayloadInSeperateFile) == 0))
 	{
 		FLinkerLoad* LinkerLoad = Package->LinkerLoad;
 		if (LinkerLoad && LinkerLoad->Loader)
diff --git a/Engine/Source/Runtime/Engine/Classes/Engine/Texture2D.h b/Engine/Source/Runtime/Engine/Classes/Engine/Texture2D.h
--- a/Engine/Source/Runtime/Engine/Classes/Engine/Texture2D.h
+++ b/Engine/Source/Runtime/Engine/Classes/Engine/Texture2D.h
@@ -263,7 +263,7 @@ public:
 	 *						Mips.Num() - FirstMipToLoad + 1 entries. Upon
 	 *						return those pointers will contain mip data.
 	 */
-	void GetMipData(int32 FirstMipToLoad, void** OutMipData);
+	ENGINE_API void GetMipData(int32 FirstMipToLoad, void** OutMipData);
 
 	/**
 	 * Returns the number of mips in this texture that are not able to be streamed.
diff --git a/Engine/Source/Runtime/Engine/Classes/Particles/ParticleSystem.h b/Engine/Source/Runtime/Engine/Classes/Particles/ParticleSystem.h
--- a/Engine/Source/Runtime/Engine/Classes/Particles/ParticleSystem.h
+++ b/Engine/Source/Runtime/Engine/Classes/Particles/ParticleSystem.h
@@ -418,6 +418,11 @@ public:
 	 */
 	ENGINE_API void SetupLODValidity();
 
+	/**
+	 * Set the time to delay spawning the particle system
+	 */
+	ENGINE_API void SetDelay(float InDelay);
+
 #if WITH_EDITOR
 	/**
 	 *	Remove all duplicate modules.
diff --git a/Engine/Source/Runtime/Engine/Private/GPUSkinVertexFactory.cpp b/Engine/Source/Runtime/Engine/Private/GPUSkinVertexFactory.cpp
--- a/Engine/Source/Runtime/Engine/Private/GPUSkinVertexFactory.cpp
+++ b/Engine/Source/Runtime/Engine/Private/GPUSkinVertexFactory.cpp
@@ -46,7 +46,9 @@ static FBoneMatricesUniformShaderParameters GBoneUniformStruct;
 	FactoryClass<bExtraBoneInfluencesT>::ShouldCache, \
 	FactoryClass<bExtraBoneInfluencesT>::ModifyCompilationEnvironment, \
 	FactoryClass<bExtraBoneInfluencesT>::SupportsTessellationShaders \
-	);
+	); \
+	template <bool bExtraBoneInfluencesT> inline FVertexFactoryType* FactoryClass<bExtraBoneInfluencesT>::GetType() const { return &StaticType; }
+
 
 #define IMPLEMENT_GPUSKINNING_VERTEX_FACTORY_TYPE(FactoryClass, ShaderFilename,bUsedWithMaterials,bSupportsStaticLighting,bSupportsDynamicLighting,bPrecisePrevWorldPos,bSupportsPositionOnly) \
 	template <bool bExtraBoneInfluencesT> FVertexFactoryShaderParameters* Construct##FactoryClass##ShaderParameters(EShaderFrequency ShaderFrequency) { return FactoryClass<bExtraBoneInfluencesT>::ConstructShaderParameters(ShaderFrequency); } \
diff --git a/Engine/Source/Runtime/Engine/Private/Particles/ParticleComponents.cpp b/Engine/Source/Runtime/Engine/Private/Particles/ParticleComponents.cpp
--- a/Engine/Source/Runtime/Engine/Private/Particles/ParticleComponents.cpp
+++ b/Engine/Source/Runtime/Engine/Private/Particles/ParticleComponents.cpp
@@ -2783,6 +2783,11 @@ void UParticleSystem::SetupLODValidity()
 	}
 }
 
+void UParticleSystem::SetDelay(float InDelay)
+{
+	Delay = InDelay;
+}
+
 #if WITH_EDITOR
 
 bool UParticleSystem::RemoveAllDuplicateModules(bool bInMarkForCooker, TMap<UObject*,bool>* OutRemovedModules)
diff --git a/Engine/Source/Runtime/OpenGLDrv/Private/OpenGLDevice.cpp b/Engine/Source/Runtime/OpenGLDrv/Private/OpenGLDevice.cpp
--- a/Engine/Source/Runtime/OpenGLDrv/Private/OpenGLDevice.cpp
+++ b/Engine/Source/Runtime/OpenGLDrv/Private/OpenGLDevice.cpp
@@ -1424,6 +1424,53 @@ static void CheckTextureCubeLodSupport()
 #endif
 }
 
+static void CheckRoundFunction()
+{
+#if PLATFORM_ANDROID
+	FOpenGL::bRequiresRoundFunctionHack = false;
+	if (IsES2Platform(GMaxRHIShaderPlatform))
+	{
+		UE_LOG(LogRHI, Display, TEXT("Testing for round() function availability"));
+		FOpenGL::bIsCheckingShaderCompilerHacks = true;
+
+		static const ANSICHAR TestVertexProgram[] = "\n"
+			"#version 100\n"
+			"attribute vec4 in_ATTRIBUTE0;\n"
+			"void main()\n"
+			"{\n"
+			"   float value = round(0.5);\n"
+			"	gl_Position.xyzw = in_ATTRIBUTE0;\n"
+			"}\n";
+
+		FShaderCode VertexShaderCode;
+		{
+			FOpenGLCodeHeader Header;
+			Header.FrequencyMarker = 0x5653;
+			Header.GlslMarker = 0x474c534c;
+
+			FMemoryWriter Writer(VertexShaderCode.GetWriteAccess(), true);
+			Writer << Header;
+			Writer.Serialize((void*)TestVertexProgram, sizeof(TestVertexProgram));
+			Writer.Close();
+		}
+
+		// Try to compile test shaders
+		TRefCountPtr<FOpenGLVertexShader> VertexShader = (FOpenGLVertexShader*)(RHICreateVertexShader(VertexShaderCode.GetReadAccess()).GetReference());
+		if (!VerifyCompiledShader(VertexShader->Resource, TestVertexProgram, false))
+		{
+			UE_LOG(LogRHI, Warning, TEXT("Using the round() function hack, because the vertex shader for round function test failed to compile."));
+			FOpenGL::bRequiresRoundFunctionHack = true;
+			FOpenGL::bIsCheckingShaderCompilerHacks = false;
+			return;
+		}
+
+		FOpenGL::bIsCheckingShaderCompilerHacks = false;
+
+		UE_LOG(LogRHI, Warning, TEXT("Disabling the need for the round() function hack"));
+	}
+#endif
+}
+
 void FOpenGLDynamicRHI::Init()
 {
 	check(!GIsRHIInitialized);
@@ -1491,6 +1538,7 @@ void FOpenGLDynamicRHI::Init()
 
 	CheckTextureCubeLodSupport();
 	CheckVaryingLimit();
+	CheckRoundFunction();
 }
 
 void FOpenGLDynamicRHI::Shutdown()
diff --git a/Engine/Source/Runtime/OpenGLDrv/Private/OpenGLES2.cpp b/Engine/Source/Runtime/OpenGLDrv/Private/OpenGLES2.cpp
--- a/Engine/Source/Runtime/OpenGLDrv/Private/OpenGLES2.cpp
+++ b/Engine/Source/Runtime/OpenGLDrv/Private/OpenGLES2.cpp
@@ -145,6 +145,9 @@ bool FOpenGLES2::bNeedsVertexAttribRemap = false;
 /* This hack fixes an issue with SGX540 compiler which can get upset with some operations that mix highp and mediump */
 bool FOpenGLES2::bRequiresTexture2DPrecisionHack = false;
 
+/* This is a hack to add a round() function when not available to a shader compiler */
+bool FOpenGLES2::bRequiresRoundFunctionHack = true;
+
 /* This is to avoid a bug in Adreno drivers that define GL_ARM_shader_framebuffer_fetch_depth_stencil even when device does not support this extension  */
 bool FOpenGLES2::bRequiresARMShaderFramebufferFetchDepthStencilUndef = false;
 
diff --git a/Engine/Source/Runtime/OpenGLDrv/Private/OpenGLESDeferred.cpp b/Engine/Source/Runtime/OpenGLDrv/Private/OpenGLESDeferred.cpp
--- a/Engine/Source/Runtime/OpenGLDrv/Private/OpenGLESDeferred.cpp
+++ b/Engine/Source/Runtime/OpenGLDrv/Private/OpenGLESDeferred.cpp
@@ -152,6 +152,9 @@ bool FOpenGLESDeferred::bRequiresGLFragCoordVaryingLimitHack = false;
 /* This hack fixes an issue with SGX540 compiler which can get upset with some operations that mix highp and mediump */
 bool FOpenGLESDeferred::bRequiresTexture2DPrecisionHack = false;
 
+/* This is a hack to add a round() function when not available to a shader compiler */
+bool FOpenGLES2Deferred::bRequiresRoundFunctionHack = true;
+
 /* This is to avoid a bug in Adreno drivers that define GL_ARM_shader_framebuffer_fetch_depth_stencil even when device does not support this extension  */
 bool FOpenGLESDeferred::bRequiresARMShaderFramebufferFetchDepthStencilUndef = false;
 
diff --git a/Engine/Source/Runtime/OpenGLDrv/Private/OpenGLRenderTarget.cpp b/Engine/Source/Runtime/OpenGLDrv/Private/OpenGLRenderTarget.cpp
--- a/Engine/Source/Runtime/OpenGLDrv/Private/OpenGLRenderTarget.cpp
+++ b/Engine/Source/Runtime/OpenGLDrv/Private/OpenGLRenderTarget.cpp
@@ -752,6 +752,24 @@ void FOpenGLDynamicRHI::ReadSurfaceDataRaw(FOpenGLContextState& ContextState, FT
 		if ( FOpenGL::SupportsBGRA8888() )
 		{
 			glReadPixels(Rect.Min.X, Rect.Min.Y, SizeX, SizeY, GL_BGRA, GL_FLOAT, FloatBGRAData );
+			GLenum Error = glGetError();
+			if (Error != GL_NO_ERROR)
+			{
+				glReadPixels(Rect.Min.X, Rect.Min.Y, SizeX, SizeY, GL_RGBA, GL_FLOAT, FloatBGRAData );
+				Error = glGetError();
+				if (Error == GL_NO_ERROR)
+				{
+					float* FloatDataPtr = FloatBGRAData;
+					float* FloatDataPtrEnd = FloatBGRAData + PixelComponentCount;
+					while (FloatDataPtr != FloatDataPtrEnd)
+					{
+						float Temp = FloatDataPtr[0];
+						FloatDataPtr[0] = FloatDataPtr[2];
+						FloatDataPtr[2] = Temp;
+						FloatDataPtr += 4;
+					}
+				}
+			}
 		}
 		else 
 		{
diff --git a/Engine/Source/Runtime/OpenGLDrv/Private/OpenGLShaders.cpp b/Engine/Source/Runtime/OpenGLDrv/Private/OpenGLShaders.cpp
--- a/Engine/Source/Runtime/OpenGLDrv/Private/OpenGLShaders.cpp
+++ b/Engine/Source/Runtime/OpenGLDrv/Private/OpenGLShaders.cpp
@@ -663,6 +663,7 @@ void OPENGLDRV_API GetCurrentOpenGLShaderDeviceCapabilities(FOpenGLShaderDeviceC
 	Capabilities.bRequiresGLFragCoordVaryingLimitHack = FOpenGL::RequiresGLFragCoordVaryingLimitHack();
 	Capabilities.MaxVaryingVectors = FOpenGL::GetMaxVaryingVectors();
 	Capabilities.bRequiresTexture2DPrecisionHack = FOpenGL::RequiresTexture2DPrecisionHack();
+	Capabilities.bRequiresRoundFunctionHack = FOpenGL::RequiresRoundFunctionHack();
 #elif PLATFORM_HTML5
 	Capabilities.TargetPlatform = EOpenGLShaderTargetPlatform::OGLSTP_HTML5;
 	Capabilities.bUseES30ShadingLanguage = FOpenGL::UseES30ShadingLanguage();
@@ -970,6 +971,54 @@ void OPENGLDRV_API GLSLToDeviceCompatibleGLSL(FAnsiCharArray& GlslCodeOriginal,
 							"#define textureCubeLodEXT textureCubeLod \n");
 					}
 
+					if (Capabilities.bRequiresRoundFunctionHack)
+					{
+						const bool bIsMediumPrecision = (FCStringAnsi::Strstr(GlslCodeOriginal.GetData(), "precision mediump float;") != nullptr);
+
+						if (!bIsMediumPrecision)
+						{
+							AppendCString(GlslCode,
+								"highp float round(highp float value)\n"
+								"{\n"
+								"	return floor(value + 0.5);\n"
+								"}\n"
+								"highp vec2 round(highp vec2 value)\n"
+								"{\n"
+								"	return floor(value + vec2(0.5, 0.5));\n"
+								"}\n"
+								"highp vec3 round(highp vec3 value)\n"
+								"{\n"
+								"	return floor(value + vec3(0.5, 0.5, 0.5));\n"
+								"}\n"
+								"highp vec4 round(highp vec4 value)\n"
+								"{\n"
+								"	return floor(value + vec4(0.5, 0.5, 0.5, 0.5));\n"
+								"}\n"
+							);
+						}
+						else
+						{
+							AppendCString(GlslCode,
+								"mediump float round(mediump float value)\n"
+								"{\n"
+								"	return floor(value + 0.5);\n"
+								"}\n"
+								"mediump vec2 round(mediump vec2 value)\n"
+								"{\n"
+								"	return floor(value + vec2(0.5, 0.5));\n"
+								"}\n"
+								"mediump vec3 round(mediump vec3 value)\n"
+								"{\n"
+								"	return floor(value + vec3(0.5, 0.5, 0.5));\n"
+								"}\n"
+								"mediump vec4 round(mediump vec4 value)\n"
+								"{\n"
+								"	return floor(value + vec4(0.5, 0.5, 0.5, 0.5));\n"
+								"}\n"
+							);
+						}
+					}
+
 					// Deal with gl_FragCoord using one of the varying vectors and shader possibly exceeding the limit
 					if (Capabilities.bRequiresGLFragCoordVaryingLimitHack)
 					{
diff --git a/Engine/Source/Runtime/OpenGLDrv/Public/OpenGLES2.h b/Engine/Source/Runtime/OpenGLDrv/Public/OpenGLES2.h
--- a/Engine/Source/Runtime/OpenGLDrv/Public/OpenGLES2.h
+++ b/Engine/Source/Runtime/OpenGLDrv/Public/OpenGLES2.h
@@ -145,6 +145,7 @@ struct FOpenGLES2 : public FOpenGLBase
 	static FORCEINLINE bool SupportsStandardDerivativesExtension()		{ return bSupportsStandardDerivativesExtension; }
 	static FORCEINLINE bool RequiresGLFragCoordVaryingLimitHack()		{ return bRequiresGLFragCoordVaryingLimitHack; }
 	static FORCEINLINE bool RequiresTexture2DPrecisionHack()			{ return bRequiresTexture2DPrecisionHack; }
+	static FORCEINLINE bool RequiresRoundFunctionHack()					{ return bRequiresRoundFunctionHack; }
 	static FORCEINLINE bool RequiresARMShaderFramebufferFetchDepthStencilUndef() { return bRequiresARMShaderFramebufferFetchDepthStencilUndef; }
 	static FORCEINLINE bool IsCheckingShaderCompilerHacks()				{ return bIsCheckingShaderCompilerHacks; }
 	static FORCEINLINE bool IsLimitingShaderCompileCount()				{ return bIsLimitingShaderCompileCount; }
@@ -536,6 +537,9 @@ public:
 	/* This hack fixes an issue with SGX540 compiler which can get upset with some operations that mix highp and mediump */
 	static bool bRequiresTexture2DPrecisionHack;
 
+	/* This is a hack to add a round() function when not available to a shader compiler */
+	static bool bRequiresRoundFunctionHack;
+
 	/* This is to avoid a bug in Adreno drivers that define GL_ARM_shader_framebuffer_fetch_depth_stencil even when device does not support this extension  */
 	static bool bRequiresARMShaderFramebufferFetchDepthStencilUndef;
 
diff --git a/Engine/Source/Runtime/OpenGLDrv/Public/OpenGLESDeferred.h b/Engine/Source/Runtime/OpenGLDrv/Public/OpenGLESDeferred.h
--- a/Engine/Source/Runtime/OpenGLDrv/Public/OpenGLESDeferred.h
+++ b/Engine/Source/Runtime/OpenGLDrv/Public/OpenGLESDeferred.h
@@ -105,6 +105,7 @@ struct FOpenGLESDeferred : public FOpenGLBase
 	static FORCEINLINE bool RequiresGLFragCoordVaryingLimitHack()		{ return bRequiresGLFragCoordVaryingLimitHack; }
 	static FORCEINLINE GLenum GetVertexHalfFloatFormat()				{ return bES2Fallback ? GL_HALF_FLOAT_OES : GL_HALF_FLOAT; }
 	static FORCEINLINE bool RequiresTexture2DPrecisionHack()			{ return bRequiresTexture2DPrecisionHack; }
+	static FORCEINLINE bool RequiresRoundFunctionHack()					{ return bRequiresRoundFunctionHack; }
 	static FORCEINLINE bool RequiresARMShaderFramebufferFetchDepthStencilUndef() { return bRequiresARMShaderFramebufferFetchDepthStencilUndef; }
 	static FORCEINLINE bool IsCheckingShaderCompilerHacks()				{ return bIsCheckingShaderCompilerHacks; }
 	static FORCEINLINE bool SupportsRGB10A2()							{ return bSupportsRGB10A2 || !bES2Fallback; }
@@ -977,6 +978,9 @@ public:
 	/* This hack fixes an issue with SGX540 compiler which can get upset with some operations that mix highp and mediump */
 	static bool bRequiresTexture2DPrecisionHack;
 
+	/* This is a hack to add a round() function when not available to a shader compiler */
+	static bool bRequiresRoundFunctionHack;
+
 	/* This is to avoid a bug in Adreno drivers that define GL_ARM_shader_framebuffer_fetch_depth_stencil even when device does not support this extension  */
 	static bool bRequiresARMShaderFramebufferFetchDepthStencilUndef;
 
diff --git a/Engine/Source/Runtime/OpenGLDrv/Public/OpenGLShaders.h b/Engine/Source/Runtime/OpenGLDrv/Public/OpenGLShaders.h
--- a/Engine/Source/Runtime/OpenGLDrv/Public/OpenGLShaders.h
+++ b/Engine/Source/Runtime/OpenGLDrv/Public/OpenGLShaders.h
@@ -41,6 +41,7 @@ struct FOpenGLShaderDeviceCapabilities
 	bool bRequiresGLFragCoordVaryingLimitHack;
 	GLint MaxVaryingVectors;
 	bool bRequiresTexture2DPrecisionHack;
+	bool bRequiresRoundFunctionHack;
 };
 
 /**
diff --git a/Engine/Source/Runtime/PakFile/Private/IPlatformFilePak.cpp b/Engine/Source/Runtime/PakFile/Private/IPlatformFilePak.cpp
--- a/Engine/Source/Runtime/PakFile/Private/IPlatformFilePak.cpp
+++ b/Engine/Source/Runtime/PakFile/Private/IPlatformFilePak.cpp
@@ -3439,6 +3439,20 @@ void FAsyncIOCPUWorkTask::DoTask(ENamedThreads::Type CurrentThread, const FGraph
 #endif  
 
 
+#if PAK_TRACKER
+TMap<FString, int32> FPakPlatformFile::GPakSizeMap;
+
+void FPakPlatformFile::TrackPak(const TCHAR* Filename, const FPakEntry* PakEntry)
+{
+	FString Key(Filename);
+
+	if(!GPakSizeMap.Find(Key))
+	{
+		GPakSizeMap.Add(Key, PakEntry->Size);
+	}
+}
+#endif
+
 IAsyncReadFileHandle* FPakPlatformFile::OpenAsyncRead(const TCHAR* Filename)
 {
 	check(GConfig);
@@ -3449,6 +3463,10 @@ IAsyncReadFileHandle* FPakPlatformFile::OpenAsyncRead(const TCHAR* Filename)
 		const FPakEntry* FileEntry = FindFileInPakFiles(Filename, &PakFile);
 		if (FileEntry && PakFile && PakFile->GetFilenameName() != NAME_None)
 		{
+#if PAK_TRACKER
+			TrackPak(Filename, FileEntry);
+#endif
+			
 			return new FPakAsyncReadFileHandle(FileEntry, PakFile, Filename);
 		}
 		if (FString(Filename).Contains(TEXT("/Saved/PakFileTest/")))
@@ -4501,6 +4519,10 @@ IFileHandle* FPakPlatformFile::OpenRead(const TCHAR* Filename, bool bAllowWrite)
 	const FPakEntry* FileEntry = FindFileInPakFiles(Filename, &PakFile);	
 	if (FileEntry != NULL)
 	{
+#if PAK_TRACKER
+		TrackPak(Filename, FileEntry);
+#endif
+		
 		Result = CreatePakFileHandle(Filename, PakFile, FileEntry);
 	}
 	else
diff --git a/Engine/Source/Runtime/PakFile/Public/IPlatformFilePak.h b/Engine/Source/Runtime/PakFile/Public/IPlatformFilePak.h
--- a/Engine/Source/Runtime/PakFile/Public/IPlatformFilePak.h
+++ b/Engine/Source/Runtime/PakFile/Public/IPlatformFilePak.h
@@ -22,6 +22,7 @@ DECLARE_DWORD_ACCUMULATOR_STAT_EXTERN(TEXT("Num open pak file handles"), STAT_Pa
 DECLARE_DELEGATE_RetVal_OneParam(bool, FFilenameSecurityDelegate, const TCHAR* /*InFilename*/);
 
 #define PAKHASH_USE_CRC	1
+#define PAK_TRACKER 0
 #define PAK_SIGNATURE_CHECK_FAILS_ARE_FATAL 0
 
 #if PAKHASH_USE_CRC
@@ -1749,6 +1750,12 @@ public:
 	void HandlePakCorruptCommand(const TCHAR* Cmd, FOutputDevice& Ar);
 #endif
 	// END Console commands
+	
+#if PAK_TRACKER
+	static TMap<FString, int32> GPakSizeMap;
+	static void TrackPak(const TCHAR* Filename, const FPakEntry* PakEntry);
+	static TMap<FString, int32>& GetPakMap() { return GPakSizeMap; }
+#endif
 };
 
 
diff --git a/Engine/Source/Runtime/Renderer/Private/IndirectLightingCache.cpp b/Engine/Source/Runtime/Renderer/Private/IndirectLightingCache.cpp
--- a/Engine/Source/Runtime/Renderer/Private/IndirectLightingCache.cpp
+++ b/Engine/Source/Runtime/Renderer/Private/IndirectLightingCache.cpp
@@ -907,6 +907,9 @@ void FIndirectLightingCache::UpdateBlock(FScene* Scene, FViewInfo* DebugDrawingV
 	BlockInfo.Block.bHasEverBeenUpdated = true;
 }
 
+template class TSHVector<2>;
+template class TSHVector<3>;
+
 template <int32 SHOrder>
 static void ReduceSHRinging(TSHVectorRGB<SHOrder>& IncidentRadiance)
 {
diff --git a/Engine/Source/Runtime/Renderer/Private/MobileSceneCaptureRendering.cpp b/Engine/Source/Runtime/Renderer/Private/MobileSceneCaptureRendering.cpp
--- a/Engine/Source/Runtime/Renderer/Private/MobileSceneCaptureRendering.cpp
+++ b/Engine/Source/Runtime/Renderer/Private/MobileSceneCaptureRendering.cpp
@@ -372,7 +372,7 @@ void UpdateSceneCaptureContentMobile_RenderThread(
 			auto& RenderTargetRHI = Target->GetRenderTargetTexture();
 			FPooledRenderTargetDesc Desc(FPooledRenderTargetDesc::Create2DDesc(Target->GetSizeXY(),
 				RenderTargetRHI.GetReference()->GetFormat(),
-				FClearValueBinding::None,
+				RenderTargetRHI.GetReference()->GetClearBinding(),
 				TexCreate_None,
 				TexCreate_RenderTargetable,
 				false));
diff --git a/Engine/Source/Runtime/ShaderCore/Public/VertexFactory.h b/Engine/Source/Runtime/ShaderCore/Public/VertexFactory.h
--- a/Engine/Source/Runtime/ShaderCore/Public/VertexFactory.h
+++ b/Engine/Source/Runtime/ShaderCore/Public/VertexFactory.h
@@ -284,7 +284,7 @@ extern SHADERCORE_API FVertexFactoryType* FindVertexFactoryType(FName TypeName);
 #define DECLARE_VERTEX_FACTORY_TYPE(FactoryClass) \
 	public: \
 	static FVertexFactoryType StaticType; \
-	virtual FVertexFactoryType* GetType() const override { return &StaticType; }
+	virtual FVertexFactoryType* GetType() const override;
 
 /**
  * A macro for implementing the static vertex factory type object, and specifying parameters used by the type.
@@ -305,7 +305,8 @@ extern SHADERCORE_API FVertexFactoryType* FindVertexFactoryType(FName TypeName);
 		FactoryClass::ShouldCache, \
 		FactoryClass::ModifyCompilationEnvironment, \
 		FactoryClass::SupportsTessellationShaders \
-		);
+		); \
+		FVertexFactoryType* FactoryClass::GetType() const { return &StaticType; }
 
 /** Encapsulates a dependency on a vertex factory type and saved state from that vertex factory type. */
 class FVertexFactoryTypeDependency
diff --git a/Engine/Source/Runtime/Slate/Private/Framework/Application/SlateApplication.cpp b/Engine/Source/Runtime/Slate/Private/Framework/Application/SlateApplication.cpp
--- a/Engine/Source/Runtime/Slate/Private/Framework/Application/SlateApplication.cpp
+++ b/Engine/Source/Runtime/Slate/Private/Framework/Application/SlateApplication.cpp
@@ -6366,6 +6366,11 @@ bool FSlateApplication::OnSizeChanged( const TSharedRef< FGenericWindow >& Platf
 
 		Renderer->RequestResize( Window, Width, Height );
 
+		if (FPlatformProperties::HasFixedResolution())
+		{
+			Renderer->SetSystemResolution(Width, Height);
+		}
+
 		if ( !bWasMinimized && Window->IsRegularWindow() && !Window->HasOSWindowBorder() && Window->IsVisible() && Window->IsDrawingEnabled() )
 		{
 			PrivateDrawWindows( Window );
diff --git a/Engine/Source/Runtime/Slate/Private/Widgets/Input/SEditableText.cpp b/Engine/Source/Runtime/Slate/Private/Widgets/Input/SEditableText.cpp
--- a/Engine/Source/Runtime/Slate/Private/Widgets/Input/SEditableText.cpp
+++ b/Engine/Source/Runtime/Slate/Private/Widgets/Input/SEditableText.cpp
@@ -326,6 +326,11 @@ void SEditableText::SetAllowContextMenu(const TAttribute< bool >& InAllowContext
 	bAllowContextMenu = InAllowContextMenu;
 }
 
+void SEditableText::SetVirtualKeyboardDismissAction(TAttribute< EVirtualKeyboardDismissAction > InVirtualKeyboardDismissAction)
+{
+	VirtualKeyboardDismissAction = InVirtualKeyboardDismissAction;
+}
+
 void SEditableText::SetTextShapingMethod(const TOptional<ETextShapingMethod>& InTextShapingMethod)
 {
 	EditableTextLayout->SetTextShapingMethod(InTextShapingMethod);
diff --git a/Engine/Source/Runtime/Slate/Private/Widgets/Input/SEditableTextBox.cpp b/Engine/Source/Runtime/Slate/Private/Widgets/Input/SEditableTextBox.cpp
--- a/Engine/Source/Runtime/Slate/Private/Widgets/Input/SEditableTextBox.cpp
+++ b/Engine/Source/Runtime/Slate/Private/Widgets/Input/SEditableTextBox.cpp
@@ -383,3 +383,8 @@ void SEditableTextBox::SetAllowContextMenu(TAttribute<bool> InAllowContextMenu)
 {
 	EditableText->SetAllowContextMenu(InAllowContextMenu);
 }
+
+void SEditableTextBox::SetVirtualKeyboardDismissAction(TAttribute<EVirtualKeyboardDismissAction> InVirtualKeyboardDismissAction)
+{
+	EditableText->SetVirtualKeyboardDismissAction(InVirtualKeyboardDismissAction);
+}
diff --git a/Engine/Source/Runtime/Slate/Private/Widgets/Input/SMultiLineEditableTextBox.cpp b/Engine/Source/Runtime/Slate/Private/Widgets/Input/SMultiLineEditableTextBox.cpp
--- a/Engine/Source/Runtime/Slate/Private/Widgets/Input/SMultiLineEditableTextBox.cpp
+++ b/Engine/Source/Runtime/Slate/Private/Widgets/Input/SMultiLineEditableTextBox.cpp
@@ -297,6 +297,11 @@ void SMultiLineEditableTextBox::SetAllowContextMenu(const TAttribute< bool >& In
 	EditableText->SetAllowContextMenu(InAllowContextMenu);
 }
 
+void SMultiLineEditableTextBox::SetVirtualKeyboardDismissAction(TAttribute<EVirtualKeyboardDismissAction> InVirtualKeyboardDismissAction)
+{
+	EditableText->SetVirtualKeyboardDismissAction(InVirtualKeyboardDismissAction);
+}
+
 void SMultiLineEditableTextBox::SetIsReadOnly(const TAttribute<bool>& InIsReadOnly)
 {
 	EditableText->SetIsReadOnly(InIsReadOnly);
diff --git a/Engine/Source/Runtime/Slate/Private/Widgets/Text/SMultiLineEditableText.cpp b/Engine/Source/Runtime/Slate/Private/Widgets/Text/SMultiLineEditableText.cpp
--- a/Engine/Source/Runtime/Slate/Private/Widgets/Text/SMultiLineEditableText.cpp
+++ b/Engine/Source/Runtime/Slate/Private/Widgets/Text/SMultiLineEditableText.cpp
@@ -186,6 +186,11 @@ void SMultiLineEditableText::SetAllowContextMenu(const TAttribute< bool >& InAll
 	bAllowContextMenu = InAllowContextMenu;
 }
 
+void SMultiLineEditableText::SetVirtualKeyboardDismissAction(TAttribute< EVirtualKeyboardDismissAction > InVirtualKeyboardDismissAction)
+{
+	VirtualKeyboardDismissAction = InVirtualKeyboardDismissAction;
+}
+
 void SMultiLineEditableText::SetIsReadOnly(const TAttribute<bool>& InIsReadOnly)
 {
 	bIsReadOnly = InIsReadOnly;
diff --git a/Engine/Source/Runtime/Slate/Public/Widgets/Input/SEditableText.h b/Engine/Source/Runtime/Slate/Public/Widgets/Input/SEditableText.h
--- a/Engine/Source/Runtime/Slate/Public/Widgets/Input/SEditableText.h
+++ b/Engine/Source/Runtime/Slate/Public/Widgets/Input/SEditableText.h
@@ -214,6 +214,9 @@ public:
 	/** See the AllowContextMenu attribute */
 	void SetAllowContextMenu(const TAttribute< bool >& InAllowContextMenu);
 
+	/** Set the VirtualKeyboardDismissAction attribute */
+	void SetVirtualKeyboardDismissAction(TAttribute< EVirtualKeyboardDismissAction > InVirtualKeyboardDismissAction);
+
 	/**
 	 * Sets the font used to draw the text
 	 *
diff --git a/Engine/Source/Runtime/Slate/Public/Widgets/Input/SEditableTextBox.h b/Engine/Source/Runtime/Slate/Public/Widgets/Input/SEditableTextBox.h
--- a/Engine/Source/Runtime/Slate/Public/Widgets/Input/SEditableTextBox.h
+++ b/Engine/Source/Runtime/Slate/Public/Widgets/Input/SEditableTextBox.h
@@ -187,6 +187,9 @@ public:
 	/** See the AllowContextMenu attribute */
 	void SetAllowContextMenu(TAttribute< bool > InAllowContextMenu);
 
+	/** Set the VirtualKeyboardDismissAction attribute */
+	void SetVirtualKeyboardDismissAction(TAttribute< EVirtualKeyboardDismissAction > InVirtualKeyboardDismissAction);
+	
 	/**
 	 * Sets the font used to draw the text
 	 *
diff --git a/Engine/Source/Runtime/Slate/Public/Widgets/Input/SMultiLineEditableTextBox.h b/Engine/Source/Runtime/Slate/Public/Widgets/Input/SMultiLineEditableTextBox.h
--- a/Engine/Source/Runtime/Slate/Public/Widgets/Input/SMultiLineEditableTextBox.h
+++ b/Engine/Source/Runtime/Slate/Public/Widgets/Input/SMultiLineEditableTextBox.h
@@ -313,6 +313,9 @@ public:
 	/** See the AllowContextMenu attribute */
 	void SetAllowContextMenu(const TAttribute< bool >& InAllowContextMenu);
 
+	/** Set the VirtualKeyboardDismissAction attribute */
+	void SetVirtualKeyboardDismissAction(TAttribute< EVirtualKeyboardDismissAction > InVirtualKeyboardDismissAction);
+	
 	/** Set the ReadOnly attribute */
 	void SetIsReadOnly(const TAttribute< bool >& InIsReadOnly);
 
diff --git a/Engine/Source/Runtime/Slate/Public/Widgets/Text/ISlateEditableTextWidget.h b/Engine/Source/Runtime/Slate/Public/Widgets/Text/ISlateEditableTextWidget.h
--- a/Engine/Source/Runtime/Slate/Public/Widgets/Text/ISlateEditableTextWidget.h
+++ b/Engine/Source/Runtime/Slate/Public/Widgets/Text/ISlateEditableTextWidget.h
@@ -58,7 +58,8 @@ enum class EVirtualKeyboardTrigger
 	OnAllFocusEvents,
 };
 
-enum class EVirtualKeyboardDismissAction
+UENUM(BlueprintType)
+enum class EVirtualKeyboardDismissAction : uint8
 {
 	/** Sends a text changed message when the virtual keyboard is dismissed by the user. */
 	TextChangeOnDismiss,
diff --git a/Engine/Source/Runtime/Slate/Public/Widgets/Text/SMultiLineEditableText.h b/Engine/Source/Runtime/Slate/Public/Widgets/Text/SMultiLineEditableText.h
--- a/Engine/Source/Runtime/Slate/Public/Widgets/Text/SMultiLineEditableText.h
+++ b/Engine/Source/Runtime/Slate/Public/Widgets/Text/SMultiLineEditableText.h
@@ -252,6 +252,9 @@ public:
 	/** See the AllowContextMenu attribute */
 	void SetAllowContextMenu(const TAttribute< bool >& InAllowContextMenu);
 
+	/** Set the VirtualKeyboardDismissAction attribute */
+	void SetVirtualKeyboardDismissAction(TAttribute< EVirtualKeyboardDismissAction > InVirtualKeyboardDismissAction);
+	
 	/** Sets the ReadOnly attribute */
 	void SetIsReadOnly(const TAttribute< bool >& InIsReadOnly);
 
diff --git a/Engine/Source/Runtime/SlateCore/Public/Rendering/SlateRenderer.h b/Engine/Source/Runtime/SlateCore/Public/Rendering/SlateRenderer.h
--- a/Engine/Source/Runtime/SlateCore/Public/Rendering/SlateRenderer.h
+++ b/Engine/Source/Runtime/SlateCore/Public/Rendering/SlateRenderer.h
@@ -182,6 +182,16 @@ public:
 	 */
 	virtual void UpdateFullscreenState( const TSharedRef<SWindow> InWindow, uint32 OverrideResX = 0, uint32 OverrideResY = 0 ) = 0;
 
+
+	/**
+	 * Set the resolution cached by the engine
+	 *
+	 * @param	Width			Width of the system resolution
+	 * @param	Height			Height of the system resolution
+	 */
+	virtual void SetSystemResolution(uint32 Width, uint32 Height) = 0;
+	
+	
 	/**
 	 * Restore the given window to the resolution settings currently cached by the engine
 	 * 
diff --git a/Engine/Source/Runtime/SlateNullRenderer/Private/SlateNullRenderer.h b/Engine/Source/Runtime/SlateNullRenderer/Private/SlateNullRenderer.h
--- a/Engine/Source/Runtime/SlateNullRenderer/Private/SlateNullRenderer.h
+++ b/Engine/Source/Runtime/SlateNullRenderer/Private/SlateNullRenderer.h
@@ -25,6 +25,7 @@ public:
 	virtual void RequestResize( const TSharedPtr<SWindow>& Window, uint32 NewWidth, uint32 NewHeight ) override;
 	virtual void CreateViewport( const TSharedRef<SWindow> Window ) override;
 	virtual void UpdateFullscreenState( const TSharedRef<SWindow> Window, uint32 OverrideResX, uint32 OverrideResY ) override;
+	virtual void SetSystemResolution(uint32 Width, uint32 Height) override {}
 	virtual void RestoreSystemResolution(const TSharedRef<SWindow> InWindow) override;
 	virtual void DrawWindows( FSlateDrawBuffer& InWindowDrawBuffer ) override;
 	virtual void ReleaseDynamicResource( const FSlateBrush& InBrush ) override;
diff --git a/Engine/Source/Runtime/SlateRHIRenderer/Private/SlateRHIRenderer.cpp b/Engine/Source/Runtime/SlateRHIRenderer/Private/SlateRHIRenderer.cpp
--- a/Engine/Source/Runtime/SlateRHIRenderer/Private/SlateRHIRenderer.cpp
+++ b/Engine/Source/Runtime/SlateRHIRenderer/Private/SlateRHIRenderer.cpp
@@ -402,6 +402,12 @@ void FSlateRHIRenderer::UpdateFullscreenState( const TSharedRef<SWindow> Window,
 	}
 }
 
+void FSlateRHIRenderer::SetSystemResolution(uint32 Width, uint32 Height)
+{
+	FSystemResolution::RequestResolutionChange(Width, Height, EWindowMode::Fullscreen);
+	IConsoleManager::Get().CallAllConsoleVariableSinks();
+}
+
 void FSlateRHIRenderer::RestoreSystemResolution(const TSharedRef<SWindow> InWindow)
 {
 	if (!GIsEditor && InWindow->GetWindowMode() == EWindowMode::Fullscreen)
diff --git a/Engine/Source/Runtime/SlateRHIRenderer/Private/SlateRHIRenderer.h b/Engine/Source/Runtime/SlateRHIRenderer/Private/SlateRHIRenderer.h
--- a/Engine/Source/Runtime/SlateRHIRenderer/Private/SlateRHIRenderer.h
+++ b/Engine/Source/Runtime/SlateRHIRenderer/Private/SlateRHIRenderer.h
@@ -175,6 +175,7 @@ public:
 	virtual void RequestResize( const TSharedPtr<SWindow>& Window, uint32 NewWidth, uint32 NewHeight ) override;
 	virtual void CreateViewport( const TSharedRef<SWindow> Window ) override;
 	virtual void UpdateFullscreenState( const TSharedRef<SWindow> Window, uint32 OverrideResX, uint32 OverrideResY ) override;
+	virtual void SetSystemResolution(uint32 Width, uint32 Height) override;
 	virtual void RestoreSystemResolution(const TSharedRef<SWindow> InWindow) override;
 	virtual void DrawWindows( FSlateDrawBuffer& InWindowDrawBuffer ) override;
 	virtual void FlushCommands() const override;
diff --git a/Engine/Source/Runtime/UMG/Private/Components/EditableText.cpp b/Engine/Source/Runtime/UMG/Private/Components/EditableText.cpp
--- a/Engine/Source/Runtime/UMG/Private/Components/EditableText.cpp
+++ b/Engine/Source/Runtime/UMG/Private/Components/EditableText.cpp
@@ -36,6 +36,7 @@ UEditableText::UEditableText(const FObjectInitializer& ObjectInitializer)
 	ClearKeyboardFocusOnCommit = Defaults._ClearKeyboardFocusOnCommit.Get();
 	SelectAllTextOnCommit = Defaults._SelectAllTextOnCommit.Get();
 	AllowContextMenu = Defaults._AllowContextMenu.Get();
+	VirtualKeyboardDismissAction = Defaults._VirtualKeyboardDismissAction.Get();
 	Clipping = Defaults._Clipping;
 }
 
@@ -59,6 +60,7 @@ TSharedRef<SWidget> UEditableText::RebuildWidget()
 		.OnTextChanged( BIND_UOBJECT_DELEGATE( FOnTextChanged, HandleOnTextChanged ) )
 		.OnTextCommitted( BIND_UOBJECT_DELEGATE( FOnTextCommitted, HandleOnTextCommitted ) )
 		.VirtualKeyboardType( EVirtualKeyboardType::AsKeyboardType( KeyboardType.GetValue() ) )
+		.VirtualKeyboardDismissAction(VirtualKeyboardDismissAction)
 		.Justification( Justification );
 	
 	return MyEditableText.ToSharedRef();
@@ -76,6 +78,7 @@ void UEditableText::SynchronizeProperties()
 	MyEditableText->SetIsReadOnly(IsReadOnly);
 	MyEditableText->SetIsPassword(IsPassword);
 	MyEditableText->SetAllowContextMenu(AllowContextMenu);
+	MyEditableText->SetVirtualKeyboardDismissAction(VirtualKeyboardDismissAction);
 	MyEditableText->SetJustification(Justification);
 	// TODO UMG Complete making all properties settable on SEditableText
 
diff --git a/Engine/Source/Runtime/UMG/Private/Components/EditableTextBox.cpp b/Engine/Source/Runtime/UMG/Private/Components/EditableTextBox.cpp
--- a/Engine/Source/Runtime/UMG/Private/Components/EditableTextBox.cpp
+++ b/Engine/Source/Runtime/UMG/Private/Components/EditableTextBox.cpp
@@ -36,6 +36,7 @@ UEditableTextBox::UEditableTextBox(const FObjectInitializer& ObjectInitializer)
 	ClearKeyboardFocusOnCommit = Defaults._ClearKeyboardFocusOnCommit.Get();
 	SelectAllTextOnCommit = Defaults._SelectAllTextOnCommit.Get();
 	AllowContextMenu = Defaults._AllowContextMenu.Get();
+	VirtualKeyboardDismissAction = Defaults._VirtualKeyboardDismissAction.Get();
 
 	WidgetStyle = *Defaults._Style;
 }
@@ -61,6 +62,7 @@ TSharedRef<SWidget> UEditableTextBox::RebuildWidget()
 		.OnTextChanged(BIND_UOBJECT_DELEGATE(FOnTextChanged, HandleOnTextChanged))
 		.OnTextCommitted(BIND_UOBJECT_DELEGATE(FOnTextCommitted, HandleOnTextCommitted))
 		.VirtualKeyboardType(EVirtualKeyboardType::AsKeyboardType(KeyboardType.GetValue()))
+		.VirtualKeyboardDismissAction(VirtualKeyboardDismissAction)
 		.Justification(Justification);
 
 	return MyEditableTextBlock.ToSharedRef();
@@ -85,6 +87,7 @@ void UEditableTextBox::SynchronizeProperties()
 	MyEditableTextBlock->SetClearKeyboardFocusOnCommit(ClearKeyboardFocusOnCommit);
 	MyEditableTextBlock->SetSelectAllTextOnCommit(SelectAllTextOnCommit);
 	MyEditableTextBlock->SetAllowContextMenu(AllowContextMenu);
+	MyEditableTextBlock->SetVirtualKeyboardDismissAction(VirtualKeyboardDismissAction);
 	MyEditableTextBlock->SetJustification(Justification);
 
 	ShapedTextOptions.SynchronizeShapedTextProperties(*MyEditableTextBlock);
diff --git a/Engine/Source/Runtime/UMG/Private/Components/MultiLineEditableText.cpp b/Engine/Source/Runtime/UMG/Private/Components/MultiLineEditableText.cpp
--- a/Engine/Source/Runtime/UMG/Private/Components/MultiLineEditableText.cpp
+++ b/Engine/Source/Runtime/UMG/Private/Components/MultiLineEditableText.cpp
@@ -23,6 +23,7 @@ UMultiLineEditableText::UMultiLineEditableText(const FObjectInitializer& ObjectI
 	ClearKeyboardFocusOnCommit = Defaults._ClearKeyboardFocusOnCommit.Get();
 	AllowContextMenu = Defaults._AllowContextMenu.Get();
 	Clipping = Defaults._Clipping;
+	VirtualKeyboardDismissAction = Defaults._VirtualKeyboardDismissAction.Get();
 	AutoWrapText = true;
 	
 	if (!IsRunningDedicatedServer())
@@ -58,6 +59,7 @@ TSharedRef<SWidget> UMultiLineEditableText::RebuildWidget()
 //	.BackgroundImageSelectionTarget(BackgroundImageSelectionTarget ? TAttribute<const FSlateBrush*>(&BackgroundImageSelectionTarget->Brush) : TAttribute<const FSlateBrush*>())
 //	.BackgroundImageComposing(BackgroundImageComposing ? TAttribute<const FSlateBrush*>(&BackgroundImageComposing->Brush) : TAttribute<const FSlateBrush*>())
 //	.CaretImage(CaretImage ? TAttribute<const FSlateBrush*>(&CaretImage->Brush) : TAttribute<const FSlateBrush*>())
+	.VirtualKeyboardDismissAction(VirtualKeyboardDismissAction)
 	.OnTextChanged(BIND_UOBJECT_DELEGATE(FOnTextChanged, HandleOnTextChanged))
 	.OnTextCommitted(BIND_UOBJECT_DELEGATE(FOnTextCommitted, HandleOnTextCommitted))
 	;
@@ -76,6 +78,7 @@ void UMultiLineEditableText::SynchronizeProperties()
 	MyMultiLineEditableText->SetHintText(HintTextBinding);
 	MyMultiLineEditableText->SetAllowContextMenu(AllowContextMenu);
 	MyMultiLineEditableText->SetIsReadOnly(bIsReadOnly);
+	MyMultiLineEditableText->SetVirtualKeyboardDismissAction(VirtualKeyboardDismissAction);
 	MyMultiLineEditableText->SetSelectAllTextWhenFocused(SelectAllTextWhenFocused);
 	MyMultiLineEditableText->SetClearTextSelectionOnFocusLoss(ClearTextSelectionOnFocusLoss);
 	MyMultiLineEditableText->SetRevertTextOnEscape(RevertTextOnEscape);
diff --git a/Engine/Source/Runtime/UMG/Private/Components/MultiLineEditableTextBox.cpp b/Engine/Source/Runtime/UMG/Private/Components/MultiLineEditableTextBox.cpp
--- a/Engine/Source/Runtime/UMG/Private/Components/MultiLineEditableTextBox.cpp
+++ b/Engine/Source/Runtime/UMG/Private/Components/MultiLineEditableTextBox.cpp
@@ -23,6 +23,7 @@ UMultiLineEditableTextBox::UMultiLineEditableTextBox(const FObjectInitializer& O
 	WidgetStyle = *Defaults._Style;
 	TextStyle = *Defaults._TextStyle;
 	AllowContextMenu = Defaults._AllowContextMenu.Get();
+	VirtualKeyboardDismissAction = Defaults._VirtualKeyboardDismissAction.Get();
 	AutoWrapText = true;
 
 	if (!IsRunningDedicatedServer())
@@ -58,6 +59,7 @@ TSharedRef<SWidget> UMultiLineEditableTextBox::RebuildWidget()
 //		.RevertTextOnEscape(RevertTextOnEscape)
 //		.ClearKeyboardFocusOnCommit(ClearKeyboardFocusOnCommit)
 //		.SelectAllTextOnCommit(SelectAllTextOnCommit)
+		.VirtualKeyboardDismissAction(VirtualKeyboardDismissAction)
 		.OnTextChanged(BIND_UOBJECT_DELEGATE(FOnTextChanged, HandleOnTextChanged))
 		.OnTextCommitted(BIND_UOBJECT_DELEGATE(FOnTextCommitted, HandleOnTextCommitted))
 		;
@@ -76,6 +78,7 @@ void UMultiLineEditableTextBox::SynchronizeProperties()
 	MyEditableTextBlock->SetHintText(HintTextBinding);
 	MyEditableTextBlock->SetAllowContextMenu(AllowContextMenu);
 	MyEditableTextBlock->SetIsReadOnly(bIsReadOnly);
+	MyEditableTextBlock->SetVirtualKeyboardDismissAction(VirtualKeyboardDismissAction);
 
 //	MyEditableTextBlock->SetIsPassword(IsPassword);
 //	MyEditableTextBlock->SetColorAndOpacity(ColorAndOpacity);
diff --git a/Engine/Source/Runtime/UMG/Public/Components/EditableText.h b/Engine/Source/Runtime/UMG/Public/Components/EditableText.h
--- a/Engine/Source/Runtime/UMG/Public/Components/EditableText.h
+++ b/Engine/Source/Runtime/UMG/Public/Components/EditableText.h
@@ -10,6 +10,7 @@
 #include "Widgets/SWidget.h"
 #include "Components/Widget.h"
 #include "Components/TextWidgetTypes.h"
+#include "Widgets/Text/ISlateEditableTextWidget.h"
 #include "EditableText.generated.h"
 
 class SEditableText;
@@ -116,6 +117,10 @@ public:
 	UPROPERTY(EditAnywhere, Category=Behavior, AdvancedDisplay)
 	TEnumAsByte<EVirtualKeyboardType::Type> KeyboardType;
 
+	/** What action should be taken when the virtual keyboard is dismissed? */
+	UPROPERTY(EditAnywhere, Category=Behavior, AdvancedDisplay)
+	EVirtualKeyboardDismissAction VirtualKeyboardDismissAction;
+	
 	/** How the text should be aligned with the margin. */
 	UPROPERTY(EditAnywhere, BlueprintReadWrite, Category=Appearance)
 	TEnumAsByte<ETextJustify::Type> Justification;
diff --git a/Engine/Source/Runtime/UMG/Public/Components/EditableTextBox.h b/Engine/Source/Runtime/UMG/Public/Components/EditableTextBox.h
--- a/Engine/Source/Runtime/UMG/Public/Components/EditableTextBox.h
+++ b/Engine/Source/Runtime/UMG/Public/Components/EditableTextBox.h
@@ -10,6 +10,7 @@
 #include "Widgets/SWidget.h"
 #include "Components/Widget.h"
 #include "Components/TextWidgetTypes.h"
+#include "Widgets/Text/ISlateEditableTextWidget.h"
 #include "EditableTextBox.generated.h"
 
 class SEditableTextBox;
@@ -118,6 +119,10 @@ public:
 	UPROPERTY(EditAnywhere, Category=Behavior, AdvancedDisplay)
 	TEnumAsByte<EVirtualKeyboardType::Type> KeyboardType;
 
+	/** What action should be taken when the virtual keyboard is dismissed? */
+	UPROPERTY(EditAnywhere, Category=Behavior, AdvancedDisplay)
+	EVirtualKeyboardDismissAction VirtualKeyboardDismissAction;
+	
 	/** How the text should be aligned with the margin. */
 	UPROPERTY(EditAnywhere, BlueprintReadWrite, Category=Appearance)
 	TEnumAsByte<ETextJustify::Type> Justification;
diff --git a/Engine/Source/Runtime/UMG/Public/Components/MultiLineEditableText.h b/Engine/Source/Runtime/UMG/Public/Components/MultiLineEditableText.h
--- a/Engine/Source/Runtime/UMG/Public/Components/MultiLineEditableText.h
+++ b/Engine/Source/Runtime/UMG/Public/Components/MultiLineEditableText.h
@@ -8,6 +8,7 @@
 #include "Styling/SlateTypes.h"
 #include "Widgets/SWidget.h"
 #include "Components/TextWidgetTypes.h"
+#include "Widgets/Text/ISlateEditableTextWidget.h"
 #include "MultiLineEditableText.generated.h"
 
 class SMultiLineEditableText;
@@ -70,6 +71,10 @@ public:
 	/** Whether the context menu can be opened */
 	UPROPERTY(EditAnywhere, Category=Behavior, AdvancedDisplay)
 	bool AllowContextMenu;
+	
+	/** What action should be taken when the virtual keyboard is dismissed? */
+	UPROPERTY(EditAnywhere, Category=Behavior, AdvancedDisplay)
+	EVirtualKeyboardDismissAction VirtualKeyboardDismissAction;
 
 	/** Called whenever the text is changed interactively by the user */
 	UPROPERTY(BlueprintAssignable, Category="Widget Event", meta=(DisplayName="OnTextChanged (Multi-Line Editable Text)"))
diff --git a/Engine/Source/Runtime/UMG/Public/Components/MultiLineEditableTextBox.h b/Engine/Source/Runtime/UMG/Public/Components/MultiLineEditableTextBox.h
--- a/Engine/Source/Runtime/UMG/Public/Components/MultiLineEditableTextBox.h
+++ b/Engine/Source/Runtime/UMG/Public/Components/MultiLineEditableTextBox.h
@@ -8,6 +8,7 @@
 #include "Styling/SlateTypes.h"
 #include "Widgets/SWidget.h"
 #include "Components/TextWidgetTypes.h"
+#include "Widgets/Text/ISlateEditableTextWidget.h"
 #include "MultiLineEditableTextBox.generated.h"
 
 class SMultiLineEditableTextBox;
@@ -56,6 +57,10 @@ public:
 	/** Whether the context menu can be opened */
 	UPROPERTY(EditAnywhere, Category=Behavior, AdvancedDisplay)
 	bool AllowContextMenu;
+	
+	/** What action should be taken when the virtual keyboard is dismissed? */
+	UPROPERTY(EditAnywhere, Category=Behavior, AdvancedDisplay)
+	EVirtualKeyboardDismissAction VirtualKeyboardDismissAction;
 
 	UPROPERTY()
 	USlateWidgetStyleAsset* Style_DEPRECATED;
diff --git a/Engine/Source/Runtime/VulkanRHI/Private/VulkanRenderTarget.cpp b/Engine/Source/Runtime/VulkanRHI/Private/VulkanRenderTarget.cpp
--- a/Engine/Source/Runtime/VulkanRHI/Private/VulkanRenderTarget.cpp
+++ b/Engine/Source/Runtime/VulkanRHI/Private/VulkanRenderTarget.cpp
@@ -631,9 +631,9 @@ void FVulkanDynamicRHI::RHIReadSurfaceData(FTextureRHIParamRef TextureRHI, FIntR
 			FFloat16Color* Src = (FFloat16Color*)StagingBuffer->GetMappedPointer() + Row * TextureRHI2D->GetSizeX() + Rect.Min.X;
 			for (int32 Col = Rect.Min.X; Col < Rect.Max.X; ++Col)
 			{
-				Dest->R = (uint8)(uint32)FMath::Clamp<int32>((int32)(Src->B.GetFloat() * 255.0f), 0, 255);
+				Dest->R = (uint8)(uint32)FMath::Clamp<int32>((int32)(Src->R.GetFloat() * 255.0f), 0, 255);
 				Dest->G = (uint8)(uint32)FMath::Clamp<int32>((int32)(Src->G.GetFloat() * 255.0f), 0, 255);
-				Dest->B = (uint8)(uint32)FMath::Clamp<int32>((int32)(Src->R.GetFloat() * 255.0f), 0, 255);
+				Dest->B = (uint8)(uint32)FMath::Clamp<int32>((int32)(Src->B.GetFloat() * 255.0f), 0, 255);
 				Dest->A = (uint8)(uint32)FMath::Clamp<int32>((int32)(Src->A.GetFloat() * 255.0f), 0, 255);
 				Dest++;
 				Src++;
