Fixed: UMG element clipping does not work on Android with Mali GPUs

diff --git a/Engine/Source/Runtime/OpenGLDrv/Private/Android/AndroidEGL.cpp b/Engine/Source/Runtime/OpenGLDrv/Private/Android/AndroidEGL.cpp
--- a/Engine/Source/Runtime/OpenGLDrv/Private/Android/AndroidEGL.cpp
+++ b/Engine/Source/Runtime/OpenGLDrv/Private/Android/AndroidEGL.cpp
@@ -107,7 +107,7 @@ const  int EGLMinGreenBits		= 6;
 const  int EGLMinBlueBits		= 5;
 const  int EGLMinAlphaBits		= 0;
 const  int EGLMinDepthBits		= 16;
-const  int EGLMinStencilBits	= 0;
+const  int EGLMinStencilBits	= 8; // This is required for UMG clipping
 const  int EGLMinSampleBuffers	= 0;
 const  int EGLMinSampleSamples	= 0;
 
