Fixed: Very large frame time increases on Android Adreno GPUs (SunTemple)
Disabled OpenGL fine grained stats
Disabled GL thread check in development (enabled in debug only)

diff --git a/Engine/Source/Runtime/OpenGLDrv/Public/OpenGLResources.h b/Engine/Source/Runtime/OpenGLDrv/Public/OpenGLResources.h
--- a/Engine/Source/Runtime/OpenGLDrv/Public/OpenGLResources.h
+++ b/Engine/Source/Runtime/OpenGLDrv/Public/OpenGLResources.h
@@ -63,11 +63,9 @@ void DecrementBufferMemory(GLenum Type, bool bStructuredBuffer, uint32 NumBytes)
 
 // Extra stats for finer-grained timing
 // They shouldn't always be on, as they may impact overall performance
-#if UE_BUILD_DEVELOPMENT
-	#define OPENGLRHI_DETAILED_STATS 1
-#else
-	#define OPENGLRHI_DETAILED_STATS 0
-#endif
+#define OPENGLRHI_DETAILED_STATS 0
+
+
 #if OPENGLRHI_DETAILED_STATS
 	DECLARE_CYCLE_STAT_EXTERN(TEXT("MapBuffer time"),STAT_OpenGLMapBufferTime,STATGROUP_OpenGLRHI, );
 	DECLARE_CYCLE_STAT_EXTERN(TEXT("UnmapBuffer time"),STAT_OpenGLUnmapBufferTime,STATGROUP_OpenGLRHI, );
diff --git a/Engine/Source/Runtime/OpenGLDrv/Public/OpenGLUtil.h b/Engine/Source/Runtime/OpenGLDrv/Public/OpenGLUtil.h
--- a/Engine/Source/Runtime/OpenGLDrv/Public/OpenGLUtil.h
+++ b/Engine/Source/Runtime/OpenGLDrv/Public/OpenGLUtil.h
@@ -15,7 +15,7 @@
 #define ENABLE_VERIFY_GL_TRACE 0
 
 // Additional check that our GL calls are occurring on the expected thread
-#define ENABLE_VERIFY_GL_THREAD (DO_CHECK)
+#define ENABLE_VERIFY_GL_THREAD (UE_BUILD_DEBUG)
 
 /** Set to 1 to verify that the the engine side uniform buffer layout matches the driver side of the GLSL shader*/
 #define ENABLE_UNIFORM_BUFFER_LAYOUT_VERIFICATION ( 0 & UE_BUILD_DEBUG & (OPENGL_ESDEFERRED | OPENGL_GL3 | OPENGL_GL4))
