Fix opengl ES2 emulated uniform buffer creation crash when ogl RHI thread is enabled.

diff --git a/Engine/Source/Runtime/OpenGLDrv/Private/OpenGLUniformBuffer.cpp b/Engine/Source/Runtime/OpenGLDrv/Private/OpenGLUniformBuffer.cpp
--- a/Engine/Source/Runtime/OpenGLDrv/Private/OpenGLUniformBuffer.cpp
+++ b/Engine/Source/Runtime/OpenGLDrv/Private/OpenGLUniformBuffer.cpp
@@ -158,12 +158,12 @@ static inline uint32 GetPoolBucketSize(uint32 NumBytes)
 	return UniformBufferSizeBuckets[GetPoolBucketIndex(NumBytes)];
 }
 
+static FCriticalSection GGLEmulatedUniformBufferDataFactoryCS;
 struct FUniformBufferDataFactory
 {
-	TMap<GLuint, FOpenGLEUniformBufferDataRef> Entries;
-
 	FOpenGLEUniformBufferDataRef Create(uint32 Size, GLuint& OutResource)
 	{
+		FScopeLock Lock(&GGLEmulatedUniformBufferDataFactoryCS);
 		static GLuint TempCounter = 0;
 		OutResource = ++TempCounter;
 
@@ -174,6 +174,7 @@ struct FUniformBufferDataFactory
 
 	FOpenGLEUniformBufferDataRef Get(GLuint Resource)
 	{
+		FScopeLock Lock(&GGLEmulatedUniformBufferDataFactoryCS);
 		FOpenGLEUniformBufferDataRef* Buffer = Entries.Find(Resource);
 		check(Buffer);
 		return *Buffer;
@@ -181,8 +182,11 @@ struct FUniformBufferDataFactory
 
 	void Destroy(GLuint Resource)
 	{
+		FScopeLock Lock(&GGLEmulatedUniformBufferDataFactoryCS);
 		Entries.Remove(Resource);
 	}
+private:
+	TMap<GLuint, FOpenGLEUniformBufferDataRef> Entries;
 };
 
 static FUniformBufferDataFactory UniformBufferDataFactory;
