Fixed: Blackscreen in app on Galaxy S6 SM-G920F

diff --git a/Engine/Source/Runtime/OpenGLDrv/Private/OpenGLShaders.cpp b/Engine/Source/Runtime/OpenGLDrv/Private/OpenGLShaders.cpp
--- a/Engine/Source/Runtime/OpenGLDrv/Private/OpenGLShaders.cpp
+++ b/Engine/Source/Runtime/OpenGLDrv/Private/OpenGLShaders.cpp
@@ -775,12 +775,14 @@ void OPENGLDRV_API GLSLToDeviceCompatibleGLSL(FAnsiCharArray& GlslCodeOriginal,
 		if (bES31)
 		{
 			// @todo Lumin hack: This is needed for AEP on Lumin, so that some shaders compile that need version 320
-			AppendCString(GlslCode, "#version 320 es\n");
-			ReplaceCString(GlslCodeOriginal, ES310Version, "");
-
-			//AppendCString(GlslCode, ES310Version);
-			//AppendCString(GlslCode, "\n");
-			//ReplaceCString(GlslCodeOriginal, ES310Version, "");
+			#if PLATFORM_LUMINGL4
+				AppendCString(GlslCode, "#version 320 es\n");
+				ReplaceCString(GlslCodeOriginal, ES310Version, "");
+			#else
+				AppendCString(GlslCode, ES310Version);
+				AppendCString(GlslCode, "\n");
+				ReplaceCString(GlslCodeOriginal, ES310Version, "");
+			#endif
 		}
 		else if (IsES2Platform(Capabilities.MaxRHIShaderPlatform))
 		{
