Niagara cleans up some resources in AppPreExit, before the RHI shuts down, and before plugins are shut down.
The logic for buffers doesn't handle potentially releasing a buffer on the game thread, so it has been modified to query the correct function.
#tests no longer triggers ensure

diff --git a/Engine/Source/Runtime/OpenGLDrv/Private/OpenGLUniformBuffer.cpp b/Engine/Source/Runtime/OpenGLDrv/Private/OpenGLUniformBuffer.cpp
--- a/Engine/Source/Runtime/OpenGLDrv/Private/OpenGLUniformBuffer.cpp
+++ b/Engine/Source/Runtime/OpenGLDrv/Private/OpenGLUniformBuffer.cpp
@@ -297,7 +297,7 @@ void BeginFrame_UniformBufferPoolCleanup()
 static bool IsPoolingEnabled()
 {
 	static const auto CVar = IConsoleManager::Get().FindTConsoleVariableDataInt(TEXT("r.UniformBufferPooling"));
-	int32 CVarValue = CVar->GetValueOnRenderThread();
+	int32 CVarValue = IsInParallelRenderingThread() ? CVar->GetValueOnRenderThread() : CVar->GetValueOnGameThread();
 	return CVarValue != 0;
 };
 
