Fixed typos in glObjectPtrLabel type definitions (contributed by [REDACTED_PERSON])
#4920

diff --git a/Engine/Source/Runtime/OpenGLDrv/Private/Linux/OpenGLLinux.h b/Engine/Source/Runtime/OpenGLDrv/Private/Linux/OpenGLLinux.h
--- a/Engine/Source/Runtime/OpenGLDrv/Private/Linux/OpenGLLinux.h
+++ b/Engine/Source/Runtime/OpenGLDrv/Private/Linux/OpenGLLinux.h
@@ -326,7 +326,7 @@ THIRD_PARTY_INCLUDES_END
 	EnumMacro(PFNGLPUSHDEBUGGROUPPROC, glPushDebugGroup)\
 	EnumMacro(PFNGLPOPDEBUGGROUPPROC, glPopDebugGroup)\
 	EnumMacro(PFNGLOBJECTLABELPROC, glObjectLabel)\
-	EnumMacro(PFNGLOBJECTLABELPROC, glObjectPtrLabel)\
+	EnumMacro(PFNGLOBJECTPTRLABELPROC, glObjectPtrLabel)\
 	EnumMacro(PFNGLPATCHPARAMETERIPROC, glPatchParameteri)\
 	EnumMacro(PFNGLBINDVERTEXBUFFERPROC, glBindVertexBuffer)\
 	EnumMacro(PFNGLVERTEXATTRIBFORMATPROC, glVertexAttribFormat)\
diff --git a/Engine/Source/Runtime/OpenGLDrv/Private/Lumin/LuminGL4.h b/Engine/Source/Runtime/OpenGLDrv/Private/Lumin/LuminGL4.h
--- a/Engine/Source/Runtime/OpenGLDrv/Private/Lumin/LuminGL4.h
+++ b/Engine/Source/Runtime/OpenGLDrv/Private/Lumin/LuminGL4.h
@@ -362,7 +362,7 @@ typedef void (APIENTRYP PFNGLGETVERTEXATTRIBLUI64VNVPROC) (GLuint index, GLenum
 	EnumMacro(PFNGLCLEARBUFFERDATAPROC, glClearBufferData)\
 	EnumMacro(PFNGLDISPATCHCOMPUTEINDIRECTPROC, glDispatchComputeIndirect) \
 	EnumMacro(PFNGLOBJECTLABELPROC, glObjectLabel)\
-	EnumMacro(PFNGLOBJECTLABELPROC, glObjectPtrLabel)\
+	EnumMacro(PFNGLOBJECTPTRLABELPROC, glObjectPtrLabel)\
 	EnumMacro(PFNGLPUSHDEBUGGROUPPROC, glPushDebugGroup)\
 	EnumMacro(PFNGLPOPDEBUGGROUPPROC, glPopDebugGroup)\
 	EnumMacro(PFNGLVERTEXBINDINGDIVISORPROC, glVertexBindingDivisor)\
diff --git a/Engine/Source/Runtime/OpenGLDrv/Private/Windows/OpenGLWindows.h b/Engine/Source/Runtime/OpenGLDrv/Private/Windows/OpenGLWindows.h
--- a/Engine/Source/Runtime/OpenGLDrv/Private/Windows/OpenGLWindows.h
+++ b/Engine/Source/Runtime/OpenGLDrv/Private/Windows/OpenGLWindows.h
@@ -327,7 +327,7 @@
 	EnumMacro(PFNGLPUSHDEBUGGROUPPROC, glPushDebugGroup)\
 	EnumMacro(PFNGLPOPDEBUGGROUPPROC, glPopDebugGroup)\
 	EnumMacro(PFNGLOBJECTLABELPROC, glObjectLabel)\
-	EnumMacro(PFNGLOBJECTLABELPROC, glObjectPtrLabel)\
+	EnumMacro(PFNGLOBJECTPTRLABELPROC, glObjectPtrLabel)\
 	EnumMacro(PFNGLPATCHPARAMETERIPROC, glPatchParameteri)\
 	EnumMacro(PFNGLBINDVERTEXBUFFERPROC, glBindVertexBuffer)\
 	EnumMacro(PFNGLVERTEXATTRIBFORMATPROC, glVertexAttribFormat)\
@@ -352,7 +352,7 @@
 	EnumMacro(PFNGLPUSHDEBUGGROUPPROC, glPushDebugGroupKHR)\
 	EnumMacro(PFNGLPOPDEBUGGROUPPROC, glPopDebugGroupKHR)\
 	EnumMacro(PFNGLOBJECTLABELPROC, glObjectLabelKHR)\
-	EnumMacro(PFNGLOBJECTLABELPROC, glObjectPtrLabelKHR)\
+	EnumMacro(PFNGLOBJECTPTRLABELPROC, glObjectPtrLabelKHR)\
 	EnumMacro(PFNGLDEBUGMESSAGECALLBACKARBPROC,glDebugMessageCallbackKHR) \
 	EnumMacro(PFNGLDEBUGMESSAGECONTROLARBPROC,glDebugMessageControlKHR) \
 	EnumMacro(PFNGLPATCHPARAMETERIPROC, glPatchParameteriEXT)\
