Fix LLM reporting for meshes on OpenGL (Android)

#ue4
#android

diff --git a/Engine/Source/Runtime/OpenGLDrv/Public/OpenGLResources.h b/Engine/Source/Runtime/OpenGLDrv/Public/OpenGLResources.h
--- a/Engine/Source/Runtime/OpenGLDrv/Public/OpenGLResources.h
+++ b/Engine/Source/Runtime/OpenGLDrv/Public/OpenGLResources.h
@@ -843,7 +843,7 @@ public:
 		}
 
 #if ENABLE_LOW_LEVEL_MEM_TRACKER
-		LLM_SCOPED_PAUSE_TRACKING_WITH_ENUM_AND_AMOUNT(ELLMTag::VertexBuffer, InSize, ELLMTracker::Default, ELLMAllocType::None);
+		LLM_SCOPED_PAUSE_TRACKING_WITH_ENUM_AND_AMOUNT(ELLMTag::Meshes, InSize, ELLMTracker::Default, ELLMAllocType::None);
 #endif
 	}
 
@@ -855,7 +855,7 @@ public:
 		}
 
 #if ENABLE_LOW_LEVEL_MEM_TRACKER
-		LLM_SCOPED_PAUSE_TRACKING_WITH_ENUM_AND_AMOUNT(ELLMTag::VertexBuffer, -(int64)GetSize(), ELLMTracker::Default, ELLMAllocType::None);
+		LLM_SCOPED_PAUSE_TRACKING_WITH_ENUM_AND_AMOUNT(ELLMTag::Meshes, -(int64)GetSize(), ELLMTracker::Default, ELLMAllocType::None);
 #endif
 	}
 
@@ -952,14 +952,14 @@ public:
 	FOpenGLBaseIndexBuffer(uint32 InStride,uint32 InSize,uint32 InUsage): FRHIIndexBuffer(InStride,InSize,InUsage)
 	{
 #if ENABLE_LOW_LEVEL_MEM_TRACKER
-		LLM_SCOPED_PAUSE_TRACKING_WITH_ENUM_AND_AMOUNT(ELLMTag::IndexBuffer, InSize, ELLMTracker::Default, ELLMAllocType::None);
+		LLM_SCOPED_PAUSE_TRACKING_WITH_ENUM_AND_AMOUNT(ELLMTag::Meshes, InSize, ELLMTracker::Default, ELLMAllocType::None);
 #endif
 	}
 
 	~FOpenGLBaseIndexBuffer(void)
 	{
 #if ENABLE_LOW_LEVEL_MEM_TRACKER
-		LLM_SCOPED_PAUSE_TRACKING_WITH_ENUM_AND_AMOUNT(ELLMTag::IndexBuffer, -(int64)GetSize(), ELLMTracker::Default, ELLMAllocType::None);
+		LLM_SCOPED_PAUSE_TRACKING_WITH_ENUM_AND_AMOUNT(ELLMTag::Meshes, -(int64)GetSize(), ELLMTracker::Default, ELLMAllocType::None);
 #endif
 	}
 
