Add some additional log info when OpenGL program creation fails.
#rnx

diff --git a/Engine/Source/Runtime/OpenGLDrv/Private/OpenGLShaders.cpp b/Engine/Source/Runtime/OpenGLDrv/Private/OpenGLShaders.cpp
--- a/Engine/Source/Runtime/OpenGLDrv/Private/OpenGLShaders.cpp
+++ b/Engine/Source/Runtime/OpenGLDrv/Private/OpenGLShaders.cpp
@@ -2151,7 +2151,7 @@ class FGLProgramCacheLRU
 			else
 			{
 				RHIGetPanicDelegate().ExecuteIfBound(FName("FailedBinaryProgramCreate"));
-				UE_LOG(LogRHI, Fatal, TEXT("RestoreGLProgramFromBinary : Failed to restore GL program from binary data!"));
+				UE_LOG(LogRHI, Fatal, TEXT("RestoreGLProgramFromBinary : Failed to restore GL program from binary data! [%s]"), *LinkedProgram->Config.ProgramKey.ToString());
 			}
 		}
 
@@ -3748,7 +3748,7 @@ FBoundShaderStateRHIRef FOpenGLDynamicRHI::RHICreateBoundShaderState_OnThisThrea
 #endif //DEBUG_GL_SHADERS
 					FName LinkFailurePanic = bFromPSOFileCache ? FName("FailedProgramLinkDuringPrecompile") : FName("FailedProgramLink");
 					RHIGetPanicDelegate().ExecuteIfBound(LinkFailurePanic);
-					UE_LOG(LogRHI, Fatal, TEXT("Failed to link program. Current total programs: %d, precompile: %d"), GNumPrograms, (uint32)bFromPSOFileCache);
+					UE_LOG(LogRHI, Fatal, TEXT("Failed to link program [%s]. Current total programs: %d, precompile: %d"), *Config.ProgramKey.ToString(), GNumPrograms, (uint32)bFromPSOFileCache);
 				}
 
 				GetOpenGLProgramsCache().Add(Config.ProgramKey, LinkedProgram);
@@ -5121,7 +5121,7 @@ bool FOpenGLProgramBinaryCache::UseCachedProgram_internal(GLuint& ProgramOUT, co
 			if (!bSuccess)
 			{
 				RHIGetPanicDelegate().ExecuteIfBound(FName("FailedBinaryProgramCreateFromOldCache"));
-				UE_LOG(LogRHI, Fatal, TEXT("UseCachedProgram : Failed to create GL program from binary data while BuildingCacheFileWithMove!"));
+				UE_LOG(LogRHI, Fatal, TEXT("UseCachedProgram : Failed to create GL program from binary data while BuildingCacheFileWithMove! [%s]"), *ProgramKey.ToString());
 			}
 			SetNewProgramStats(ProgramOUT);
 			// Now write to new cache, we're returning true here so no attempt will be made to add it back to the cache later.
@@ -5264,7 +5264,7 @@ void FOpenGLProgramBinaryCache::CompleteLoadedGLProgramRequest_internal(FGLProgr
 		if(!bSuccess)
 		{
 			RHIGetPanicDelegate().ExecuteIfBound(FName("FailedBinaryProgramCreate"));
-			UE_LOG(LogRHI, Fatal, TEXT("CompleteLoadedGLProgramRequest_internal : Failed to create GL program from binary data!"));
+			UE_LOG(LogRHI, Fatal, TEXT("CompleteLoadedGLProgramRequest_internal : Failed to create GL program from binary data! [%s]"), *ProgramKey.ToString());
 		}
 		VerifyProgramPipeline(PendingGLCreate->GLProgramId);
 		FOpenGLLinkedProgram* NewLinkedProgram = new FOpenGLLinkedProgram(ProgramKey, PendingGLCreate->GLProgramId);
