Adding more logging for failed binary program creates.
#rnx

diff --git a/Engine/Source/Runtime/OpenGLDrv/Private/OpenGLShaders.cpp b/Engine/Source/Runtime/OpenGLDrv/Private/OpenGLShaders.cpp
--- a/Engine/Source/Runtime/OpenGLDrv/Private/OpenGLShaders.cpp
+++ b/Engine/Source/Runtime/OpenGLDrv/Private/OpenGLShaders.cpp
@@ -2150,6 +2150,7 @@ class FGLProgramCacheLRU
 			}
 			else
 			{
+				UE_LOG(LogRHI, Log, TEXT("[%s, %d, %d]"), *LinkedProgram->Config.ProgramKey.ToString(), LinkedProgram->Program, GetProgramBinary().Num() );
 				RHIGetPanicDelegate().ExecuteIfBound(FName("FailedBinaryProgramCreate"));
 				UE_LOG(LogRHI, Fatal, TEXT("RestoreGLProgramFromBinary : Failed to restore GL program from binary data! [%s]"), *LinkedProgram->Config.ProgramKey.ToString());
 			}
@@ -5120,6 +5121,7 @@ bool FOpenGLProgramBinaryCache::UseCachedProgram_internal(GLuint& ProgramOUT, co
 			bool bSuccess = CreateGLProgramFromBinary(ProgramOUT, CachedProgramBinaryOUT);
 			if (!bSuccess)
 			{
+				UE_LOG(LogRHI, Log, TEXT("[%s, %d, %d]"), *ProgramKey.ToString(), ProgramOUT, CachedProgramBinaryOUT.Num());
 				RHIGetPanicDelegate().ExecuteIfBound(FName("FailedBinaryProgramCreateFromOldCache"));
 				UE_LOG(LogRHI, Fatal, TEXT("UseCachedProgram : Failed to create GL program from binary data while BuildingCacheFileWithMove! [%s]"), *ProgramKey.ToString());
 			}
@@ -5263,7 +5265,8 @@ void FOpenGLProgramBinaryCache::CompleteLoadedGLProgramRequest_internal(FGLProgr
 		bool bSuccess = CreateGLProgramFromBinary(PendingGLCreate->GLProgramId, PendingGLCreate->ProgramBinaryData);
 		if(!bSuccess)
 		{
-			RHIGetPanicDelegate().ExecuteIfBound(FName("FailedBinaryProgramCreate"));
+			UE_LOG(LogRHI, Log, TEXT("[%s, %d, %d]"), *ProgramKey.ToString(), PendingGLCreate->GLProgramId, PendingGLCreate->ProgramBinaryData.Num() );
+			RHIGetPanicDelegate().ExecuteIfBound(FName("FailedBinaryProgramCreateLoadRequest"));
 			UE_LOG(LogRHI, Fatal, TEXT("CompleteLoadedGLProgramRequest_internal : Failed to create GL program from binary data! [%s]"), *ProgramKey.ToString());
 		}
 		VerifyProgramPipeline(PendingGLCreate->GLProgramId);
