UE4.23 - Deprecate FUniformBufferRHIParamRef, FIndexBufferRHIParamRef, FVertexBufferRHIParamRef, FStructuredBufferRHIParamRef, FTextureRHIParamRef, FTexture2DRHIParamRef, FTexture2DArrayRHIParamRef, F

diff --git a/Engine/Source/Runtime/OpenGLDrv/Public/OpenGLDrv.h b/Engine/Source/Runtime/OpenGLDrv/Public/OpenGLDrv.h
--- a/Engine/Source/Runtime/OpenGLDrv/Public/OpenGLDrv.h
+++ b/Engine/Source/Runtime/OpenGLDrv/Public/OpenGLDrv.h
@@ -362,66 +362,66 @@ public:
 	virtual FComputeShaderRHIRef RHICreateComputeShader(FRHIShaderLibrary* Library, FSHAHash Hash) final override;
 
 	virtual FUniformBufferRHIRef RHICreateUniformBuffer(const void* Contents, const FRHIUniformBufferLayout& Layout, EUniformBufferUsage Usage, EUniformBufferValidation Validation) final override;
-	virtual void RHIUpdateUniformBuffer(FUniformBufferRHIParamRef UniformBufferRHI, const void* Contents) final override;
+	virtual void RHIUpdateUniformBuffer(FRHIUniformBuffer* UniformBufferRHI, const void* Contents) final override;
 	virtual FIndexBufferRHIRef RHICreateIndexBuffer(uint32 Stride, uint32 Size, uint32 InUsage, FRHIResourceCreateInfo& CreateInfo) final override;
-	virtual void* RHILockIndexBuffer(FIndexBufferRHIParamRef IndexBuffer, uint32 Offset, uint32 Size, EResourceLockMode LockMode) final override;
-	virtual void RHIUnlockIndexBuffer(FIndexBufferRHIParamRef IndexBuffer) final override;
+	virtual void* RHILockIndexBuffer(FRHIIndexBuffer* IndexBuffer, uint32 Offset, uint32 Size, EResourceLockMode LockMode) final override;
+	virtual void RHIUnlockIndexBuffer(FRHIIndexBuffer* IndexBuffer) final override;
 	virtual FVertexBufferRHIRef RHICreateVertexBuffer(uint32 Size, uint32 InUsage, FRHIResourceCreateInfo& CreateInfo) final override;
-	virtual void* RHILockVertexBuffer(FVertexBufferRHIParamRef VertexBuffer, uint32 Offset, uint32 SizeRHI, EResourceLockMode LockMode) final override;
-	virtual void RHIUnlockVertexBuffer(FVertexBufferRHIParamRef VertexBuffer) final override;
-	virtual void RHICopyVertexBuffer(FVertexBufferRHIParamRef SourceBuffer,FVertexBufferRHIParamRef DestBuffer) final override;
+	virtual void* RHILockVertexBuffer(FRHIVertexBuffer* VertexBuffer, uint32 Offset, uint32 SizeRHI, EResourceLockMode LockMode) final override;
+	virtual void RHIUnlockVertexBuffer(FRHIVertexBuffer* VertexBuffer) final override;
+	virtual void RHICopyVertexBuffer(FRHIVertexBuffer* SourceBuffer, FRHIVertexBuffer* DestBuffer) final override;
 	virtual FStructuredBufferRHIRef RHICreateStructuredBuffer(uint32 Stride, uint32 Size, uint32 InUsage, FRHIResourceCreateInfo& CreateInfo) final override;
-	virtual void* RHILockStructuredBuffer(FStructuredBufferRHIParamRef StructuredBuffer, uint32 Offset, uint32 SizeRHI, EResourceLockMode LockMode) final override;
-	virtual void RHIUnlockStructuredBuffer(FStructuredBufferRHIParamRef StructuredBuffer) final override;
-	virtual FUnorderedAccessViewRHIRef RHICreateUnorderedAccessView(FStructuredBufferRHIParamRef StructuredBuffer, bool bUseUAVCounter, bool bAppendBuffer) final override;
-	virtual FUnorderedAccessViewRHIRef RHICreateUnorderedAccessView(FTextureRHIParamRef Texture, uint32 MipLevel) final override;
-	virtual FUnorderedAccessViewRHIRef RHICreateUnorderedAccessView(FVertexBufferRHIParamRef VertexBuffer, uint8 Format) final override;
-	virtual FUnorderedAccessViewRHIRef RHICreateUnorderedAccessView(FIndexBufferRHIParamRef IndexBuffer, uint8 Format) final override;
-	virtual FShaderResourceViewRHIRef RHICreateShaderResourceView(FStructuredBufferRHIParamRef StructuredBuffer) final override;
-	virtual FShaderResourceViewRHIRef RHICreateShaderResourceView(FVertexBufferRHIParamRef VertexBuffer, uint32 Stride, uint8 Format) final override;
-	virtual FShaderResourceViewRHIRef RHICreateShaderResourceView(FIndexBufferRHIParamRef Buffer) final override;
-	virtual void RHIUpdateShaderResourceView(FRHIShaderResourceView* SRV, FVertexBufferRHIParamRef VertexBuffer, uint32 Stride, uint8 Format) final override;
+	virtual void* RHILockStructuredBuffer(FRHIStructuredBuffer* StructuredBuffer, uint32 Offset, uint32 SizeRHI, EResourceLockMode LockMode) final override;
+	virtual void RHIUnlockStructuredBuffer(FRHIStructuredBuffer* StructuredBuffer) final override;
+	virtual FUnorderedAccessViewRHIRef RHICreateUnorderedAccessView(FRHIStructuredBuffer* StructuredBuffer, bool bUseUAVCounter, bool bAppendBuffer) final override;
+	virtual FUnorderedAccessViewRHIRef RHICreateUnorderedAccessView(FRHITexture* Texture, uint32 MipLevel) final override;
+	virtual FUnorderedAccessViewRHIRef RHICreateUnorderedAccessView(FRHIVertexBuffer* VertexBuffer, uint8 Format) final override;
+	virtual FUnorderedAccessViewRHIRef RHICreateUnorderedAccessView(FRHIIndexBuffer* IndexBuffer, uint8 Format) final override;
+	virtual FShaderResourceViewRHIRef RHICreateShaderResourceView(FRHIStructuredBuffer* StructuredBuffer) final override;
+	virtual FShaderResourceViewRHIRef RHICreateShaderResourceView(FRHIVertexBuffer* VertexBuffer, uint32 Stride, uint8 Format) final override;
+	virtual FShaderResourceViewRHIRef RHICreateShaderResourceView(FRHIIndexBuffer* Buffer) final override;
+	virtual void RHIUpdateShaderResourceView(FRHIShaderResourceView* SRV, FRHIVertexBuffer* VertexBuffer, uint32 Stride, uint8 Format) final override;
 	virtual uint64 RHICalcTexture2DPlatformSize(uint32 SizeX, uint32 SizeY, uint8 Format, uint32 NumMips, uint32 NumSamples, uint32 Flags, uint32& OutAlign) final override;
 	virtual uint64 RHICalcTexture3DPlatformSize(uint32 SizeX, uint32 SizeY, uint32 SizeZ, uint8 Format, uint32 NumMips, uint32 Flags, uint32& OutAlign) final override;
 	virtual uint64 RHICalcTextureCubePlatformSize(uint32 Size, uint8 Format, uint32 NumMips, uint32 Flags, uint32& OutAlign) final override;
 	virtual void RHIGetTextureMemoryStats(FTextureMemoryStats& OutStats) final override;
 	virtual bool RHIGetTextureMemoryVisualizeData(FColor* TextureData,int32 SizeX,int32 SizeY,int32 Pitch,int32 PixelSize) final override;
 	virtual FTextureReferenceRHIRef RHICreateTextureReference(FLastRenderTimeContainer* LastRenderTime) final override;
-	virtual void RHIUpdateTextureReference(FRHITextureReference* TextureRef, FTextureRHIParamRef NewTexture) final override;
+	virtual void RHIUpdateTextureReference(FRHITextureReference* TextureRef, FRHITexture* NewTexture) final override;
 	virtual FTexture2DRHIRef RHICreateTexture2D(uint32 SizeX, uint32 SizeY, uint8 Format, uint32 NumMips, uint32 NumSamples, uint32 Flags, FRHIResourceCreateInfo& CreateInfo) final override;
 	virtual FTexture2DRHIRef RHICreateTextureExternal2D(uint32 SizeX, uint32 SizeY, uint8 Format, uint32 NumMips, uint32 NumSamples, uint32 Flags, FRHIResourceCreateInfo& CreateInfo) final override;
 	virtual FTexture2DRHIRef RHIAsyncCreateTexture2D(uint32 SizeX, uint32 SizeY, uint8 Format, uint32 NumMips, uint32 Flags, void** InitialMipData, uint32 NumInitialMips) final override;
-	virtual void RHICopySharedMips(FTexture2DRHIParamRef DestTexture2D, FTexture2DRHIParamRef SrcTexture2D) final override;
+	virtual void RHICopySharedMips(FRHITexture2D* DestTexture2D, FRHITexture2D* SrcTexture2D) final override;
 	virtual FTexture2DArrayRHIRef RHICreateTexture2DArray(uint32 SizeX, uint32 SizeY, uint32 SizeZ, uint8 Format, uint32 NumMips, uint32 NumSamples, uint32 Flags, FRHIResourceCreateInfo& CreateInfo) final override;
 	virtual FTexture3DRHIRef RHICreateTexture3D(uint32 SizeX, uint32 SizeY, uint32 SizeZ, uint8 Format, uint32 NumMips, uint32 Flags, FRHIResourceCreateInfo& CreateInfo) final override;
-	virtual void RHIGetResourceInfo(FTextureRHIParamRef Ref, FRHIResourceInfo& OutInfo) final override;
-	virtual FShaderResourceViewRHIRef RHICreateShaderResourceView(FTextureRHIParamRef Texture, const FRHITextureSRVCreateInfo& CreateInfo) final override;
-	virtual void RHIGenerateMips(FTextureRHIParamRef Texture) final override;
-	virtual uint32 RHIComputeMemorySize(FTextureRHIParamRef TextureRHI) final override;
-	virtual FTexture2DRHIRef RHIAsyncReallocateTexture2D(FTexture2DRHIParamRef Texture2D, int32 NewMipCount, int32 NewSizeX, int32 NewSizeY, FThreadSafeCounter* RequestStatus) final override;
-	virtual ETextureReallocationStatus RHIFinalizeAsyncReallocateTexture2D(FTexture2DRHIParamRef Texture2D, bool bBlockUntilCompleted) final override;
-	virtual ETextureReallocationStatus RHICancelAsyncReallocateTexture2D(FTexture2DRHIParamRef Texture2D, bool bBlockUntilCompleted) final override;
-	virtual void* RHILockTexture2D(FTexture2DRHIParamRef Texture, uint32 MipIndex, EResourceLockMode LockMode, uint32& DestStride, bool bLockWithinMiptail) final override;
-	virtual void RHIUnlockTexture2D(FTexture2DRHIParamRef Texture, uint32 MipIndex, bool bLockWithinMiptail) final override;
-	virtual void* RHILockTexture2DArray(FTexture2DArrayRHIParamRef Texture, uint32 TextureIndex, uint32 MipIndex, EResourceLockMode LockMode, uint32& DestStride, bool bLockWithinMiptail) final override;
-	virtual void RHIUnlockTexture2DArray(FTexture2DArrayRHIParamRef Texture, uint32 TextureIndex, uint32 MipIndex, bool bLockWithinMiptail) final override;
-	virtual void RHIUpdateTexture2D(FTexture2DRHIParamRef Texture, uint32 MipIndex, const struct FUpdateTextureRegion2D& UpdateRegion, uint32 SourcePitch, const uint8* SourceData) final override;
-	virtual void RHIUpdateTexture3D(FTexture3DRHIParamRef Texture, uint32 MipIndex, const struct FUpdateTextureRegion3D& UpdateRegion, uint32 SourceRowPitch, uint32 SourceDepthPitch, const uint8* SourceData) final override;
+	virtual void RHIGetResourceInfo(FRHITexture* Ref, FRHIResourceInfo& OutInfo) final override;
+	virtual FShaderResourceViewRHIRef RHICreateShaderResourceView(FRHITexture* Texture, const FRHITextureSRVCreateInfo& CreateInfo) final override;
+	virtual void RHIGenerateMips(FRHITexture* Texture) final override;
+	virtual uint32 RHIComputeMemorySize(FRHITexture* TextureRHI) final override;
+	virtual FTexture2DRHIRef RHIAsyncReallocateTexture2D(FRHITexture2D* Texture2D, int32 NewMipCount, int32 NewSizeX, int32 NewSizeY, FThreadSafeCounter* RequestStatus) final override;
+	virtual ETextureReallocationStatus RHIFinalizeAsyncReallocateTexture2D(FRHITexture2D* Texture2D, bool bBlockUntilCompleted) final override;
+	virtual ETextureReallocationStatus RHICancelAsyncReallocateTexture2D(FRHITexture2D* Texture2D, bool bBlockUntilCompleted) final override;
+	virtual void* RHILockTexture2D(FRHITexture2D* Texture, uint32 MipIndex, EResourceLockMode LockMode, uint32& DestStride, bool bLockWithinMiptail) final override;
+	virtual void RHIUnlockTexture2D(FRHITexture2D* Texture, uint32 MipIndex, bool bLockWithinMiptail) final override;
+	virtual void* RHILockTexture2DArray(FRHITexture2DArray* Texture, uint32 TextureIndex, uint32 MipIndex, EResourceLockMode LockMode, uint32& DestStride, bool bLockWithinMiptail) final override;
+	virtual void RHIUnlockTexture2DArray(FRHITexture2DArray* Texture, uint32 TextureIndex, uint32 MipIndex, bool bLockWithinMiptail) final override;
+	virtual void RHIUpdateTexture2D(FRHITexture2D* Texture, uint32 MipIndex, const struct FUpdateTextureRegion2D& UpdateRegion, uint32 SourcePitch, const uint8* SourceData) final override;
+	virtual void RHIUpdateTexture3D(FRHITexture3D* Texture, uint32 MipIndex, const struct FUpdateTextureRegion3D& UpdateRegion, uint32 SourceRowPitch, uint32 SourceDepthPitch, const uint8* SourceData) final override;
 	virtual FTextureCubeRHIRef RHICreateTextureCube(uint32 Size, uint8 Format, uint32 NumMips, uint32 Flags, FRHIResourceCreateInfo& CreateInfo) final override;
 	virtual FTextureCubeRHIRef RHICreateTextureCubeArray(uint32 Size, uint32 ArraySize, uint8 Format, uint32 NumMips, uint32 Flags, FRHIResourceCreateInfo& CreateInfo) final override;
-	virtual void* RHILockTextureCubeFace(FTextureCubeRHIParamRef Texture, uint32 FaceIndex, uint32 ArrayIndex, uint32 MipIndex, EResourceLockMode LockMode, uint32& DestStride, bool bLockWithinMiptail) final override;
-	virtual void RHIUnlockTextureCubeFace(FTextureCubeRHIParamRef Texture, uint32 FaceIndex, uint32 ArrayIndex, uint32 MipIndex, bool bLockWithinMiptail) final override;
-	virtual void RHIBindDebugLabelName(FTextureRHIParamRef Texture, const TCHAR* Name) final override;
-	virtual void RHIReadSurfaceData(FTextureRHIParamRef Texture,FIntRect Rect,TArray<FColor>& OutData,FReadSurfaceDataFlags InFlags) final override;
-	virtual void RHIReadSurfaceData(FTextureRHIParamRef Texture, FIntRect Rect, TArray<FLinearColor>& OutData, FReadSurfaceDataFlags InFlags) final override;
-	virtual void RHIMapStagingSurface(FTextureRHIParamRef Texture,void*& OutData,int32& OutWidth,int32& OutHeight) final override;
-	virtual void RHIUnmapStagingSurface(FTextureRHIParamRef Texture) final override;
-	virtual void RHIReadSurfaceFloatData(FTextureRHIParamRef Texture,FIntRect Rect,TArray<FFloat16Color>& OutData,ECubeFace CubeFace,int32 ArrayIndex,int32 MipIndex) final override;
-	virtual void RHIRead3DSurfaceFloatData(FTextureRHIParamRef Texture,FIntRect Rect,FIntPoint ZMinMax,TArray<FFloat16Color>& OutData) final override;
+	virtual void* RHILockTextureCubeFace(FRHITextureCube* Texture, uint32 FaceIndex, uint32 ArrayIndex, uint32 MipIndex, EResourceLockMode LockMode, uint32& DestStride, bool bLockWithinMiptail) final override;
+	virtual void RHIUnlockTextureCubeFace(FRHITextureCube* Texture, uint32 FaceIndex, uint32 ArrayIndex, uint32 MipIndex, bool bLockWithinMiptail) final override;
+	virtual void RHIBindDebugLabelName(FRHITexture* Texture, const TCHAR* Name) final override;
+	virtual void RHIReadSurfaceData(FRHITexture* Texture,FIntRect Rect,TArray<FColor>& OutData,FReadSurfaceDataFlags InFlags) final override;
+	virtual void RHIReadSurfaceData(FRHITexture* Texture, FIntRect Rect, TArray<FLinearColor>& OutData, FReadSurfaceDataFlags InFlags) final override;
+	virtual void RHIMapStagingSurface(FRHITexture* Texture,void*& OutData,int32& OutWidth,int32& OutHeight) final override;
+	virtual void RHIUnmapStagingSurface(FRHITexture* Texture) final override;
+	virtual void RHIReadSurfaceFloatData(FRHITexture* Texture,FIntRect Rect,TArray<FFloat16Color>& OutData,ECubeFace CubeFace,int32 ArrayIndex,int32 MipIndex) final override;
+	virtual void RHIRead3DSurfaceFloatData(FRHITexture* Texture,FIntRect Rect,FIntPoint ZMinMax,TArray<FFloat16Color>& OutData) final override;
 	virtual FRenderQueryRHIRef RHICreateRenderQuery(ERenderQueryType QueryType) final override;
 	virtual bool RHIGetRenderQueryResult(FRHIRenderQuery* RenderQuery, uint64& OutResult, bool bWait) final override;
 	virtual FTexture2DRHIRef RHIGetViewportBackBuffer(FRHIViewport* Viewport) final override;
-	virtual void RHIAliasTextureResources(FTextureRHIParamRef DestTexture, FTextureRHIParamRef SrcTexture) final override;
+	virtual void RHIAliasTextureResources(FRHITexture* DestTexture, FRHITexture* SrcTexture) final override;
 	virtual void RHIAdvanceFrameForGetViewportBackBuffer(FRHIViewport* Viewport) final override;
 	virtual void RHIAcquireThreadOwnership() final override;
 	virtual void RHIReleaseThreadOwnership() final override;
@@ -430,14 +430,14 @@ public:
 	virtual FViewportRHIRef RHICreateViewport(void* WindowHandle, uint32 SizeX, uint32 SizeY, bool bIsFullscreen, EPixelFormat PreferredPixelFormat) final override;
 	virtual void RHIResizeViewport(FRHIViewport* Viewport, uint32 SizeX, uint32 SizeY, bool bIsFullscreen) final override;
 	virtual void RHITick(float DeltaTime) final override;
-	virtual void RHISetStreamOutTargets(uint32 NumTargets,const FVertexBufferRHIParamRef* VertexBuffers,const uint32* Offsets) final override;
+	virtual void RHISetStreamOutTargets(uint32 NumTargets, FRHIVertexBuffer* const* VertexBuffers,const uint32* Offsets) final override;
 	virtual void RHIBlockUntilGPUIdle() final override;
 	virtual void RHISubmitCommandsAndFlushGPU() final override;
 	virtual void RHIPollOcclusionQueries() final override;
 	virtual bool RHIGetAvailableResolutions(FScreenResolutionArray& Resolutions, bool bIgnoreRefreshRate) final override;
 	virtual void RHIGetSupportedResolution(uint32& Width, uint32& Height) final override;
-	virtual void RHIVirtualTextureSetFirstMipInMemory(FTexture2DRHIParamRef Texture, uint32 FirstMip) final override;
-	virtual void RHIVirtualTextureSetFirstMipVisible(FTexture2DRHIParamRef Texture, uint32 FirstMip) final override;
+	virtual void RHIVirtualTextureSetFirstMipInMemory(FRHITexture2D* Texture, uint32 FirstMip) final override;
+	virtual void RHIVirtualTextureSetFirstMipVisible(FRHITexture2D* Texture, uint32 FirstMip) final override;
 	virtual void RHIExecuteCommandList(FRHICommandList* CmdList) final override;
 	virtual void* RHIGetNativeDevice() final override;
 	virtual class IRHICommandContext* RHIGetDefaultContext() final override;
@@ -448,32 +448,32 @@ public:
 
 	virtual void RHISetComputeShader(FRHIComputeShader* ComputeShader) final override;
 	virtual void RHIDispatchComputeShader(uint32 ThreadGroupCountX, uint32 ThreadGroupCountY, uint32 ThreadGroupCountZ) final override;
-	virtual void RHIDispatchIndirectComputeShader(FVertexBufferRHIParamRef ArgumentBuffer, uint32 ArgumentOffset) final override;
+	virtual void RHIDispatchIndirectComputeShader(FRHIVertexBuffer* ArgumentBuffer, uint32 ArgumentOffset) final override;
 	virtual void RHIAutomaticCacheFlushAfterComputeShader(bool bEnable) final override;
 	virtual void RHIFlushComputeShaderCache() final override;
 	virtual void RHISetMultipleViewports(uint32 Count, const FViewportBounds* Data) final override;
 	virtual void RHIClearTinyUAV(FRHIUnorderedAccessView* UnorderedAccessViewRHI, const uint32* Values) final override;
-	virtual void RHICopyToResolveTarget(FTextureRHIParamRef SourceTexture, FTextureRHIParamRef DestTexture, const FResolveParams& ResolveParams) final override;
+	virtual void RHICopyToResolveTarget(FRHITexture* SourceTexture, FRHITexture* DestTexture, const FResolveParams& ResolveParams) final override;
 	virtual void RHIBeginRenderQuery(FRHIRenderQuery* RenderQuery) final override;
 	virtual void RHIEndRenderQuery(FRHIRenderQuery* RenderQuery) final override;
 	virtual void RHISubmitCommandsHint() final override;
-	virtual void RHIBeginDrawingViewport(FRHIViewport* Viewport, FTextureRHIParamRef RenderTargetRHI) final override;
+	virtual void RHIBeginDrawingViewport(FRHIViewport* Viewport, FRHITexture* RenderTargetRHI) final override;
 	virtual void RHIEndDrawingViewport(FRHIViewport* Viewport, bool bPresent, bool bLockToVsync) final override;
 	virtual void RHIBeginFrame() final override;
 	virtual void RHIEndFrame() final override;
 	virtual void RHIBeginScene() final override;
 	virtual void RHIEndScene() final override;
-	virtual void RHISetStreamSource(uint32 StreamIndex, FVertexBufferRHIParamRef VertexBuffer, uint32 Offset) final override;
+	virtual void RHISetStreamSource(uint32 StreamIndex, FRHIVertexBuffer* VertexBuffer, uint32 Offset) final override;
 	virtual void RHISetRasterizerState(FRHIRasterizerState* NewState) final override;
 	virtual void RHISetViewport(uint32 MinX, uint32 MinY, float MinZ, uint32 MaxX, uint32 MaxY, float MaxZ) final override;
 	virtual void RHISetScissorRect(bool bEnable, uint32 MinX, uint32 MinY, uint32 MaxX, uint32 MaxY) final override;
 	virtual void RHISetBoundShaderState(FRHIBoundShaderState* BoundShaderState) final override;
-	virtual void RHISetShaderTexture(FRHIVertexShader* VertexShader, uint32 TextureIndex, FTextureRHIParamRef NewTexture) final override;
-	virtual void RHISetShaderTexture(FRHIHullShader* HullShader, uint32 TextureIndex, FTextureRHIParamRef NewTexture) final override;
-	virtual void RHISetShaderTexture(FRHIDomainShader* DomainShader, uint32 TextureIndex, FTextureRHIParamRef NewTexture) final override;
-	virtual void RHISetShaderTexture(FRHIGeometryShader* GeometryShader, uint32 TextureIndex, FTextureRHIParamRef NewTexture) final override;
-	virtual void RHISetShaderTexture(FRHIPixelShader* PixelShader, uint32 TextureIndex, FTextureRHIParamRef NewTexture) final override;
-	virtual void RHISetShaderTexture(FRHIComputeShader* PixelShader, uint32 TextureIndex, FTextureRHIParamRef NewTexture) final override;
+	virtual void RHISetShaderTexture(FRHIVertexShader* VertexShader, uint32 TextureIndex, FRHITexture* NewTexture) final override;
+	virtual void RHISetShaderTexture(FRHIHullShader* HullShader, uint32 TextureIndex, FRHITexture* NewTexture) final override;
+	virtual void RHISetShaderTexture(FRHIDomainShader* DomainShader, uint32 TextureIndex, FRHITexture* NewTexture) final override;
+	virtual void RHISetShaderTexture(FRHIGeometryShader* GeometryShader, uint32 TextureIndex, FRHITexture* NewTexture) final override;
+	virtual void RHISetShaderTexture(FRHIPixelShader* PixelShader, uint32 TextureIndex, FRHITexture* NewTexture) final override;
+	virtual void RHISetShaderTexture(FRHIComputeShader* PixelShader, uint32 TextureIndex, FRHITexture* NewTexture) final override;
 	virtual void RHISetShaderSampler(FRHIComputeShader* ComputeShader, uint32 SamplerIndex, FRHISamplerState* NewState) final override;
 	virtual void RHISetShaderSampler(FRHIVertexShader* VertexShader, uint32 SamplerIndex, FRHISamplerState* NewState) final override;
 	virtual void RHISetShaderSampler(FRHIGeometryShader* GeometryShader, uint32 SamplerIndex, FRHISamplerState* NewState) final override;
@@ -488,12 +488,12 @@ public:
 	virtual void RHISetShaderResourceViewParameter(FRHIHullShader* HullShader, uint32 SamplerIndex, FRHIShaderResourceView* SRV) final override;
 	virtual void RHISetShaderResourceViewParameter(FRHIDomainShader* DomainShader, uint32 SamplerIndex, FRHIShaderResourceView* SRV) final override;
 	virtual void RHISetShaderResourceViewParameter(FRHIGeometryShader* GeometryShader, uint32 SamplerIndex, FRHIShaderResourceView* SRV) final override;
-	virtual void RHISetShaderUniformBuffer(FRHIVertexShader* VertexShader, uint32 BufferIndex, FUniformBufferRHIParamRef Buffer) final override;
-	virtual void RHISetShaderUniformBuffer(FRHIHullShader* HullShader, uint32 BufferIndex, FUniformBufferRHIParamRef Buffer) final override;
-	virtual void RHISetShaderUniformBuffer(FRHIDomainShader* DomainShader, uint32 BufferIndex, FUniformBufferRHIParamRef Buffer) final override;
-	virtual void RHISetShaderUniformBuffer(FRHIGeometryShader* GeometryShader, uint32 BufferIndex, FUniformBufferRHIParamRef Buffer) final override;
-	virtual void RHISetShaderUniformBuffer(FRHIPixelShader* PixelShader, uint32 BufferIndex, FUniformBufferRHIParamRef Buffer) final override;
-	virtual void RHISetShaderUniformBuffer(FRHIComputeShader* ComputeShader, uint32 BufferIndex, FUniformBufferRHIParamRef Buffer) final override;
+	virtual void RHISetShaderUniformBuffer(FRHIVertexShader* VertexShader, uint32 BufferIndex, FRHIUniformBuffer* Buffer) final override;
+	virtual void RHISetShaderUniformBuffer(FRHIHullShader* HullShader, uint32 BufferIndex, FRHIUniformBuffer* Buffer) final override;
+	virtual void RHISetShaderUniformBuffer(FRHIDomainShader* DomainShader, uint32 BufferIndex, FRHIUniformBuffer* Buffer) final override;
+	virtual void RHISetShaderUniformBuffer(FRHIGeometryShader* GeometryShader, uint32 BufferIndex, FRHIUniformBuffer* Buffer) final override;
+	virtual void RHISetShaderUniformBuffer(FRHIPixelShader* PixelShader, uint32 BufferIndex, FRHIUniformBuffer* Buffer) final override;
+	virtual void RHISetShaderUniformBuffer(FRHIComputeShader* ComputeShader, uint32 BufferIndex, FRHIUniformBuffer* Buffer) final override;
 	virtual void RHISetShaderParameter(FRHIVertexShader* VertexShader, uint32 BufferIndex, uint32 BaseIndex, uint32 NumBytes, const void* NewValue) final override;
 	virtual void RHISetShaderParameter(FRHIPixelShader* PixelShader, uint32 BufferIndex, uint32 BaseIndex, uint32 NumBytes, const void* NewValue) final override;
 	virtual void RHISetShaderParameter(FRHIHullShader* HullShader, uint32 BufferIndex, uint32 BaseIndex, uint32 NumBytes, const void* NewValue) final override;
@@ -510,10 +510,10 @@ public:
 	virtual void RHISetRenderTargets(uint32 NumSimultaneousRenderTargets, const FRHIRenderTargetView* NewRenderTargets, const FRHIDepthRenderTargetView* NewDepthStencilTarget, uint32 NumUAVs, FRHIUnorderedAccessView* const* UAVs) final override;
 	virtual void RHISetRenderTargetsAndClear(const FRHISetRenderTargetsInfo& RenderTargetsInfo) final override;
 	virtual void RHIDrawPrimitive(uint32 BaseVertexIndex, uint32 NumPrimitives, uint32 NumInstances) final override;
-	virtual void RHIDrawPrimitiveIndirect(FVertexBufferRHIParamRef ArgumentBuffer, uint32 ArgumentOffset) final override;
-	virtual void RHIDrawIndexedIndirect(FIndexBufferRHIParamRef IndexBufferRHI, FStructuredBufferRHIParamRef ArgumentsBufferRHI, int32 DrawArgumentsIndex, uint32 NumInstances) final override;
-	virtual void RHIDrawIndexedPrimitive(FIndexBufferRHIParamRef IndexBuffer, int32 BaseVertexIndex, uint32 FirstInstance, uint32 NumVertices, uint32 StartIndex, uint32 NumPrimitives, uint32 NumInstances) final override;
-	virtual void RHIDrawIndexedPrimitiveIndirect(FIndexBufferRHIParamRef IndexBuffer, FVertexBufferRHIParamRef ArgumentBuffer, uint32 ArgumentOffset) final override;
+	virtual void RHIDrawPrimitiveIndirect(FRHIVertexBuffer* ArgumentBuffer, uint32 ArgumentOffset) final override;
+	virtual void RHIDrawIndexedIndirect(FRHIIndexBuffer* IndexBufferRHI, FRHIStructuredBuffer* ArgumentsBufferRHI, int32 DrawArgumentsIndex, uint32 NumInstances) final override;
+	virtual void RHIDrawIndexedPrimitive(FRHIIndexBuffer* IndexBuffer, int32 BaseVertexIndex, uint32 FirstInstance, uint32 NumVertices, uint32 StartIndex, uint32 NumPrimitives, uint32 NumInstances) final override;
+	virtual void RHIDrawIndexedPrimitiveIndirect(FRHIIndexBuffer* IndexBuffer, FRHIVertexBuffer* ArgumentBuffer, uint32 ArgumentOffset) final override;
 	virtual void RHIEnableDepthBoundsTest(bool bEnable) final override;
 	virtual void RHISetDepthBounds(float MinDepth, float MaxDepth) final override;
 	virtual void RHIPushEvent(const TCHAR* Name, FColor Color) final override;
@@ -521,8 +521,8 @@ public:
 	virtual void RHIInvalidateCachedState() final override;
 	virtual void RHIDiscardRenderTargets(bool Depth,bool Stencil,uint32 ColorBitMask) final override;
 
-	virtual void RHICopySubTextureRegion(FTexture2DRHIParamRef SourceTexture, FTexture2DRHIParamRef DestinationTexture, FBox2D SourceBox, FBox2D DestinationBox) final override;
-	virtual void RHICopyTexture(FTextureRHIParamRef SourceTexture, FTextureRHIParamRef DestTexture, const FRHICopyTextureInfo& CopyInfo) final override;
+	virtual void RHICopySubTextureRegion(FRHITexture2D* SourceTexture, FRHITexture2D* DestinationTexture, FBox2D SourceBox, FBox2D DestinationBox) final override;
+	virtual void RHICopyTexture(FRHITexture* SourceTexture, FRHITexture* DestTexture, const FRHICopyTextureInfo& CopyInfo) final override;
 
 	// FOpenGLDynamicRHI interface.
 	virtual FTexture2DRHIRef RHICreateTexture2DFromResource(EPixelFormat Format, uint32 SizeX, uint32 SizeY, uint32 NumMips, uint32 NumSamples, uint32 NumSamplesTileMem, const FClearValueBinding& ClearValueBinding, GLuint Resource, uint32 Flags);
@@ -530,7 +530,7 @@ public:
 	virtual FTextureCubeRHIRef RHICreateTextureCubeFromResource(EPixelFormat Format, uint32 Size, bool bArray, uint32 ArraySize, uint32 NumMips, uint32 NumSamples, uint32 NumSamplesTileMem, const FClearValueBinding& ClearValueBinding, GLuint Resource, uint32 Flags);
 
 	// Inline copy
-	virtual void RHICopyToStagingBuffer(FVertexBufferRHIParamRef SourceBufferRHI, FRHIStagingBuffer* DestinationStagingBufferRHI, uint32 InOffset, uint32 InNumBytes) final override;
+	virtual void RHICopyToStagingBuffer(FRHIVertexBuffer* SourceBufferRHI, FRHIStagingBuffer* DestinationStagingBufferRHI, uint32 InOffset, uint32 InNumBytes) final override;
 	virtual void RHIWriteGPUFence(FRHIGPUFence* FenceRHI) final override;
 	virtual void* RHILockStagingBuffer(FRHIStagingBuffer* StagingBuffer, uint32 Offset, uint32 SizeRHI) final override;
 	virtual void RHIUnlockStagingBuffer(FRHIStagingBuffer* StagingBuffer) final override;
@@ -764,10 +764,10 @@ public:
 		}
 	};
 
-	virtual void* LockTexture2D_RenderThread(class FRHICommandListImmediate& RHICmdList, FTexture2DRHIParamRef Texture, uint32 MipIndex, EResourceLockMode LockMode, uint32& DestStride, bool bLockWithinMiptail, bool bNeedsDefaultRHIFlush) final override;
-	virtual void UnlockTexture2D_RenderThread(class FRHICommandListImmediate& RHICmdList, FTexture2DRHIParamRef Texture, uint32 MipIndex, bool bLockWithinMiptail, bool bNeedsDefaultRHIFlush) final override;
-	virtual void* RHILockTextureCubeFace_RenderThread(class FRHICommandListImmediate& RHICmdList, FTextureCubeRHIParamRef Texture, uint32 FaceIndex, uint32 ArrayIndex, uint32 MipIndex, EResourceLockMode LockMode, uint32& DestStride, bool bLockWithinMiptail) final override;
-	virtual void RHIUnlockTextureCubeFace_RenderThread(class FRHICommandListImmediate& RHICmdList, FTextureCubeRHIParamRef Texture, uint32 FaceIndex, uint32 ArrayIndex, uint32 MipIndex, bool bLockWithinMiptail) final override;
+	virtual void* LockTexture2D_RenderThread(class FRHICommandListImmediate& RHICmdList, FRHITexture2D* Texture, uint32 MipIndex, EResourceLockMode LockMode, uint32& DestStride, bool bLockWithinMiptail, bool bNeedsDefaultRHIFlush) final override;
+	virtual void UnlockTexture2D_RenderThread(class FRHICommandListImmediate& RHICmdList, FRHITexture2D* Texture, uint32 MipIndex, bool bLockWithinMiptail, bool bNeedsDefaultRHIFlush) final override;
+	virtual void* RHILockTextureCubeFace_RenderThread(class FRHICommandListImmediate& RHICmdList, FRHITextureCube* Texture, uint32 FaceIndex, uint32 ArrayIndex, uint32 MipIndex, EResourceLockMode LockMode, uint32& DestStride, bool bLockWithinMiptail) final override;
+	virtual void RHIUnlockTextureCubeFace_RenderThread(class FRHICommandListImmediate& RHICmdList, FRHITextureCube* Texture, uint32 FaceIndex, uint32 ArrayIndex, uint32 MipIndex, bool bLockWithinMiptail) final override;
 
 	virtual FIndexBufferRHIRef CreateIndexBuffer_RenderThread(class FRHICommandListImmediate& RHICmdList, uint32 Stride, uint32 Size, uint32 InUsage, FRHIResourceCreateInfo& CreateInfo) final override;
 
@@ -783,12 +783,12 @@ public:
 		RHITHREAD_GLCOMMAND_EPILOGUE_RETURN(FStructuredBufferRHIRef);
 	}
 
-	virtual FShaderResourceViewRHIRef CreateShaderResourceView_RenderThread(class FRHICommandListImmediate& RHICmdList, FVertexBufferRHIParamRef VertexBuffer, uint32 Stride, uint8 Format) final override
+	virtual FShaderResourceViewRHIRef CreateShaderResourceView_RenderThread(class FRHICommandListImmediate& RHICmdList, FRHIVertexBuffer* VertexBuffer, uint32 Stride, uint8 Format) final override
 	{
 		return this->RHICreateShaderResourceView(VertexBuffer, Stride, Format);
 	}
 
-	virtual FShaderResourceViewRHIRef CreateShaderResourceView_RenderThread(class FRHICommandListImmediate& RHICmdList, FIndexBufferRHIParamRef Buffer) final override
+	virtual FShaderResourceViewRHIRef CreateShaderResourceView_RenderThread(class FRHICommandListImmediate& RHICmdList, FRHIIndexBuffer* Buffer) final override
 	{
 		return this->RHICreateShaderResourceView(Buffer);
 	}
@@ -839,43 +839,43 @@ public:
 		RHITHREAD_GLCOMMAND_EPILOGUE_RETURN(FTexture3DRHIRef);
 	}
 
-	virtual FUnorderedAccessViewRHIRef RHICreateUnorderedAccessView_RenderThread(class FRHICommandListImmediate& RHICmdList, FStructuredBufferRHIParamRef StructuredBuffer, bool bUseUAVCounter, bool bAppendBuffer) final override
+	virtual FUnorderedAccessViewRHIRef RHICreateUnorderedAccessView_RenderThread(class FRHICommandListImmediate& RHICmdList, FRHIStructuredBuffer* StructuredBuffer, bool bUseUAVCounter, bool bAppendBuffer) final override
 	{
 		RHITHREAD_GLCOMMAND_PROLOGUE();
 		return this->RHICreateUnorderedAccessView(StructuredBuffer, bUseUAVCounter, bAppendBuffer);
 		RHITHREAD_GLCOMMAND_EPILOGUE_RETURN(FUnorderedAccessViewRHIRef);
 	}
 
-	virtual FUnorderedAccessViewRHIRef RHICreateUnorderedAccessView_RenderThread(class FRHICommandListImmediate& RHICmdList, FTextureRHIParamRef Texture, uint32 MipLevel) final override
+	virtual FUnorderedAccessViewRHIRef RHICreateUnorderedAccessView_RenderThread(class FRHICommandListImmediate& RHICmdList, FRHITexture* Texture, uint32 MipLevel) final override
 	{
 		RHITHREAD_GLCOMMAND_PROLOGUE();
 		return this->RHICreateUnorderedAccessView(Texture, MipLevel);
 		RHITHREAD_GLCOMMAND_EPILOGUE_RETURN(FUnorderedAccessViewRHIRef);
 	}
 
-	virtual FUnorderedAccessViewRHIRef RHICreateUnorderedAccessView_RenderThread(class FRHICommandListImmediate& RHICmdList, FVertexBufferRHIParamRef VertexBuffer, uint8 Format) final override
+	virtual FUnorderedAccessViewRHIRef RHICreateUnorderedAccessView_RenderThread(class FRHICommandListImmediate& RHICmdList, FRHIVertexBuffer* VertexBuffer, uint8 Format) final override
 	{
 		RHITHREAD_GLCOMMAND_PROLOGUE();
 		return this->RHICreateUnorderedAccessView(VertexBuffer, Format);
 		RHITHREAD_GLCOMMAND_EPILOGUE_RETURN(FUnorderedAccessViewRHIRef);
 	}
 
-	virtual FShaderResourceViewRHIRef RHICreateShaderResourceView_RenderThread(class FRHICommandListImmediate& RHICmdList, FTextureRHIParamRef Texture, const FRHITextureSRVCreateInfo& CreateInfo) final override
+	virtual FShaderResourceViewRHIRef RHICreateShaderResourceView_RenderThread(class FRHICommandListImmediate& RHICmdList, FRHITexture* Texture, const FRHITextureSRVCreateInfo& CreateInfo) final override
 	{
 		return this->RHICreateShaderResourceView(Texture, CreateInfo);
 	}
 
-	virtual FShaderResourceViewRHIRef RHICreateShaderResourceView_RenderThread(class FRHICommandListImmediate& RHICmdList, FVertexBufferRHIParamRef VertexBuffer, uint32 Stride, uint8 Format) final override
+	virtual FShaderResourceViewRHIRef RHICreateShaderResourceView_RenderThread(class FRHICommandListImmediate& RHICmdList, FRHIVertexBuffer* VertexBuffer, uint32 Stride, uint8 Format) final override
 	{
 		return this->RHICreateShaderResourceView(VertexBuffer, Stride, Format);
 	}
 
-	virtual FShaderResourceViewRHIRef RHICreateShaderResourceView_RenderThread(class FRHICommandListImmediate& RHICmdList, FIndexBufferRHIParamRef Buffer) final override
+	virtual FShaderResourceViewRHIRef RHICreateShaderResourceView_RenderThread(class FRHICommandListImmediate& RHICmdList, FRHIIndexBuffer* Buffer) final override
 	{
 		return this->RHICreateShaderResourceView(Buffer);
 	}
 
-	virtual FShaderResourceViewRHIRef RHICreateShaderResourceView_RenderThread(class FRHICommandListImmediate& RHICmdList, FStructuredBufferRHIParamRef StructuredBuffer) final override
+	virtual FShaderResourceViewRHIRef RHICreateShaderResourceView_RenderThread(class FRHICommandListImmediate& RHICmdList, FRHIStructuredBuffer* StructuredBuffer) final override
 	{
 		return this->RHICreateShaderResourceView(StructuredBuffer);
 	}
@@ -911,14 +911,14 @@ public:
 		return this->RHICreateRenderQuery(QueryType);
 	}
 
-	virtual FTexture2DRHIRef AsyncReallocateTexture2D_RenderThread(class FRHICommandListImmediate& RHICmdList, FTexture2DRHIParamRef Texture2D, int32 NewMipCount, int32 NewSizeX, int32 NewSizeY, FThreadSafeCounter* RequestStatus) final override;
+	virtual FTexture2DRHIRef AsyncReallocateTexture2D_RenderThread(class FRHICommandListImmediate& RHICmdList, FRHITexture2D* Texture2D, int32 NewMipCount, int32 NewSizeX, int32 NewSizeY, FThreadSafeCounter* RequestStatus) final override;
 
-	virtual ETextureReallocationStatus FinalizeAsyncReallocateTexture2D_RenderThread(class FRHICommandListImmediate& RHICmdList, FTexture2DRHIParamRef Texture2D, bool bBlockUntilCompleted) final override
+	virtual ETextureReallocationStatus FinalizeAsyncReallocateTexture2D_RenderThread(class FRHICommandListImmediate& RHICmdList, FRHITexture2D* Texture2D, bool bBlockUntilCompleted) final override
 	{
 		return this->RHIFinalizeAsyncReallocateTexture2D(Texture2D, bBlockUntilCompleted);
 	}
 
-	virtual ETextureReallocationStatus CancelAsyncReallocateTexture2D_RenderThread(class FRHICommandListImmediate& RHICmdList, FTexture2DRHIParamRef Texture2D, bool bBlockUntilCompleted) final override
+	virtual ETextureReallocationStatus CancelAsyncReallocateTexture2D_RenderThread(class FRHICommandListImmediate& RHICmdList, FRHITexture2D* Texture2D, bool bBlockUntilCompleted) final override
 	{
 		return this->RHICancelAsyncReallocateTexture2D(Texture2D, bBlockUntilCompleted);
 	}
@@ -955,12 +955,12 @@ public:
 
 	virtual FSamplerStateRHIRef RHICreateSamplerState(const FSamplerStateInitializerRHI& Initializer) final override;
 
-	virtual void UpdateTexture2D_RenderThread(class FRHICommandListImmediate& RHICmdList, FTexture2DRHIParamRef Texture, uint32 MipIndex, const struct FUpdateTextureRegion2D& UpdateRegion, uint32 SourcePitch, const uint8* SourceData) final override
+	virtual void UpdateTexture2D_RenderThread(class FRHICommandListImmediate& RHICmdList, FRHITexture2D* Texture, uint32 MipIndex, const struct FUpdateTextureRegion2D& UpdateRegion, uint32 SourcePitch, const uint8* SourceData) final override
 	{
 		GDynamicRHI->RHIUpdateTexture2D(Texture, MipIndex, UpdateRegion, SourcePitch, SourceData);
 	}
 
-	virtual void UpdateTexture3D_RenderThread(class FRHICommandListImmediate& RHICmdList, FTexture3DRHIParamRef Texture, uint32 MipIndex, const struct FUpdateTextureRegion3D& UpdateRegion, uint32 SourceRowPitch, uint32 SourceDepthPitch, const uint8* SourceData) final override
+	virtual void UpdateTexture3D_RenderThread(class FRHICommandListImmediate& RHICmdList, FRHITexture3D* Texture, uint32 MipIndex, const struct FUpdateTextureRegion3D& UpdateRegion, uint32 SourceRowPitch, uint32 SourceDepthPitch, const uint8* SourceData) final override
 	{
 		RHITHREAD_GLCOMMAND_PROLOGUE();
 		GDynamicRHI->RHIUpdateTexture3D(Texture, MipIndex, UpdateRegion, SourceRowPitch, SourceDepthPitch, SourceData);
@@ -1162,7 +1162,7 @@ private:
 
 	void ApplyTextureStage(FOpenGLContextState& ContextState, GLint TextureIndex, const FTextureStage& TextureStage, FOpenGLSamplerState* SamplerState);
 
-	void ReadSurfaceDataRaw(FOpenGLContextState& ContextState, FTextureRHIParamRef TextureRHI,FIntRect Rect,TArray<uint8>& OutData, FReadSurfaceDataFlags InFlags);
+	void ReadSurfaceDataRaw(FOpenGLContextState& ContextState, FRHITexture* TextureRHI,FIntRect Rect,TArray<uint8>& OutData, FReadSurfaceDataFlags InFlags);
 
 	void BindUniformBufferBase(FOpenGLContextState& ContextState, int32 NumUniformBuffers, FUniformBufferRHIRef* BoundUniformBuffers, uint32 FirstUniformBuffer, bool ForceUpdate);
 
diff --git a/Engine/Source/Runtime/OpenGLDrv/Public/OpenGLResources.h b/Engine/Source/Runtime/OpenGLDrv/Public/OpenGLResources.h
--- a/Engine/Source/Runtime/OpenGLDrv/Public/OpenGLResources.h
+++ b/Engine/Source/Runtime/OpenGLDrv/Public/OpenGLResources.h
@@ -1818,7 +1818,7 @@ class FOpenGLTextureUnorderedAccessView : public FOpenGLUnorderedAccessView
 {
 public:
 
-	FOpenGLTextureUnorderedAccessView(FTextureRHIParamRef InTexture);
+	FOpenGLTextureUnorderedAccessView(FRHITexture* InTexture);
 
 	FTextureRHIRef TextureRHI; // to keep the texture alive
 };
@@ -1830,7 +1830,7 @@ public:
 
 	FOpenGLVertexBufferUnorderedAccessView();
 
-	FOpenGLVertexBufferUnorderedAccessView(	FOpenGLDynamicRHI* InOpenGLRHI, FVertexBufferRHIParamRef InVertexBuffer, uint8 Format);
+	FOpenGLVertexBufferUnorderedAccessView(	FOpenGLDynamicRHI* InOpenGLRHI, FRHIVertexBuffer* InVertexBuffer, uint8 Format);
 
 	virtual ~FOpenGLVertexBufferUnorderedAccessView();
 
@@ -1846,7 +1846,7 @@ class FOpenGLStructuredBufferUnorderedAccessView : public FOpenGLUnorderedAccess
 public:
 	FOpenGLStructuredBufferUnorderedAccessView();
 
-	FOpenGLStructuredBufferUnorderedAccessView(	FOpenGLDynamicRHI* InOpenGLRHI, FStructuredBufferRHIParamRef InBuffer, uint8 Format);
+	FOpenGLStructuredBufferUnorderedAccessView(	FOpenGLDynamicRHI* InOpenGLRHI, FRHIStructuredBuffer* InBuffer, uint8 Format);
 
 	virtual ~FOpenGLStructuredBufferUnorderedAccessView();
 
@@ -1888,7 +1888,7 @@ public:
 	,	OwnsResource(true)
 	{}
 
-	FOpenGLShaderResourceView( FOpenGLDynamicRHI* InOpenGLRHI, GLuint InResource, GLenum InTarget, FVertexBufferRHIParamRef InVertexBuffer, uint8 InFormat )
+	FOpenGLShaderResourceView( FOpenGLDynamicRHI* InOpenGLRHI, GLuint InResource, GLenum InTarget, FRHIVertexBuffer* InVertexBuffer, uint8 InFormat )
 	:	Resource(InResource)
 	,	Target(InTarget)
 	,	LimitMip(-1)
@@ -1951,7 +1951,7 @@ struct TIsGLProxyObject<FOpenGLShaderResourceViewProxy>
 void OPENGLDRV_API OpenGLTextureDeleted(FRHITexture* Texture);
 void OPENGLDRV_API OpenGLTextureAllocated( FRHITexture* Texture , uint32 Flags);
 
-void OPENGLDRV_API ReleaseOpenGLFramebuffers(FOpenGLDynamicRHI* Device, FTextureRHIParamRef TextureRHI);
+void OPENGLDRV_API ReleaseOpenGLFramebuffers(FOpenGLDynamicRHI* Device, FRHITexture* TextureRHI);
 
 /** A OpenGL event query resource. */
 class FOpenGLEventQuery : public FRenderResource
diff --git a/Engine/Source/Runtime/RHI/Private/DynamicRHI.cpp b/Engine/Source/Runtime/RHI/Private/DynamicRHI.cpp
--- a/Engine/Source/Runtime/RHI/Private/DynamicRHI.cpp
+++ b/Engine/Source/Runtime/RHI/Private/DynamicRHI.cpp
@@ -335,22 +335,22 @@ void FDynamicRHI::EnableIdealGPUCaptureOptions(bool bEnabled)
 	}	
 }
 
-void FDynamicRHI::RHITransferIndexBufferUnderlyingResource(FIndexBufferRHIParamRef DestIndexBuffer, FIndexBufferRHIParamRef SrcIndexBuffer)
+void FDynamicRHI::RHITransferIndexBufferUnderlyingResource(FRHIIndexBuffer* DestIndexBuffer, FRHIIndexBuffer* SrcIndexBuffer)
 {
 	UE_LOG(LogRHI, Fatal, TEXT("RHITransferIndexBufferUnderlyingResource isn't implemented for the current RHI"));
 }
 
-void FDynamicRHI::RHITransferVertexBufferUnderlyingResource(FVertexBufferRHIParamRef DestVertexBuffer, FVertexBufferRHIParamRef SrcVertexBuffer)
+void FDynamicRHI::RHITransferVertexBufferUnderlyingResource(FRHIVertexBuffer* DestVertexBuffer, FRHIVertexBuffer* SrcVertexBuffer)
 {
 	UE_LOG(LogRHI, Fatal, TEXT("RHITransferVertexBufferUnderlyingResource isn't implemented for the current RHI"));
 }
 
-void FDynamicRHI::RHIUpdateShaderResourceView(FRHIShaderResourceView* SRV, FVertexBufferRHIParamRef VertexBuffer, uint32 Stride, uint8 Format)
+void FDynamicRHI::RHIUpdateShaderResourceView(FRHIShaderResourceView* SRV, FRHIVertexBuffer* VertexBuffer, uint32 Stride, uint8 Format)
 {
 	UE_LOG(LogRHI, Fatal, TEXT("RHIUpdateShaderResourceView isn't implemented for the current RHI"));
 }
 
-void FDynamicRHI::RHIUpdateShaderResourceView(FRHIShaderResourceView* SRV, FIndexBufferRHIParamRef IndexBuffer)
+void FDynamicRHI::RHIUpdateShaderResourceView(FRHIShaderResourceView* SRV, FRHIIndexBuffer* IndexBuffer)
 {
 	UE_LOG(LogRHI, Fatal, TEXT("RHIUpdateShaderResourceView isn't implemented for the current RHI"));
 }
diff --git a/Engine/Source/Runtime/RHI/Private/RHICommandList.cpp b/Engine/Source/Runtime/RHI/Private/RHICommandList.cpp
--- a/Engine/Source/Runtime/RHI/Private/RHICommandList.cpp
+++ b/Engine/Source/Runtime/RHI/Private/RHICommandList.cpp
@@ -1488,7 +1488,7 @@ void FRHICommandList::EndScene()
 }
 
 
-void FRHICommandList::BeginDrawingViewport(FRHIViewport* Viewport, FTextureRHIParamRef RenderTargetRHI)
+void FRHICommandList::BeginDrawingViewport(FRHIViewport* Viewport, FRHITexture* RenderTargetRHI)
 {
 	check(IsImmediate() && IsInRenderingThread());
 	if (Bypass())
@@ -1660,13 +1660,13 @@ void FRHICommandListBase::WaitForDispatch()
 	}
 }
 
-void FDynamicRHI::VirtualTextureSetFirstMipInMemory_RenderThread(class FRHICommandListImmediate& RHICmdList, FTexture2DRHIParamRef Texture, uint32 FirstMip)
+void FDynamicRHI::VirtualTextureSetFirstMipInMemory_RenderThread(class FRHICommandListImmediate& RHICmdList, FRHITexture2D* Texture, uint32 FirstMip)
 {
 	RHICmdList.ImmediateFlush(EImmediateFlushType::FlushRHIThread);
 	GDynamicRHI->RHIVirtualTextureSetFirstMipInMemory(Texture, FirstMip);
 }
 
-void FDynamicRHI::VirtualTextureSetFirstMipVisible_RenderThread(class FRHICommandListImmediate& RHICmdList, FTexture2DRHIParamRef Texture, uint32 FirstMip)
+void FDynamicRHI::VirtualTextureSetFirstMipVisible_RenderThread(class FRHICommandListImmediate& RHICmdList, FRHITexture2D* Texture, uint32 FirstMip)
 {
 	RHICmdList.ImmediateFlush(EImmediateFlushType::FlushRHIThread);
 	GDynamicRHI->RHIVirtualTextureSetFirstMipVisible(Texture, FirstMip);
@@ -1904,14 +1904,14 @@ FIndexBufferRHIRef FDynamicRHI::CreateIndexBuffer_RenderThread(class FRHICommand
 	return GDynamicRHI->RHICreateIndexBuffer(Stride, Size, InUsage, CreateInfo);
 }
 
-FShaderResourceViewRHIRef FDynamicRHI::CreateShaderResourceView_RenderThread(class FRHICommandListImmediate& RHICmdList, FVertexBufferRHIParamRef VertexBuffer, uint32 Stride, uint8 Format)
+FShaderResourceViewRHIRef FDynamicRHI::CreateShaderResourceView_RenderThread(class FRHICommandListImmediate& RHICmdList, FRHIVertexBuffer* VertexBuffer, uint32 Stride, uint8 Format)
 {
 	CSV_SCOPED_TIMING_STAT(RHITStalls, CreateShaderResourceView_RenderThread_VB);
 	FScopedRHIThreadStaller StallRHIThread(RHICmdList);
 	return GDynamicRHI->RHICreateShaderResourceView(VertexBuffer, Stride, Format);
 }
 
-FShaderResourceViewRHIRef FDynamicRHI::CreateShaderResourceView_RenderThread(class FRHICommandListImmediate& RHICmdList, FIndexBufferRHIParamRef Buffer)
+FShaderResourceViewRHIRef FDynamicRHI::CreateShaderResourceView_RenderThread(class FRHICommandListImmediate& RHICmdList, FRHIIndexBuffer* Buffer)
 {
 	CSV_SCOPED_TIMING_STAT(RHITStalls, CreateShaderResourceView_RenderThread_IB);
 	FScopedRHIThreadStaller StallRHIThread(RHICmdList);
@@ -1920,12 +1920,12 @@ FShaderResourceViewRHIRef FDynamicRHI::CreateShaderResourceView_RenderThread(cla
 
 struct FRHICommandUpdateVertexBuffer final : public FRHICommand<FRHICommandUpdateVertexBuffer>
 {
-	FVertexBufferRHIParamRef VertexBuffer;
+	FRHIVertexBuffer* VertexBuffer;
 	void* Buffer;
 	uint32 BufferSize;
 	uint32 Offset;
 
-	FORCEINLINE_DEBUGGABLE FRHICommandUpdateVertexBuffer(FVertexBufferRHIParamRef InVertexBuffer, void* InBuffer, uint32 InOffset, uint32 InBufferSize)
+	FORCEINLINE_DEBUGGABLE FRHICommandUpdateVertexBuffer(FRHIVertexBuffer* InVertexBuffer, void* InBuffer, uint32 InOffset, uint32 InBufferSize)
 		: VertexBuffer(InVertexBuffer)
 		, Buffer(InBuffer)
 		, BufferSize(InBufferSize)
@@ -1944,12 +1944,12 @@ struct FRHICommandUpdateVertexBuffer final : public FRHICommand<FRHICommandUpdat
 
 struct FRHICommandUpdateIndexBuffer final : public FRHICommand<FRHICommandUpdateIndexBuffer>
 {
-	FIndexBufferRHIParamRef IndexBuffer;
+	FRHIIndexBuffer* IndexBuffer;
 	void* Buffer;
 	uint32 BufferSize;
 	uint32 Offset;
 
-	FORCEINLINE_DEBUGGABLE FRHICommandUpdateIndexBuffer(FIndexBufferRHIParamRef InIndexBuffer, void* InBuffer, uint32 InOffset, uint32 InBufferSize)
+	FORCEINLINE_DEBUGGABLE FRHICommandUpdateIndexBuffer(FRHIIndexBuffer* InIndexBuffer, void* InBuffer, uint32 InOffset, uint32 InBufferSize)
 		: IndexBuffer(InIndexBuffer)
 		, Buffer(InBuffer)
 		, BufferSize(InBufferSize)
@@ -1969,7 +1969,7 @@ struct FRHICommandUpdateIndexBuffer final : public FRHICommand<FRHICommandUpdate
 static FLockTracker GLockTracker;
 
 
-void* FDynamicRHI::LockVertexBuffer_RenderThread(class FRHICommandListImmediate& RHICmdList, FVertexBufferRHIParamRef VertexBuffer, uint32 Offset, uint32 SizeRHI, EResourceLockMode LockMode)
+void* FDynamicRHI::LockVertexBuffer_RenderThread(class FRHICommandListImmediate& RHICmdList, FRHIVertexBuffer* VertexBuffer, uint32 Offset, uint32 SizeRHI, EResourceLockMode LockMode)
 {
 	QUICK_SCOPE_CYCLE_COUNTER(STAT_FDynamicRHI_LockVertexBuffer_RenderThread);
 	check(IsInRenderingThread());
@@ -1996,7 +1996,7 @@ void* FDynamicRHI::LockVertexBuffer_RenderThread(class FRHICommandListImmediate&
 	return Result;
 }
 
-void FDynamicRHI::UnlockVertexBuffer_RenderThread(class FRHICommandListImmediate& RHICmdList, FVertexBufferRHIParamRef VertexBuffer)
+void FDynamicRHI::UnlockVertexBuffer_RenderThread(class FRHICommandListImmediate& RHICmdList, FRHIVertexBuffer* VertexBuffer)
 {
 	QUICK_SCOPE_CYCLE_COUNTER(STAT_FDynamicRHI_UnlockVertexBuffer_RenderThread);
 	check(IsInRenderingThread());
@@ -2028,7 +2028,7 @@ void FDynamicRHI::UnlockVertexBuffer_RenderThread(class FRHICommandListImmediate
 	}
 }
 
-void* FDynamicRHI::LockIndexBuffer_RenderThread(class FRHICommandListImmediate& RHICmdList, FIndexBufferRHIParamRef IndexBuffer, uint32 Offset, uint32 SizeRHI, EResourceLockMode LockMode)
+void* FDynamicRHI::LockIndexBuffer_RenderThread(class FRHICommandListImmediate& RHICmdList, FRHIIndexBuffer* IndexBuffer, uint32 Offset, uint32 SizeRHI, EResourceLockMode LockMode)
 {
 	QUICK_SCOPE_CYCLE_COUNTER(STAT_FDynamicRHI_LockIndexBuffer_RenderThread);
 	check(IsInRenderingThread());
@@ -2055,7 +2055,7 @@ void* FDynamicRHI::LockIndexBuffer_RenderThread(class FRHICommandListImmediate&
 	return Result;
 }
 
-void FDynamicRHI::UnlockIndexBuffer_RenderThread(class FRHICommandListImmediate& RHICmdList, FIndexBufferRHIParamRef IndexBuffer)
+void FDynamicRHI::UnlockIndexBuffer_RenderThread(class FRHICommandListImmediate& RHICmdList, FRHIIndexBuffer* IndexBuffer)
 {
 	QUICK_SCOPE_CYCLE_COUNTER(STAT_FDynamicRHI_UnlockIndexBuffer_RenderThread);
 	check(IsInRenderingThread());
@@ -2115,21 +2115,21 @@ void FDynamicRHI::UnlockStagingBuffer_RenderThread(class FRHICommandListImmediat
 	GDynamicRHI->RHIUnlockStagingBuffer(StagingBuffer);
 }
 
-FTexture2DRHIRef FDynamicRHI::AsyncReallocateTexture2D_RenderThread(class FRHICommandListImmediate& RHICmdList, FTexture2DRHIParamRef Texture2D, int32 NewMipCount, int32 NewSizeX, int32 NewSizeY, FThreadSafeCounter* RequestStatus)
+FTexture2DRHIRef FDynamicRHI::AsyncReallocateTexture2D_RenderThread(class FRHICommandListImmediate& RHICmdList, FRHITexture2D* Texture2D, int32 NewMipCount, int32 NewSizeX, int32 NewSizeY, FThreadSafeCounter* RequestStatus)
 {
 	QUICK_SCOPE_CYCLE_COUNTER(STAT_RHIMETHOD_AsyncReallocateTexture2D_Flush);
 	RHICmdList.ImmediateFlush(EImmediateFlushType::FlushRHIThread);  
 	return GDynamicRHI->RHIAsyncReallocateTexture2D(Texture2D, NewMipCount, NewSizeX, NewSizeY, RequestStatus);
 }
 
-ETextureReallocationStatus FDynamicRHI::FinalizeAsyncReallocateTexture2D_RenderThread(class FRHICommandListImmediate& RHICmdList, FTexture2DRHIParamRef Texture2D, bool bBlockUntilCompleted)
+ETextureReallocationStatus FDynamicRHI::FinalizeAsyncReallocateTexture2D_RenderThread(class FRHICommandListImmediate& RHICmdList, FRHITexture2D* Texture2D, bool bBlockUntilCompleted)
 {
 	CSV_SCOPED_TIMING_STAT(RHITStalls, FinalizeAsyncReallocateTexture2D_RenderThread);
 	FScopedRHIThreadStaller StallRHIThread(RHICmdList);
 	return GDynamicRHI->RHIFinalizeAsyncReallocateTexture2D(Texture2D, bBlockUntilCompleted);
 }
 
-ETextureReallocationStatus FDynamicRHI::CancelAsyncReallocateTexture2D_RenderThread(class FRHICommandListImmediate& RHICmdList, FTexture2DRHIParamRef Texture2D, bool bBlockUntilCompleted)
+ETextureReallocationStatus FDynamicRHI::CancelAsyncReallocateTexture2D_RenderThread(class FRHICommandListImmediate& RHICmdList, FRHITexture2D* Texture2D, bool bBlockUntilCompleted)
 {
 	CSV_SCOPED_TIMING_STAT(RHITStalls, CancelAsyncReallocateTexture2D_RenderThread);
 	FScopedRHIThreadStaller StallRHIThread(RHICmdList);
@@ -2234,20 +2234,20 @@ FDomainShaderRHIRef FDynamicRHI::CreateDomainShader_RenderThread(class FRHIComma
 	return GDynamicRHI->RHICreateDomainShader(Library, Hash);
 }
 
-void FDynamicRHI::UpdateTexture2D_RenderThread(class FRHICommandListImmediate& RHICmdList, FTexture2DRHIParamRef Texture, uint32 MipIndex, const struct FUpdateTextureRegion2D& UpdateRegion, uint32 SourcePitch, const uint8* SourceData)
+void FDynamicRHI::UpdateTexture2D_RenderThread(class FRHICommandListImmediate& RHICmdList, FRHITexture2D* Texture, uint32 MipIndex, const struct FUpdateTextureRegion2D& UpdateRegion, uint32 SourcePitch, const uint8* SourceData)
 {
 	CSV_SCOPED_TIMING_STAT(RHITStalls, UpdateTexture2D_RenderThread);
 	FScopedRHIThreadStaller StallRHIThread(RHICmdList);
 	return GDynamicRHI->RHIUpdateTexture2D(Texture, MipIndex, UpdateRegion, SourcePitch, SourceData);
 }
 
-void FDynamicRHI::UpdateFromBufferTexture2D_RenderThread(class FRHICommandListImmediate& RHICmdList, FTexture2DRHIParamRef Texture, uint32 MipIndex, const struct FUpdateTextureRegion2D& UpdateRegion, uint32 SourcePitch, FStructuredBufferRHIParamRef Buffer, uint32 BufferOffset)
+void FDynamicRHI::UpdateFromBufferTexture2D_RenderThread(class FRHICommandListImmediate& RHICmdList, FRHITexture2D* Texture, uint32 MipIndex, const struct FUpdateTextureRegion2D& UpdateRegion, uint32 SourcePitch, FRHIStructuredBuffer* Buffer, uint32 BufferOffset)
 {
 	FScopedRHIThreadStaller StallRHIThread(RHICmdList);
 	return GDynamicRHI->RHIUpdateFromBufferTexture2D(Texture, MipIndex, UpdateRegion, SourcePitch, Buffer, BufferOffset);
 }
 
-FUpdateTexture3DData FDynamicRHI::BeginUpdateTexture3D_RenderThread(class FRHICommandListImmediate& RHICmdList, FTexture3DRHIParamRef Texture, uint32 MipIndex, const struct FUpdateTextureRegion3D& UpdateRegion)
+FUpdateTexture3DData FDynamicRHI::BeginUpdateTexture3D_RenderThread(class FRHICommandListImmediate& RHICmdList, FRHITexture3D* Texture, uint32 MipIndex, const struct FUpdateTextureRegion3D& UpdateRegion)
 {
 	check(IsInRenderingThread());
 
@@ -2280,14 +2280,14 @@ void FDynamicRHI::EndMultiUpdateTexture3D_RenderThread(class FRHICommandListImme
 	}
 }
 
-void FDynamicRHI::UpdateTexture3D_RenderThread(class FRHICommandListImmediate& RHICmdList, FTexture3DRHIParamRef Texture, uint32 MipIndex, const struct FUpdateTextureRegion3D& UpdateRegion, uint32 SourceRowPitch, uint32 SourceDepthPitch, const uint8* SourceData)
+void FDynamicRHI::UpdateTexture3D_RenderThread(class FRHICommandListImmediate& RHICmdList, FRHITexture3D* Texture, uint32 MipIndex, const struct FUpdateTextureRegion3D& UpdateRegion, uint32 SourceRowPitch, uint32 SourceDepthPitch, const uint8* SourceData)
 {
 	CSV_SCOPED_TIMING_STAT(RHITStalls, UpdateTexture3D_RenderThread);
 	FScopedRHIThreadStaller StallRHIThread(RHICmdList);
 	GDynamicRHI->RHIUpdateTexture3D(Texture, MipIndex, UpdateRegion, SourceRowPitch, SourceDepthPitch, SourceData);
 }
 
-void* FDynamicRHI::LockTexture2D_RenderThread(class FRHICommandListImmediate& RHICmdList, FTexture2DRHIParamRef Texture, uint32 MipIndex, EResourceLockMode LockMode, uint32& DestStride, bool bLockWithinMiptail, bool bNeedsDefaultRHIFlush)
+void* FDynamicRHI::LockTexture2D_RenderThread(class FRHICommandListImmediate& RHICmdList, FRHITexture2D* Texture, uint32 MipIndex, EResourceLockMode LockMode, uint32& DestStride, bool bLockWithinMiptail, bool bNeedsDefaultRHIFlush)
 {
 	if (bNeedsDefaultRHIFlush) 
 	{
@@ -2300,7 +2300,7 @@ void* FDynamicRHI::LockTexture2D_RenderThread(class FRHICommandListImmediate& RH
 	return GDynamicRHI->RHILockTexture2D(Texture, MipIndex, LockMode, DestStride, bLockWithinMiptail);
 }
 
-void FDynamicRHI::UnlockTexture2D_RenderThread(class FRHICommandListImmediate& RHICmdList, FTexture2DRHIParamRef Texture, uint32 MipIndex, bool bLockWithinMiptail, bool bNeedsDefaultRHIFlush)
+void FDynamicRHI::UnlockTexture2D_RenderThread(class FRHICommandListImmediate& RHICmdList, FRHITexture2D* Texture, uint32 MipIndex, bool bLockWithinMiptail, bool bNeedsDefaultRHIFlush)
 {
 	if (bNeedsDefaultRHIFlush)
 	{
@@ -2354,55 +2354,55 @@ FTexture3DRHIRef FDynamicRHI::RHICreateTexture3D_RenderThread(class FRHICommandL
 	return GDynamicRHI->RHICreateTexture3D(SizeX, SizeY, SizeZ, Format, NumMips, Flags, CreateInfo);
 }
 
-FUnorderedAccessViewRHIRef FDynamicRHI::RHICreateUnorderedAccessView_RenderThread(class FRHICommandListImmediate& RHICmdList, FStructuredBufferRHIParamRef StructuredBuffer, bool bUseUAVCounter, bool bAppendBuffer)
+FUnorderedAccessViewRHIRef FDynamicRHI::RHICreateUnorderedAccessView_RenderThread(class FRHICommandListImmediate& RHICmdList, FRHIStructuredBuffer* StructuredBuffer, bool bUseUAVCounter, bool bAppendBuffer)
 {
 	CSV_SCOPED_TIMING_STAT(RHITStalls, RHICreateUnorderedAccessView_RenderThread);
 	FScopedRHIThreadStaller StallRHIThread(RHICmdList);
 	return GDynamicRHI->RHICreateUnorderedAccessView(StructuredBuffer, bUseUAVCounter, bAppendBuffer);
 }
 
-FUnorderedAccessViewRHIRef FDynamicRHI::RHICreateUnorderedAccessView_RenderThread(class FRHICommandListImmediate& RHICmdList, FTextureRHIParamRef Texture, uint32 MipLevel)
+FUnorderedAccessViewRHIRef FDynamicRHI::RHICreateUnorderedAccessView_RenderThread(class FRHICommandListImmediate& RHICmdList, FRHITexture* Texture, uint32 MipLevel)
 {
 	CSV_SCOPED_TIMING_STAT(RHITStalls, RHICreateUnorderedAccessView_RenderThread);
 	FScopedRHIThreadStaller StallRHIThread(RHICmdList);
 	return GDynamicRHI->RHICreateUnorderedAccessView(Texture, MipLevel);
 }
 
-FUnorderedAccessViewRHIRef FDynamicRHI::RHICreateUnorderedAccessView_RenderThread(class FRHICommandListImmediate& RHICmdList, FVertexBufferRHIParamRef VertexBuffer, uint8 Format)
+FUnorderedAccessViewRHIRef FDynamicRHI::RHICreateUnorderedAccessView_RenderThread(class FRHICommandListImmediate& RHICmdList, FRHIVertexBuffer* VertexBuffer, uint8 Format)
 {
 	CSV_SCOPED_TIMING_STAT(RHITStalls, RHICreateUnorderedAccessView_RenderThread);
 	FScopedRHIThreadStaller StallRHIThread(RHICmdList);
 	return GDynamicRHI->RHICreateUnorderedAccessView(VertexBuffer, Format);
 }
 
-FUnorderedAccessViewRHIRef FDynamicRHI::RHICreateUnorderedAccessView_RenderThread(class FRHICommandListImmediate& RHICmdList, FIndexBufferRHIParamRef IndexBuffer, uint8 Format)
+FUnorderedAccessViewRHIRef FDynamicRHI::RHICreateUnorderedAccessView_RenderThread(class FRHICommandListImmediate& RHICmdList, FRHIIndexBuffer* IndexBuffer, uint8 Format)
 {
 	FScopedRHIThreadStaller StallRHIThread(RHICmdList);
 	return GDynamicRHI->RHICreateUnorderedAccessView(IndexBuffer, Format);
 }
 
-FShaderResourceViewRHIRef FDynamicRHI::RHICreateShaderResourceView_RenderThread(class FRHICommandListImmediate& RHICmdList, FTextureRHIParamRef Texture, const FRHITextureSRVCreateInfo& CreateInfo)
+FShaderResourceViewRHIRef FDynamicRHI::RHICreateShaderResourceView_RenderThread(class FRHICommandListImmediate& RHICmdList, FRHITexture* Texture, const FRHITextureSRVCreateInfo& CreateInfo)
 {
 	CSV_SCOPED_TIMING_STAT(RHITStalls, RHICreateShaderResourceView_RenderThread_Tex2D); // TODO - clean this up
 	FScopedRHIThreadStaller StallRHIThread(RHICmdList);
 	return GDynamicRHI->RHICreateShaderResourceView(Texture, CreateInfo);
 }
 
-FShaderResourceViewRHIRef FDynamicRHI::RHICreateShaderResourceView_RenderThread(class FRHICommandListImmediate& RHICmdList, FVertexBufferRHIParamRef VertexBuffer, uint32 Stride, uint8 Format)
+FShaderResourceViewRHIRef FDynamicRHI::RHICreateShaderResourceView_RenderThread(class FRHICommandListImmediate& RHICmdList, FRHIVertexBuffer* VertexBuffer, uint32 Stride, uint8 Format)
 {
 	CSV_SCOPED_TIMING_STAT(RHITStalls, RHICreateShaderResourceView_RenderThread_VB);
 	FScopedRHIThreadStaller StallRHIThread(RHICmdList);
 	return GDynamicRHI->RHICreateShaderResourceView(VertexBuffer, Stride, Format);
 }
 
-FShaderResourceViewRHIRef FDynamicRHI::RHICreateShaderResourceView_RenderThread(class FRHICommandListImmediate& RHICmdList, FIndexBufferRHIParamRef Buffer)
+FShaderResourceViewRHIRef FDynamicRHI::RHICreateShaderResourceView_RenderThread(class FRHICommandListImmediate& RHICmdList, FRHIIndexBuffer* Buffer)
 {
 	CSV_SCOPED_TIMING_STAT(RHITStalls, RHICreateShaderResourceView_RenderThread_IB);
 	FScopedRHIThreadStaller StallRHIThread(RHICmdList);
 	return GDynamicRHI->RHICreateShaderResourceView(Buffer);
 }
 
-FShaderResourceViewRHIRef FDynamicRHI::RHICreateShaderResourceView_RenderThread(class FRHICommandListImmediate& RHICmdList, FStructuredBufferRHIParamRef StructuredBuffer)
+FShaderResourceViewRHIRef FDynamicRHI::RHICreateShaderResourceView_RenderThread(class FRHICommandListImmediate& RHICmdList, FRHIStructuredBuffer* StructuredBuffer)
 {
 	CSV_SCOPED_TIMING_STAT(RHITStalls, RHICreateShaderResourceView_RenderThread_SB);
 	FScopedRHIThreadStaller StallRHIThread(RHICmdList);
@@ -2430,14 +2430,14 @@ FRenderQueryRHIRef FDynamicRHI::RHICreateRenderQuery_RenderThread(class FRHIComm
 	return GDynamicRHI->RHICreateRenderQuery(QueryType);
 }
 
-void* FDynamicRHI::RHILockTextureCubeFace_RenderThread(class FRHICommandListImmediate& RHICmdList, FTextureCubeRHIParamRef Texture, uint32 FaceIndex, uint32 ArrayIndex, uint32 MipIndex, EResourceLockMode LockMode, uint32& DestStride, bool bLockWithinMiptail)
+void* FDynamicRHI::RHILockTextureCubeFace_RenderThread(class FRHICommandListImmediate& RHICmdList, FRHITextureCube* Texture, uint32 FaceIndex, uint32 ArrayIndex, uint32 MipIndex, EResourceLockMode LockMode, uint32& DestStride, bool bLockWithinMiptail)
 {
 	QUICK_SCOPE_CYCLE_COUNTER(STAT_RHIMETHOD_LockTextureCubeFace_Flush);
 	RHICmdList.ImmediateFlush(EImmediateFlushType::FlushRHIThread);
 	return GDynamicRHI->RHILockTextureCubeFace(Texture, FaceIndex, ArrayIndex, MipIndex, LockMode, DestStride, bLockWithinMiptail);
 }
 
-void FDynamicRHI::RHIUnlockTextureCubeFace_RenderThread(class FRHICommandListImmediate& RHICmdList, FTextureCubeRHIParamRef Texture, uint32 FaceIndex, uint32 ArrayIndex, uint32 MipIndex, bool bLockWithinMiptail)
+void FDynamicRHI::RHIUnlockTextureCubeFace_RenderThread(class FRHICommandListImmediate& RHICmdList, FRHITextureCube* Texture, uint32 FaceIndex, uint32 ArrayIndex, uint32 MipIndex, bool bLockWithinMiptail)
 {
 	QUICK_SCOPE_CYCLE_COUNTER(STAT_RHIMETHOD_UnlockTextureCubeFace_Flush);
 	RHICmdList.ImmediateFlush(EImmediateFlushType::FlushRHIThread);
@@ -2445,7 +2445,7 @@ void FDynamicRHI::RHIUnlockTextureCubeFace_RenderThread(class FRHICommandListImm
 }
 
 
-void FDynamicRHI::RHIMapStagingSurface_RenderThread(class FRHICommandListImmediate& RHICmdList, FTextureRHIParamRef Texture, void*& OutData, int32& OutWidth, int32& OutHeight)
+void FDynamicRHI::RHIMapStagingSurface_RenderThread(class FRHICommandListImmediate& RHICmdList, FRHITexture* Texture, void*& OutData, int32& OutWidth, int32& OutHeight)
 {
 	{
 		QUICK_SCOPE_CYCLE_COUNTER(STAT_RHIMETHOD_MapStagingSurface_Flush);
@@ -2454,7 +2454,7 @@ void FDynamicRHI::RHIMapStagingSurface_RenderThread(class FRHICommandListImmedia
 	GDynamicRHI->RHIMapStagingSurface(Texture, OutData, OutWidth, OutHeight);
 }
 
-void FDynamicRHI::RHIUnmapStagingSurface_RenderThread(class FRHICommandListImmediate& RHICmdList, FTextureRHIParamRef Texture)
+void FDynamicRHI::RHIUnmapStagingSurface_RenderThread(class FRHICommandListImmediate& RHICmdList, FRHITexture* Texture)
 {
 	{
 		QUICK_SCOPE_CYCLE_COUNTER(STAT_RHIMETHOD_MapStagingSurface_Flush);
@@ -2463,7 +2463,7 @@ void FDynamicRHI::RHIUnmapStagingSurface_RenderThread(class FRHICommandListImmed
 	GDynamicRHI->RHIUnmapStagingSurface(Texture);
 }
 
-void FDynamicRHI::RHIReadSurfaceFloatData_RenderThread(class FRHICommandListImmediate& RHICmdList, FTextureRHIParamRef Texture, FIntRect Rect, TArray<FFloat16Color>& OutData, ECubeFace CubeFace, int32 ArrayIndex, int32 MipIndex)
+void FDynamicRHI::RHIReadSurfaceFloatData_RenderThread(class FRHICommandListImmediate& RHICmdList, FRHITexture* Texture, FIntRect Rect, TArray<FFloat16Color>& OutData, ECubeFace CubeFace, int32 ArrayIndex, int32 MipIndex)
 {
 	QUICK_SCOPE_CYCLE_COUNTER(STAT_RHIMETHOD_ReadSurfaceFloatData_Flush);
 	RHICmdList.ImmediateFlush(EImmediateFlushType::FlushRHIThread);
@@ -2471,7 +2471,7 @@ void FDynamicRHI::RHIReadSurfaceFloatData_RenderThread(class FRHICommandListImme
 }
 
 
-void FRHICommandListImmediate::UpdateTextureReference(FRHITextureReference* TextureRef, FTextureRHIParamRef NewTexture)
+void FRHICommandListImmediate::UpdateTextureReference(FRHITextureReference* TextureRef, FRHITexture* NewTexture)
 {
 	if (Bypass() || !IsRunningRHIInSeparateThread() || CVarRHICmdFlushUpdateTextureReference.GetValueOnRenderThread() > 0)
 	{
@@ -2513,7 +2513,7 @@ void FRHICommandListImmediate::UpdateRHIResources(FRHIResourceUpdateInfo* Update
 	}
 }
 
-void FDynamicRHI::RHICopySubTextureRegion_RenderThread(class FRHICommandListImmediate& RHICmdList, FTexture2DRHIParamRef SourceTexture, FTexture2DRHIParamRef DestinationTexture, FBox2D SourceBox, FBox2D DestinationBox)
+void FDynamicRHI::RHICopySubTextureRegion_RenderThread(class FRHICommandListImmediate& RHICmdList, FRHITexture2D* SourceTexture, FRHITexture2D* DestinationTexture, FBox2D SourceBox, FBox2D DestinationBox)
 {
 	CSV_SCOPED_TIMING_STAT(RHITStalls, RHICopySubTextureRegion_RenderThread);
 	FScopedRHIThreadStaller StallRHIThread(RHICmdList);
diff --git a/Engine/Source/Runtime/RHI/Private/RHIGPUReadback.cpp b/Engine/Source/Runtime/RHI/Private/RHIGPUReadback.cpp
--- a/Engine/Source/Runtime/RHI/Private/RHIGPUReadback.cpp
+++ b/Engine/Source/Runtime/RHI/Private/RHIGPUReadback.cpp
@@ -65,7 +65,7 @@ FRHIGPUBufferReadback::FRHIGPUBufferReadback(FName RequestName) : FRHIGPUMemoryR
 	DestinationStagingBuffer = RHICreateStagingBuffer();
 }
 
-void FRHIGPUBufferReadback::EnqueueCopy(FRHICommandList& RHICmdList, FVertexBufferRHIParamRef SourceBuffer, uint32 NumBytes)
+void FRHIGPUBufferReadback::EnqueueCopy(FRHICommandList& RHICmdList, FRHIVertexBuffer* SourceBuffer, uint32 NumBytes)
 {
 	Fence->Clear();
 	RHICmdList.CopyToStagingBuffer(SourceBuffer, DestinationStagingBuffer, 0, NumBytes ? NumBytes : SourceBuffer->GetSize());
@@ -99,7 +99,7 @@ FRHIGPUTextureReadback::FRHIGPUTextureReadback(FName RequestName) : FRHIGPUMemor
 {
 }
 
-void FRHIGPUTextureReadback::EnqueueCopy(FRHICommandList& RHICmdList, FTextureRHIParamRef SourceTexture, FResolveRect Rect)
+void FRHIGPUTextureReadback::EnqueueCopy(FRHICommandList& RHICmdList, FRHITexture* SourceTexture, FResolveRect Rect)
 {
 	Fence->Clear();
 
