Do not pass GC guard to different threads.

#rnx

diff --git a/Engine/Source/Runtime/OpenGLDrv/Private/Android/AndroidEGL.cpp b/Engine/Source/Runtime/OpenGLDrv/Private/Android/AndroidEGL.cpp
--- a/Engine/Source/Runtime/OpenGLDrv/Private/Android/AndroidEGL.cpp
+++ b/Engine/Source/Runtime/OpenGLDrv/Private/Android/AndroidEGL.cpp
@@ -1716,6 +1716,9 @@ void BlockOnLostWindowRenderCommand(TSharedPtr<FEvent, ESPMode::ThreadSafe> RTBl
 {
 	check(IsInRenderingThread());
 
+	// Hold GC scope guard, as GC will timeout if anything waits for RT fences.
+	FGCScopeGuard GCGuard;
+	
 	UE_LOG(LogAndroid, Log, TEXT("Blocking renderer"));
 	if (FAndroidMisc::ShouldUseVulkan())
 	{
@@ -1767,7 +1770,9 @@ void BlockRendering()
 	check(GIsRHIInitialized);
 
 	UE_LOG(LogAndroid, Log, TEXT("Blocking renderer on invalid window."));
-
+	
+	// Wait for GC to complete and prevent further GCs
+	FGCScopeGuard GCGuard;
 	TSharedPtr<FEvent, ESPMode::ThreadSafe> RTBlockedTrigger = MakeShareable(FPlatformProcess::GetSynchEventFromPool(), [](FEvent* EventToDelete)
 	{
 		FPlatformProcess::ReturnSynchEventToPool(EventToDelete);
