Fix for android CIS build
#rnx

diff --git a/Engine/Source/Runtime/OpenGLDrv/Private/Android/AndroidEGL.cpp b/Engine/Source/Runtime/OpenGLDrv/Private/Android/AndroidEGL.cpp
--- a/Engine/Source/Runtime/OpenGLDrv/Private/Android/AndroidEGL.cpp
+++ b/Engine/Source/Runtime/OpenGLDrv/Private/Android/AndroidEGL.cpp
@@ -14,6 +14,7 @@
 #endif
 #include "OpenGLDrvPrivate.h"
 #include "Misc/ScopeLock.h"
+#include "UObject/GarbageCollection.h"
 
 
 AndroidEGL* AndroidEGL::Singleton = NULL;
