#rnx

this solution was based on: UE-48208

diff --git a/Engine/Source/Runtime/OpenGLDrv/Private/OpenGLDevice.cpp b/Engine/Source/Runtime/OpenGLDrv/Private/OpenGLDevice.cpp
--- a/Engine/Source/Runtime/OpenGLDrv/Private/OpenGLDevice.cpp
+++ b/Engine/Source/Runtime/OpenGLDrv/Private/OpenGLDevice.cpp
@@ -1023,6 +1023,9 @@ static void InitRHICapabilitiesForGL()
 		SetupTextureFormat(PF_G8, FOpenGLTextureFormat(GL_LUMINANCE, GL_LUMINANCE, GL_LUMINANCE, GL_LUMINANCE, GL_LUMINANCE, GL_UNSIGNED_BYTE, false, false));
 		SetupTextureFormat(PF_A8, FOpenGLTextureFormat(GL_ALPHA, GL_ALPHA, GL_ALPHA, GL_ALPHA, GL_ALPHA, GL_UNSIGNED_BYTE, false, false));
 	#endif
+	#if PLATFORM_HTML5
+		SetupTextureFormat(PF_R8G8B8A8_UINT, FOpenGLTextureFormat(GL_RGBA8UI, GL_RGBA8UI, GL_RGBA_INTEGER, GL_UNSIGNED_BYTE, false, false));
+	#endif
 
 		if (FOpenGL::SupportsColorBufferHalfFloat() && FOpenGL::SupportsTextureHalfFloat())
 		{
