Fixed: MSAA depthbuffer never be released on Android

diff --git a/Engine/Source/Runtime/OpenGLDrv/Private/OpenGLRenderTarget.cpp b/Engine/Source/Runtime/OpenGLDrv/Private/OpenGLRenderTarget.cpp
--- a/Engine/Source/Runtime/OpenGLDrv/Private/OpenGLRenderTarget.cpp
+++ b/Engine/Source/Runtime/OpenGLDrv/Private/OpenGLRenderTarget.cpp
@@ -289,37 +289,30 @@ GLuint FOpenGLDynamicRHI::GetOpenGLFramebuffer(uint32 NumSimultaneousRenderTarge
 #endif
 		case GL_TEXTURE_2D_MULTISAMPLE:
 		{
-#if PLATFORM_ANDROID && !PLATFORM_LUMINGL4
+			if (!FOpenGL::SupportsCombinedDepthStencilAttachment() && DepthStencilTarget->Attachment == GL_DEPTH_STENCIL_ATTACHMENT)
+			{
+				FOpenGL::FramebufferTexture2D(GL_FRAMEBUFFER, GL_DEPTH_ATTACHMENT, DepthStencilTarget->Target, DepthStencilTarget->Resource, 0);
+				FOpenGL::FramebufferTexture2D(GL_FRAMEBUFFER, GL_STENCIL_ATTACHMENT, DepthStencilTarget->Target, DepthStencilTarget->Resource, 0);
+			}
+			else
+			{
+				FOpenGL::FramebufferTexture2D(GL_FRAMEBUFFER, DepthStencilTarget->Attachment, DepthStencilTarget->Target, DepthStencilTarget->Resource, 0);
+			}
+			break;
+		}
+		case GL_RENDERBUFFER:
+		{
 			FOpenGLTexture2D* DepthStencilTarget2D = (FOpenGLTexture2D*)DepthStencilTarget;
 			const uint32 NumSamplesTileMem = DepthStencilTarget2D->GetNumSamplesTileMem();
 			if (NumSamplesTileMem > 1)
-			{
-				GLuint DepthBuffer;
-				glGenRenderbuffers(1, &DepthBuffer);
-				glBindRenderbuffer(GL_RENDERBUFFER, DepthBuffer);
-				glRenderbufferStorageMultisampleEXT(GL_RENDERBUFFER, NumSamplesTileMem, FOpenGL::SupportsPackedDepthStencil() ? GL_DEPTH24_STENCIL8 : GL_DEPTH_COMPONENT24, DepthStencilTarget2D->GetSizeX(), DepthStencilTarget2D->GetSizeY());
-				VERIFY_GL(glRenderbufferStorageMultisampleEXT);
-				glBindRenderbuffer(GL_RENDERBUFFER, 0);
-				glFramebufferRenderbuffer(GL_FRAMEBUFFER, GL_DEPTH_ATTACHMENT, GL_RENDERBUFFER, DepthBuffer);
+			{	
+				glFramebufferRenderbuffer(GL_FRAMEBUFFER, GL_DEPTH_ATTACHMENT, GL_RENDERBUFFER, DepthStencilTarget->Resource);
 				if (FOpenGL::SupportsPackedDepthStencil())
 				{
-					glFramebufferRenderbuffer(GL_FRAMEBUFFER, GL_STENCIL_ATTACHMENT, GL_RENDERBUFFER, DepthBuffer);
+					glFramebufferRenderbuffer(GL_FRAMEBUFFER, GL_STENCIL_ATTACHMENT, GL_RENDERBUFFER, DepthStencilTarget->Resource);
 				}
 				VERIFY_GL(glFramebufferRenderbuffer);
 			}
-			else
-#endif
-			{
-				if (!FOpenGL::SupportsCombinedDepthStencilAttachment() && DepthStencilTarget->Attachment == GL_DEPTH_STENCIL_ATTACHMENT)
-				{
-					FOpenGL::FramebufferTexture2D(GL_FRAMEBUFFER, GL_DEPTH_ATTACHMENT, DepthStencilTarget->Target, DepthStencilTarget->Resource, 0);
-					FOpenGL::FramebufferTexture2D(GL_FRAMEBUFFER, GL_STENCIL_ATTACHMENT, DepthStencilTarget->Target, DepthStencilTarget->Resource, 0);
-				}
-				else
-				{
-					FOpenGL::FramebufferTexture2D(GL_FRAMEBUFFER, DepthStencilTarget->Attachment, DepthStencilTarget->Target, DepthStencilTarget->Resource, 0);
-				}
-			}
 			break;
 		}
 		case GL_TEXTURE_3D:
diff --git a/Engine/Source/Runtime/OpenGLDrv/Private/OpenGLTexture.cpp b/Engine/Source/Runtime/OpenGLDrv/Private/OpenGLTexture.cpp
--- a/Engine/Source/Runtime/OpenGLDrv/Private/OpenGLTexture.cpp
+++ b/Engine/Source/Runtime/OpenGLDrv/Private/OpenGLTexture.cpp
@@ -380,11 +380,16 @@ void FOpenGLDynamicRHI::InitializeGLTexture(FRHITexture* Texture, uint32 SizeX,
 
 	bool bAllocatedStorage = false;
 
+	GLenum Target = bCubeTexture ? ((FOpenGLTextureCube*)Texture)->Target : ((FOpenGLTexture2D*)Texture)->Target;
+	const uint32 NumSamplesTileMem = bCubeTexture ? 1 : ((FOpenGLTexture2D*)Texture)->GetNumSamplesTileMem();
+	const bool TileMemDepth = NumSamplesTileMem > 1 && (Flags & TexCreate_DepthStencilTargetable);
+
 	GLuint TextureID = 0;
-	FOpenGL::GenTextures(1, &TextureID);
-	
-	const GLenum Target = bCubeTexture ? ((FOpenGLTextureCube*)Texture)->Target : ((FOpenGLTexture2D*)Texture)->Target;
-	
+	if (!TileMemDepth)
+	{
+		FOpenGL::GenTextures(1, &TextureID);
+	}
+		
 	const bool bSRGB = (Flags&TexCreate_SRGB) != 0;
 	const FOpenGLTextureFormat& GLFormat = GOpenGLTextureFormats[Format];
 	if (GLFormat.InternalFormat[bSRGB] == GL_NONE)
@@ -403,7 +408,7 @@ void FOpenGLDynamicRHI::InitializeGLTexture(FRHITexture* Texture, uint32 SizeX,
 	// For client storage textures we allocate a single backing store buffer.
 	uint8* TextureRange = nullptr;
 	
-	if (NumSamples == 1)
+	if (NumSamples == 1 && !TileMemDepth)
 	{
 		if (Target == GL_TEXTURE_EXTERNAL_OES || !FMath::IsPowerOfTwo(SizeX) || !FMath::IsPowerOfTwo(SizeY))
 		{
@@ -614,6 +619,17 @@ void FOpenGLDynamicRHI::InitializeGLTexture(FRHITexture* Texture, uint32 SizeX,
 			BulkData->Discard();
 		}
 	}
+	else if (TileMemDepth)
+	{
+#if PLATFORM_ANDROID && !PLATFORM_LUMINGL4		
+		Target = GL_RENDERBUFFER;
+		glGenRenderbuffers(1, &TextureID);
+		glBindRenderbuffer(GL_RENDERBUFFER, TextureID);
+		glRenderbufferStorageMultisampleEXT(GL_RENDERBUFFER, NumSamplesTileMem, FOpenGL::SupportsPackedDepthStencil() ? GL_DEPTH24_STENCIL8 : GL_DEPTH_COMPONENT24, SizeX, SizeY);
+		VERIFY_GL(glRenderbufferStorageMultisampleEXT);
+		glBindRenderbuffer(GL_RENDERBUFFER, 0);
+#endif
+	}
 	else
 	{
 		check( FOpenGL::SupportsMultisampledTextures() );
