HTML5 -- bypass some checks when using WebGL1

#rnx

diff --git a/Engine/Source/Runtime/OpenGLDrv/Private/OpenGLShaders.cpp b/Engine/Source/Runtime/OpenGLDrv/Private/OpenGLShaders.cpp
--- a/Engine/Source/Runtime/OpenGLDrv/Private/OpenGLShaders.cpp
+++ b/Engine/Source/Runtime/OpenGLDrv/Private/OpenGLShaders.cpp
@@ -23,6 +23,10 @@
 #endif
 #include "SceneUtils.h"
 
+#if PLATFORM_HTML5
+#include "HTML5JavaScriptFx.h"
+#endif
+
 static TAutoConsoleVariable<int32> CVarEnableLRU(
 	TEXT("r.OpenGL.EnableProgramLRUCache"),
 	0,
@@ -292,12 +296,27 @@ static bool VerifyCompiledShader(GLuint Shader, const ANSICHAR* GlslCode )
 	#endif
 			}	
 	#endif
-			UE_LOG(LogRHI,Fatal,TEXT("Failed to compile shader. Compile log:\n%s"), ANSI_TO_TCHAR(CompileLog));
+#if PLATFORM_HTML5 // TODO: REMOVE this by-pass when WebGL1 support has been discontinued.
+			if (UE_BrowserWebGLVersion() == 1)
+			{
+				UE_LOG(LogRHI,Warning,TEXT("Failed to compile shader. Compile log:\n%s"), ANSI_TO_TCHAR(CompileLog));
+			}
+			else
+#endif
+			{
+				UE_LOG(LogRHI,Fatal,TEXT("Failed to compile shader. Compile log:\n%s"), ANSI_TO_TCHAR(CompileLog));
+			}
 
 			if (LogLength > 1)
 			{
 				FMemory::Free(CompileLog);
 			}
+#if PLATFORM_HTML5 // TODO: REMOVE this by-pass when WebGL1 support has been discontinued.
+			if (UE_BrowserWebGLVersion() == 1)
+			{
+				return true;
+			}
+#endif
 			return false;
 		}
 	}
@@ -3124,6 +3143,12 @@ static FOpenGLLinkedProgram* LinkProgram( const FOpenGLLinkedProgramConfiguratio
 			}
 		}
 	}
+#if PLATFORM_HTML5 // TODO: REMOVE this by-pass when WebGL1 support has been discontinued.
+	else if (UE_BrowserWebGLVersion() == 1)
+	{
+		UE_LOG(LogRHI, Warning, TEXT("Failed to link program. Current total programs: %d, precompile: %d"), GNumPrograms, (uint32)bFromPSOFileCache);
+	}
+#endif
 	else
 	{
 		FName LinkFailurePanic = bFromPSOFileCache ? FName("FailedProgramLinkDuringPrecompile") : FName("FailedProgramLink");
@@ -3654,7 +3679,6 @@ FBoundShaderStateRHIRef FOpenGLDynamicRHI::RHICreateBoundShaderState_OnThisThrea
 
 				// Link program, using the data provided in config
 				LinkedProgram = LinkProgram(Config, bFromPSOFileCache);
-
 				GetOpenGLProgramsCache().Add(Config.ProgramKey, LinkedProgram);
 
 				// if building the cache file and using the LRU then evict the last shader created. this will reduce the risk of fragmentation of the driver's program memory.
