Copied from CL# 4497468. Fix skeletal mesh rendering on PowerVR devices running feature level ES3.1

diff --git a/Engine/Source/Runtime/OpenGLDrv/Private/Android/AndroidOpenGL.cpp b/Engine/Source/Runtime/OpenGLDrv/Private/Android/AndroidOpenGL.cpp
--- a/Engine/Source/Runtime/OpenGLDrv/Private/Android/AndroidOpenGL.cpp
+++ b/Engine/Source/Runtime/OpenGLDrv/Private/Android/AndroidOpenGL.cpp
@@ -1107,16 +1107,6 @@ void FAndroidOpenGL::ProcessExtensions(const FString& ExtensionsString)
 		bSupportsInstancing = false;
 	}
 
-	// PowerVR Rogue doesn't like glVertexAttribIPointer so disable it
-	if (bIsPoverVRBased && bES30Support)
-	{
-		if (RendererString.Contains(TEXT("Rogue")))
-		{
-			glVertexAttribIPointer = nullptr;
-			UE_LOG(LogRHI, Warning, TEXT("Disabling glVertexAttribIPointer on PowerVR Rogue"));
-		}
-	}
-
 	if (bSupportsBGRA8888)
 	{
 		// Check whether device supports BGRA as color attachment
diff --git a/Engine/Source/Runtime/OpenGLDrv/Private/Android/AndroidOpenGL.h b/Engine/Source/Runtime/OpenGLDrv/Private/Android/AndroidOpenGL.h
--- a/Engine/Source/Runtime/OpenGLDrv/Private/Android/AndroidOpenGL.h
+++ b/Engine/Source/Runtime/OpenGLDrv/Private/Android/AndroidOpenGL.h
@@ -519,7 +519,7 @@ struct FAndroidOpenGL : public FOpenGLES2
 
 	static FORCEINLINE void VertexAttribIPointer(GLuint Index, GLint Size, GLenum Type, GLsizei Stride, const GLvoid* Pointer)
 	{
-		if (IsES31Usable() && glVertexAttribIPointer != nullptr)
+		if (IsES31Usable())
 		{
 			glVertexAttribIPointer(Index, Size, Type, Stride, Pointer);
 		}
