Add option to disable opengl program linking checks.
Used with strategy game.

diff --git a/Engine/Source/Runtime/OpenGLDrv/Private/OpenGLShaders.cpp b/Engine/Source/Runtime/OpenGLDrv/Private/OpenGLShaders.cpp
--- a/Engine/Source/Runtime/OpenGLDrv/Private/OpenGLShaders.cpp
+++ b/Engine/Source/Runtime/OpenGLDrv/Private/OpenGLShaders.cpp
@@ -70,6 +70,14 @@ static TAutoConsoleVariable<int32> CVarLRUKeepProgramBinaryResident(
 	TEXT("1: Program binary is retained, this improves eviction and re-creation performance but uses more memory."),
 	ECVF_ReadOnly |ECVF_RenderThreadSafe);
 
+static TAutoConsoleVariable<int32> CVarIgnoreLinkFailure(
+	TEXT("r.OpenGL.IgnoreLinkFailure"),
+	0,
+	TEXT("Ignore OpenGL program link failures.\n")
+	TEXT("0: Program link failure generates a fatal error when encountered. (default)\n")
+	TEXT("1: Ignore link failures. this may allow a program to continue but could lead to undefined rendering behaviour."),
+	ECVF_RenderThreadSafe);
+
 #if PLATFORM_ANDROID
 bool GOpenGLShaderHackLastCompileSuccess = false;
 #endif
@@ -176,7 +184,7 @@ static bool VerifyLinkedProgram(GLuint Program, VerifyProgramPipelineFailurePoli
 {
 	auto LogError = [](const FString&& InMessage, VerifyProgramPipelineFailurePolicy InFailurePolicy)
 {
-		bool bCatchError = InFailurePolicy == VerifyProgramPipelineFailurePolicy::CatchFailure;
+		bool bCatchError = InFailurePolicy == VerifyProgramPipelineFailurePolicy::CatchFailure && CVarIgnoreLinkFailure.GetValueOnAnyThread() == 0;
 		if (bCatchError)
 		{
 			UE_LOG(LogRHI, Fatal, TEXT("%s"), *InMessage);
@@ -215,7 +223,8 @@ static bool VerifyLinkedProgram(GLuint Program, VerifyProgramPipelineFailurePoli
 #else
 		LogError(FString::Printf(TEXT("Failed to link program. Current total programs:%d"), GNumPrograms), FailurePolicy);
 #endif
-		return false;
+		// if we're required to ignore link failure then we return true here.
+		return CVarIgnoreLinkFailure.GetValueOnAnyThread() == 1;
 	}
 	return true;
 }
@@ -3151,9 +3160,6 @@ static FOpenGLLinkedProgram* LinkProgram( const FOpenGLLinkedProgramConfiguratio
 #endif
 	else
 	{
-		FName LinkFailurePanic = bFromPSOFileCache ? FName("FailedProgramLinkDuringPrecompile") : FName("FailedProgramLink");
-		RHIGetPanicDelegate().ExecuteIfBound(LinkFailurePanic);
-		UE_LOG(LogRHI, Fatal, TEXT("Failed to link program. Current total programs: %d, precompile: %d"), GNumPrograms, (uint32)bFromPSOFileCache);
 		return nullptr;
 	}
 	
@@ -3679,42 +3685,45 @@ FBoundShaderStateRHIRef FOpenGLDynamicRHI::RHICreateBoundShaderState_OnThisThrea
 
 				// Link program, using the data provided in config
 				LinkedProgram = LinkProgram(Config, bFromPSOFileCache);
-				GetOpenGLProgramsCache().Add(Config.ProgramKey, LinkedProgram);
-
-				// if building the cache file and using the LRU then evict the last shader created. this will reduce the risk of fragmentation of the driver's program memory.
-				if (bFindAndCreateEvictedProgram == false && FOpenGLProgramBinaryCache::IsBuildingCache())
-				{
-					GetOpenGLProgramsCache().EvictMostRecent();
-				}
 
 				if (LinkedProgram == NULL)
 				{
 #if DEBUG_GL_SHADERS
-					if (VertexShader->bSuccessfullyCompiled)
+					if (VertexShader->bSuccessfullyCompiled || FOpenGLProgramBinaryCache::IsEnabled())
 					{
-						UE_LOG(LogRHI,Error,TEXT("Vertex Shader:\n%s"),ANSI_TO_TCHAR(VertexShader->GlslCode.GetData()));
+						UE_LOG(LogRHI, Error, TEXT("Vertex Shader:\n%s"), ANSI_TO_TCHAR(VertexShader->GlslCode.GetData()));
 					}
-					if (PixelShader->bSuccessfullyCompiled)
+					if (PixelShader->bSuccessfullyCompiled || FOpenGLProgramBinaryCache::IsEnabled())
 					{
-						UE_LOG(LogRHI,Error,TEXT("Pixel Shader:\n%s"),ANSI_TO_TCHAR(PixelShader->GlslCode.GetData()));
+						UE_LOG(LogRHI, Error, TEXT("Pixel Shader:\n%s"), ANSI_TO_TCHAR(PixelShader->GlslCode.GetData()));
 					}
 					if (GeometryShader && GeometryShader->bSuccessfullyCompiled)
 					{
-						UE_LOG(LogRHI,Error,TEXT("Geometry Shader:\n%s"),ANSI_TO_TCHAR(GeometryShader->GlslCode.GetData()));
+						UE_LOG(LogRHI, Error, TEXT("Geometry Shader:\n%s"), ANSI_TO_TCHAR(GeometryShader->GlslCode.GetData()));
 					}
-					if ( FOpenGL::SupportsTessellation() )
+					if (FOpenGL::SupportsTessellation())
 					{
 						if (HullShader && HullShader->bSuccessfullyCompiled)
 						{
-							UE_LOG(LogRHI,Error,TEXT("Hull Shader:\n%s"),ANSI_TO_TCHAR(HullShader->GlslCode.GetData()));
+							UE_LOG(LogRHI, Error, TEXT("Hull Shader:\n%s"), ANSI_TO_TCHAR(HullShader->GlslCode.GetData()));
 						}
 						if (DomainShader && DomainShader->bSuccessfullyCompiled)
 						{
-							UE_LOG(LogRHI,Error,TEXT("Domain Shader:\n%s"),ANSI_TO_TCHAR(DomainShader->GlslCode.GetData()));
+							UE_LOG(LogRHI, Error, TEXT("Domain Shader:\n%s"), ANSI_TO_TCHAR(DomainShader->GlslCode.GetData()));
 						}
 					}
 #endif //DEBUG_GL_SHADERS
-					check(LinkedProgram);
+					FName LinkFailurePanic = bFromPSOFileCache ? FName("FailedProgramLinkDuringPrecompile") : FName("FailedProgramLink");
+					RHIGetPanicDelegate().ExecuteIfBound(LinkFailurePanic);
+					UE_LOG(LogRHI, Fatal, TEXT("Failed to link program. Current total programs: %d, precompile: %d"), GNumPrograms, (uint32)bFromPSOFileCache);
+				}
+
+				GetOpenGLProgramsCache().Add(Config.ProgramKey, LinkedProgram);
+
+				// if building the cache file and using the LRU then evict the last shader created. this will reduce the risk of fragmentation of the driver's program memory.
+				if (bFindAndCreateEvictedProgram == false && FOpenGLProgramBinaryCache::IsBuildingCache())
+				{
+					GetOpenGLProgramsCache().EvictMostRecent();
 				}
 			}
 		}
