Fix incorrect ensure firing when OpenGL shader LRU is evicting unused shaders.

diff --git a/Engine/Source/Runtime/OpenGLDrv/Private/OpenGLShaders.cpp b/Engine/Source/Runtime/OpenGLDrv/Private/OpenGLShaders.cpp
--- a/Engine/Source/Runtime/OpenGLDrv/Private/OpenGLShaders.cpp
+++ b/Engine/Source/Runtime/OpenGLDrv/Private/OpenGLShaders.cpp
@@ -2545,8 +2545,8 @@ void FDelayedEvictionContainer::Tick()
 			++It;
 			bMeasureEviction = true;
 			check(LinkedProgram->LRUInfo.EvictBucket == EvictBucketIndex);
+			LinkedProgram->LRUInfo.EvictBucket = -3; // Mark EvictBucket to indicated evicted from ProgramsToEvict, Prevent EvictProgram from attempting to remove again.
 			GetOpenGLProgramsCache().EvictProgram(LinkedProgram->Config.ProgramKey);
-			LinkedProgram->LRUInfo.EvictBucket = -3;
 			bMeasureEviction = false;
 		}
 	}
