#DEV-VR - Wrapped Android functions on Lumin

diff --git a/Engine/Source/Runtime/OpenGLDrv/Private/OpenGLDevice.cpp b/Engine/Source/Runtime/OpenGLDrv/Private/OpenGLDevice.cpp
--- a/Engine/Source/Runtime/OpenGLDrv/Private/OpenGLDevice.cpp
+++ b/Engine/Source/Runtime/OpenGLDrv/Private/OpenGLDevice.cpp
@@ -1220,7 +1220,7 @@ static bool VerifyCompiledShader(GLuint Shader, const ANSICHAR* GlslCode, bool I
 }
 #endif
 
-#if PLATFORM_ANDROID
+#if PLATFORM_ANDROID && !PLATFORM_LUMINGL4
 static void AttemptLinkProgramExecute(GLuint VertexShaderResource, GLuint PixelShaderResource)
 {
 	GLuint Program = glCreateProgram();
