DR - Copy/Proper fix for 4600165
Fix for skylight on OpenGL

diff --git a/Engine/Source/Runtime/OpenGLDrv/Public/OpenGLDrv.h b/Engine/Source/Runtime/OpenGLDrv/Public/OpenGLDrv.h
--- a/Engine/Source/Runtime/OpenGLDrv/Public/OpenGLDrv.h
+++ b/Engine/Source/Runtime/OpenGLDrv/Public/OpenGLDrv.h
@@ -961,13 +961,6 @@ public:
 
 	virtual FSamplerStateRHIRef RHICreateSamplerState(const FSamplerStateInitializerRHI& Initializer) final override;
 
-	virtual void RHIReadSurfaceFloatData_RenderThread(class FRHICommandListImmediate& RHICmdList, FTextureRHIParamRef Texture, FIntRect Rect, TArray<FFloat16Color>& OutData, ECubeFace CubeFace, int32 ArrayIndex, int32 MipIndex) final override
-	{
-		RHITHREAD_GLCOMMAND_PROLOGUE();
-		RHIReadSurfaceFloatData(Texture, Rect, OutData, CubeFace, ArrayIndex, MipIndex);
-		RHITHREAD_GLCOMMAND_EPILOGUE();
-	}
-
 	virtual void UpdateTexture2D_RenderThread(class FRHICommandListImmediate& RHICmdList, FTexture2DRHIParamRef Texture, uint32 MipIndex, const struct FUpdateTextureRegion2D& UpdateRegion, uint32 SourcePitch, const uint8* SourceData) final override
 	{
 		GDynamicRHI->RHIUpdateTexture2D(Texture, MipIndex, UpdateRegion, SourcePitch, SourceData);
