Copy from 4582523 Restore OpenGL ShaderCodeLibrary hooks.

diff --git a/Engine/Source/Runtime/OpenGLDrv/Private/OpenGLDevice.cpp b/Engine/Source/Runtime/OpenGLDrv/Private/OpenGLDevice.cpp
--- a/Engine/Source/Runtime/OpenGLDrv/Private/OpenGLDevice.cpp
+++ b/Engine/Source/Runtime/OpenGLDrv/Private/OpenGLDevice.cpp
@@ -1609,6 +1609,7 @@ void FOpenGLDynamicRHI::Init()
 	VERIFY_GL_SCOPE();
 
 	FOpenGLProgramBinaryCache::Initialize();
+	RegisterSharedShaderCodeDelegates();
 	InitializeStateResources();
 
 	// Create a default point sampler state for internal use.
