Support R32, R16 integer formats on E3.0 and up

diff --git a/Engine/Source/Runtime/OpenGLDrv/Private/OpenGLDevice.cpp b/Engine/Source/Runtime/OpenGLDrv/Private/OpenGLDevice.cpp
--- a/Engine/Source/Runtime/OpenGLDrv/Private/OpenGLDevice.cpp
+++ b/Engine/Source/Runtime/OpenGLDrv/Private/OpenGLDevice.cpp
@@ -1058,6 +1058,18 @@ static void InitRHICapabilitiesForGL()
 			SetupTextureFormat( PF_R16F_FILTER,		FOpenGLTextureFormat( GL_R16F,					GL_R16F,				GL_RED,			GL_HALF_FLOAT,						false,	false));
 		}
 
+		// ES3.0 and up 
+		// should expose some function to identify ES version, to make it cleaner
+#if PLATFORM_ANDROID 
+		if (FOpenGL::UseES30ShadingLanguage())
+		{
+			SetupTextureFormat( PF_R32_UINT,		FOpenGLTextureFormat( GL_R32UI,					GL_R32UI,				GL_RED_INTEGER,	GL_UNSIGNED_INT,					false,	false));
+			SetupTextureFormat( PF_R32_SINT,		FOpenGLTextureFormat( GL_R32I,					GL_R32I,				GL_RED_INTEGER,	GL_INT,								false,	false));
+			SetupTextureFormat( PF_R16_UINT,		FOpenGLTextureFormat( GL_R16UI,					GL_R16UI,				GL_RED_INTEGER,	GL_UNSIGNED_SHORT,					false,	false));
+			SetupTextureFormat( PF_R16_SINT,		FOpenGLTextureFormat( GL_R16I,					GL_R16I,				GL_RED_INTEGER,	GL_SHORT,							false,	false));
+		}
+#endif
+	
 		if (FOpenGL::SupportsPackedDepthStencil())
 		{
 			SetupTextureFormat(PF_DepthStencil, FOpenGLTextureFormat(GL_DEPTH_STENCIL_OES, GL_NONE, GL_DEPTH_STENCIL_OES, GL_UNSIGNED_INT_24_8_OES, false, false));
