Added RenderDoc detection on Android devices running GL

diff --git a/Engine/Source/Runtime/OpenGLDrv/Private/Android/AndroidOpenGL.cpp b/Engine/Source/Runtime/OpenGLDrv/Private/Android/AndroidOpenGL.cpp
--- a/Engine/Source/Runtime/OpenGLDrv/Private/Android/AndroidOpenGL.cpp
+++ b/Engine/Source/Runtime/OpenGLDrv/Private/Android/AndroidOpenGL.cpp
@@ -136,10 +136,17 @@ FPlatformOpenGLDevice::FPlatformOpenGLDevice()
 // call out to JNI to see if the application was packaged for Gear VR
 extern bool AndroidThunkCpp_IsGearVRApplication();
 
+
+// RenderDoc
+#define GL_DEBUG_TOOL_EXT	0x6789
+static bool bRunningUnderRenderDoc = false;
+
 void FPlatformOpenGLDevice::Init()
 {
 	extern void InitDebugContext();
 
+	bRunningUnderRenderDoc = glIsEnabled(GL_DEBUG_TOOL_EXT) != GL_FALSE;
+
 	FPlatformMisc::LowLevelOutputDebugString(TEXT("FPlatformOpenGLDevice:Init"));
 	bool bCreateSurface = !AndroidThunkCpp_IsGearVRApplication();
 	AndroidEGL::GetInstance()->InitSurface(false, bCreateSurface);
@@ -166,7 +173,7 @@ FPlatformOpenGLDevice* PlatformCreateOpenGLDevice()
 
 bool PlatformCanEnableGPUCapture()
 {
-	return false;
+	return bRunningUnderRenderDoc;
 }
 
 void PlatformReleaseOpenGLContext(FPlatformOpenGLDevice* Device, FPlatformOpenGLContext* Context)
