OpenGL implementation of Fence and CopyToStagingBuffer

diff --git a/Engine/Source/Runtime/OpenGLDrv/Private/OpenGLCommands.cpp b/Engine/Source/Runtime/OpenGLDrv/Private/OpenGLCommands.cpp
--- a/Engine/Source/Runtime/OpenGLDrv/Private/OpenGLCommands.cpp
+++ b/Engine/Source/Runtime/OpenGLDrv/Private/OpenGLCommands.cpp
@@ -3593,6 +3593,34 @@ void FOpenGLDynamicRHI::RHIInvalidateCachedState()
 	SharedContextState.InitializeResources(FOpenGL::GetMaxCombinedTextureImageUnits(), OGL_MAX_COMPUTE_STAGE_UAV_UNITS);
 }
 
+void FOpenGLDynamicRHI::RHICopyToStagingBuffer(FVertexBufferRHIParamRef SourceBufferRHI, FStagingBufferRHIParamRef DestinationStagingBufferRHI, uint32 InOffset, uint32 InNumBytes, FGPUFenceRHIParamRef FenceRHI)
+{
+#if OPENGL_GL3
+	VERIFY_GL_SCOPE();
+	FOpenGLVertexBuffer* SourceBuffer = ResourceCast(SourceBufferRHI);
+	FOpenGLStagingBuffer* DestinationBuffer = ResourceCast(DestinationStagingBufferRHI);
+	FOpenGLGPUFence* CopyFence = ResourceCast(FenceRHI);
+
+	check(DestinationBuffer->ShadowBuffer != 0);
+	glBindBuffer(GL_COPY_WRITE_BUFFER, DestinationBuffer->ShadowBuffer);
+	if (DestinationBuffer->ShadowSize < InNumBytes)
+	{
+		// Orphan the existing buffer.
+		glBufferData(GL_COPY_WRITE_BUFFER, InNumBytes, NULL, GL_STREAM_READ);
+		DestinationBuffer->ShadowSize = InNumBytes;
+	}
+
+	glBindBuffer(GL_COPY_READ_BUFFER, SourceBuffer->Resource);
+	FOpenGL::CopyBufferSubData(GL_COPY_READ_BUFFER, GL_COPY_WRITE_BUFFER, InOffset, 0, InNumBytes);
+	glBindBuffer(GL_COPY_READ_BUFFER, 0);
+	glBindBuffer(GL_COPY_WRITE_BUFFER, 0);
+
+	CopyFence->WriteInternal();
+#else
+	UE_LOG(LogRHI, Fatal, TEXT("Only OpenGL3 or higher supports CopyToStagingBuffer"));
+#endif
+}
+
 #if PLATFORM_USES_FIXED_RHI_CLASS
 #define INTERNAL_DECORATOR(Method) ((FOpenGLDynamicRHI&)CmdList.GetContext()).FOpenGLDynamicRHI::Method
 #include "RHICommandListCommandExecutes.inl"
diff --git a/Engine/Source/Runtime/OpenGLDrv/Private/OpenGLQuery.cpp b/Engine/Source/Runtime/OpenGLDrv/Private/OpenGLQuery.cpp
--- a/Engine/Source/Runtime/OpenGLDrv/Private/OpenGLQuery.cpp
+++ b/Engine/Source/Runtime/OpenGLDrv/Private/OpenGLQuery.cpp
@@ -1135,3 +1135,91 @@ void FOpenGLDisjointTimeStampQuery::ReleaseResources()
 		PlatformReleaseRenderQuery(DisjointQuery, Context);
 	}
 }
+
+
+
+// Fence implementation
+
+FGPUFenceRHIRef FOpenGLDynamicRHI::RHICreateGPUFence(const FName &Name)
+{
+#if OPENGL_GL3
+	return new FOpenGLGPUFence(Name);
+#else
+	UE_LOG(LogRHI, Fatal, TEXT("Fences are only available in OpenGL3 or later"));
+#endif
+}
+
+FOpenGLGPUFence::~FOpenGLGPUFence()
+{
+#if OPENGL_GL3
+	if (bValidSync)
+	{
+		FRHICommandListImmediate& RHICmdList = FRHICommandListExecutor::GetImmediateCommandList();
+		RHITHREAD_GLCOMMAND_PROLOGUE();
+		VERIFY_GL_SCOPE();
+		FOpenGL::DeleteSync(Fence);
+		RHITHREAD_GLCOMMAND_EPILOGUE_NORETURN();
+	}
+#else
+	UE_LOG(LogRHI, Fatal, TEXT("Fences are only available in OpenGL3 or later"));
+#endif
+}
+
+void FOpenGLGPUFence::Clear()
+{
+#if OPENGL_GL3
+
+	if (bValidSync)
+	{
+		FRHICommandListImmediate& RHICmdList = FRHICommandListExecutor::GetImmediateCommandList();
+		RHITHREAD_GLCOMMAND_PROLOGUE();
+		VERIFY_GL_SCOPE();
+		FOpenGL::DeleteSync(Fence);
+		bValidSync = false;
+		RHITHREAD_GLCOMMAND_EPILOGUE();
+	}
+#else
+	UE_LOG(LogRHI, Fatal, TEXT("Fences are only available in OpenGL3 or later"));
+#endif
+}
+bool FOpenGLGPUFence::Poll() const
+{
+#if OPENGL_GL3
+	if (!bValidSync)
+	{
+		return false;
+	}
+
+	FRHICommandListImmediate& RHICmdList = FRHICommandListExecutor::GetImmediateCommandList();
+	RHITHREAD_GLCOMMAND_PROLOGUE();
+	VERIFY_GL_SCOPE();
+
+	FOpenGLBase::EFenceResult Result = FOpenGL::ClientWaitSync(Fence, 0, 0);
+	return (Result == FOpenGLBase::FR_AlreadySignaled || Result == FOpenGLBase::FR_ConditionSatisfied);
+
+	RHITHREAD_GLCOMMAND_EPILOGUE_RETURN(bool);
+#else
+	UE_LOG(LogRHI, Fatal, TEXT("Fences are only available in OpenGL3 or later"));
+#endif
+}
+
+void FOpenGLGPUFence::WriteInternal()
+{
+#if OPENGL_GL3
+	FRHICommandListImmediate& RHICmdList = FRHICommandListExecutor::GetImmediateCommandList();
+	RHITHREAD_GLCOMMAND_PROLOGUE();
+	VERIFY_GL_SCOPE();
+
+	if (bValidSync)
+	{
+		FOpenGL::DeleteSync(Fence);
+		bValidSync = false;
+	}
+
+	Fence = FOpenGL::FenceSync(GL_SYNC_GPU_COMMANDS_COMPLETE, 0);
+	bValidSync = true;
+	RHITHREAD_GLCOMMAND_EPILOGUE();
+#else
+	UE_LOG(LogRHI, Fatal, TEXT("Fences are only available in OpenGL3 or later"));
+#endif
+}
\ No newline at end of file
diff --git a/Engine/Source/Runtime/OpenGLDrv/Private/OpenGLVertexBuffer.cpp b/Engine/Source/Runtime/OpenGLDrv/Private/OpenGLVertexBuffer.cpp
--- a/Engine/Source/Runtime/OpenGLDrv/Private/OpenGLVertexBuffer.cpp
+++ b/Engine/Source/Runtime/OpenGLDrv/Private/OpenGLVertexBuffer.cpp
@@ -253,3 +253,89 @@ void FOpenGLDynamicRHI::RHICopyVertexBuffer(FVertexBufferRHIParamRef SourceBuffe
 	glBindBuffer(GL_COPY_WRITE_BUFFER,0);
 }
 
+FStagingBufferRHIRef FOpenGLDynamicRHI::RHICreateStagingBuffer()
+{
+#if OPENGL_GL3
+	return new FOpenGLStagingBuffer();
+#else
+	UE_LOG(LogRHI, Fatal, TEXT("Staging Buffers are only available in OpenGL3 or later"));
+#endif
+}
+
+void FOpenGLStagingBuffer::Initialize()
+{
+	ShadowBuffer = 0;
+	ShadowSize = 0;
+
+	FRHICommandListImmediate& RHICmdList = FRHICommandListExecutor::GetImmediateCommandList();
+	RHITHREAD_GLCOMMAND_PROLOGUE();
+	VERIFY_GL_SCOPE();
+	glGenBuffers(1, &ShadowBuffer);
+	RHITHREAD_GLCOMMAND_EPILOGUE();
+}
+
+FOpenGLStagingBuffer::~FOpenGLStagingBuffer()
+{
+	FRHICommandListImmediate& RHICmdList = FRHICommandListExecutor::GetImmediateCommandList();
+	RHITHREAD_GLCOMMAND_PROLOGUE();
+	VERIFY_GL_SCOPE();
+	glDeleteBuffers(1, &ShadowBuffer);
+	RHITHREAD_GLCOMMAND_EPILOGUE_NORETURN();
+}
+
+// I don't see a way to do this without stalling the RHI thread.
+void* FOpenGLStagingBuffer::Lock(uint32 Offset, uint32 NumBytes)
+{
+#if OPENGL_GL3
+	FRHICommandListImmediate& RHICmdList = FRHICommandListExecutor::GetImmediateCommandList();
+	RHITHREAD_GLCOMMAND_PROLOGUE();
+	VERIFY_GL_SCOPE();
+
+	check(ShadowBuffer != 0);
+	glBindBuffer(GL_COPY_WRITE_BUFFER, ShadowBuffer);
+	void* Mapping = glMapBuffer(GL_COPY_WRITE_BUFFER, GL_READ_ONLY);
+	check(Mapping);
+	return reinterpret_cast<uint8*>(Mapping) + Offset;
+
+	RHITHREAD_GLCOMMAND_EPILOGUE_RETURN(void*);
+#else
+	UE_LOG(LogRHI, Fatal, TEXT("Staging Buffers are only available in OpenGL3 or later"));
+#endif
+}
+
+// Unfortunately I think we have to stall the RHI thread here as well to play nice with OpenGL.
+// Since this will probably be close to a call to lock we've probably paid most of the cost already.
+void FOpenGLStagingBuffer::Unlock()
+{
+#if OPENGL_GL3
+	FRHICommandListImmediate& RHICmdList = FRHICommandListExecutor::GetImmediateCommandList();
+	RHITHREAD_GLCOMMAND_PROLOGUE();
+	glUnmapBuffer(GL_COPY_WRITE_BUFFER);
+	glBindBuffer(GL_COPY_WRITE_BUFFER, 0);
+	RHITHREAD_GLCOMMAND_EPILOGUE();
+#else
+	UE_LOG(LogRHI, Fatal, TEXT("Staging Buffers are only available in OpenGL3 or later"));
+#endif
+}
+
+void* FOpenGLDynamicRHI::RHILockStagingBuffer(FStagingBufferRHIParamRef StagingBuffer, uint32 Offset, uint32 SizeRHI)
+{
+	FOpenGLStagingBuffer* Buffer = ResourceCast(StagingBuffer);
+	return Buffer->Lock(Offset, SizeRHI);	
+}
+
+void FOpenGLDynamicRHI::RHIUnlockStagingBuffer(FStagingBufferRHIParamRef StagingBuffer)
+{
+	FOpenGLStagingBuffer* Buffer = ResourceCast(StagingBuffer);
+	Buffer->Unlock();
+}
+
+void* FOpenGLDynamicRHI::LockStagingBuffer_RenderThread(class FRHICommandListImmediate& RHICmdList, FStagingBufferRHIParamRef StagingBuffer, uint32 Offset, uint32 SizeRHI)
+{
+	return RHILockStagingBuffer(StagingBuffer, Offset, SizeRHI);
+}
+
+void FOpenGLDynamicRHI::UnlockStagingBuffer_RenderThread(class FRHICommandListImmediate& RHICmdList, FStagingBufferRHIParamRef StagingBuffer)
+{
+	RHIUnlockStagingBuffer(StagingBuffer);
+}
\ No newline at end of file
diff --git a/Engine/Source/Runtime/OpenGLDrv/Public/OpenGLDrv.h b/Engine/Source/Runtime/OpenGLDrv/Public/OpenGLDrv.h
--- a/Engine/Source/Runtime/OpenGLDrv/Public/OpenGLDrv.h
+++ b/Engine/Source/Runtime/OpenGLDrv/Public/OpenGLDrv.h
@@ -533,6 +533,16 @@ public:
 	virtual FTextureCubeRHIRef RHICreateTextureCubeFromResource(EPixelFormat Format, uint32 Size, bool bArray, uint32 ArraySize, uint32 NumMips, uint32 NumSamples, uint32 NumSamplesTileMem, const FClearValueBinding& ClearValueBinding, GLuint Resource, uint32 Flags);
 	virtual void RHIAliasTextureResources(FTextureRHIParamRef DestTexture, FTextureRHIParamRef SrcTexture);
 
+	// Inline copy
+	virtual void RHICopyToStagingBuffer(FVertexBufferRHIParamRef SourceBufferRHI, FStagingBufferRHIParamRef DestinationStagingBufferRHI, uint32 InOffset, uint32 InNumBytes, FGPUFenceRHIParamRef FenceRHI) final override;
+	virtual void* RHILockStagingBuffer(FStagingBufferRHIParamRef StagingBuffer, uint32 Offset, uint32 SizeRHI) final override;
+	virtual void RHIUnlockStagingBuffer(FStagingBufferRHIParamRef StagingBuffer) final override;
+	virtual void* LockStagingBuffer_RenderThread(class FRHICommandListImmediate& RHICmdList, FStagingBufferRHIParamRef StagingBuffer, uint32 Offset, uint32 SizeRHI) final override;
+	virtual	void UnlockStagingBuffer_RenderThread(class FRHICommandListImmediate& RHICmdList, FStagingBufferRHIParamRef StagingBuffer) final override;
+
+	virtual FStagingBufferRHIRef RHICreateStagingBuffer() final override;
+	virtual FGPUFenceRHIRef RHICreateGPUFence(const FName &Name) final override;
+
 	void Cleanup();
 
 	void PurgeFramebufferFromCaches(GLuint Framebuffer);
@@ -686,6 +696,18 @@ public:
 			RHICmdList.ImmediateFlush(EImmediateFlushType::FlushRHIThread);\
 		}\
 
+#define RHITHREAD_GLCOMMAND_EPILOGUE_NORETURN() };\
+		if (ShouldRunGLRenderContextOpOnThisThread(RHICmdList))\
+		{\
+			GLCommand();\
+		}\
+		else\
+		{\
+			ALLOC_COMMAND_CL(RHICmdList, FRHICommandGLCommand)( GLCommand ); \
+			RHITHREAD_GLTRACE_BLOCKING;\
+			RHICmdList.ImmediateFlush(EImmediateFlushType::FlushRHIThread);\
+		}\
+
 	struct FTextureLockTracker
 	{
 		struct FLockParams
diff --git a/Engine/Source/Runtime/OpenGLDrv/Public/OpenGLResources.h b/Engine/Source/Runtime/OpenGLDrv/Public/OpenGLResources.h
--- a/Engine/Source/Runtime/OpenGLDrv/Public/OpenGLResources.h
+++ b/Engine/Source/Runtime/OpenGLDrv/Public/OpenGLResources.h
@@ -2010,11 +2010,67 @@ private:
 	FCustomPresentRHIRef CustomPresent;
 };
 
+// Fences
+
+// Note that Poll() and WriteInternal() will stall the RHI thread if one is present.
+class FOpenGLGPUFence : public FRHIGPUFence
+{
+public:
+	FOpenGLGPUFence(FName InName)
+		: FRHIGPUFence(InName)
+		, bValidSync(false)
+	{}
+
+	virtual ~FOpenGLGPUFence() final override;
+
+	virtual void Clear() final override;
+	virtual bool Poll() const final override;
+	
+	void WriteInternal();
+private:
+	UGLsync Fence;
+	// We shadow the sync state to know if/when we need to destroy it.
+	bool bValidSync;
+};
+
+class FOpenGLStagingBuffer : public FRHIStagingBuffer
+{
+	friend class FOpenGLDynamicRHI;
+public:
+	FOpenGLStagingBuffer() : FRHIStagingBuffer()
+	{
+		Initialize();
+	}
+
+	virtual ~FOpenGLStagingBuffer() final override;
+
+	// Locks the shadow of VertexBuffer for read. This will stall the RHI thread.
+	virtual void *Lock(uint32 Offset, uint32 NumBytes) final override;
+
+	// Unlocks the shadow. This is an error if it was not locked previously.
+	virtual void Unlock() final override;
+private:
+	void Initialize();
+
+	GLuint ShadowBuffer;
+	uint32 ShadowSize;
+};
+
 template<class T>
 struct TOpenGLResourceTraits
 {
 };
 template<>
+struct TOpenGLResourceTraits<FRHIGPUFence>
+{
+	typedef FOpenGLGPUFence TConcreteType;
+};
+template<>
+struct TOpenGLResourceTraits<FRHIStagingBuffer>
+{
+	typedef FOpenGLStagingBuffer TConcreteType;
+};
+template<>
 struct TOpenGLResourceTraits<FRHIVertexDeclaration>
 {
 	typedef FOpenGLVertexDeclaration TConcreteType;
