Log Fatal doesn't have noreturn, so return null.

diff --git a/Engine/Source/Runtime/OpenGLDrv/Private/OpenGLQuery.cpp b/Engine/Source/Runtime/OpenGLDrv/Private/OpenGLQuery.cpp
--- a/Engine/Source/Runtime/OpenGLDrv/Private/OpenGLQuery.cpp
+++ b/Engine/Source/Runtime/OpenGLDrv/Private/OpenGLQuery.cpp
@@ -1146,6 +1146,7 @@ FGPUFenceRHIRef FOpenGLDynamicRHI::RHICreateGPUFence(const FName &Name)
 	return new FOpenGLGPUFence(Name);
 #else
 	UE_LOG(LogRHI, Fatal, TEXT("Fences are only available in OpenGL3 or later"));
+	return nullptr;
 #endif
 }
 
@@ -1200,6 +1201,7 @@ bool FOpenGLGPUFence::Poll() const
 	RHITHREAD_GLCOMMAND_EPILOGUE_RETURN(bool);
 #else
 	UE_LOG(LogRHI, Fatal, TEXT("Fences are only available in OpenGL3 or later"));
+	return false;
 #endif
 }
 
@@ -1222,4 +1224,4 @@ void FOpenGLGPUFence::WriteInternal()
 #else
 	UE_LOG(LogRHI, Fatal, TEXT("Fences are only available in OpenGL3 or later"));
 #endif
-}
\ No newline at end of file
+}
diff --git a/Engine/Source/Runtime/OpenGLDrv/Private/OpenGLVertexBuffer.cpp b/Engine/Source/Runtime/OpenGLDrv/Private/OpenGLVertexBuffer.cpp
--- a/Engine/Source/Runtime/OpenGLDrv/Private/OpenGLVertexBuffer.cpp
+++ b/Engine/Source/Runtime/OpenGLDrv/Private/OpenGLVertexBuffer.cpp
@@ -259,6 +259,7 @@ FStagingBufferRHIRef FOpenGLDynamicRHI::RHICreateStagingBuffer()
 	return new FOpenGLStagingBuffer();
 #else
 	UE_LOG(LogRHI, Fatal, TEXT("Staging Buffers are only available in OpenGL3 or later"));
+	return nullptr;
 #endif
 }
 
@@ -300,6 +301,7 @@ void* FOpenGLStagingBuffer::Lock(uint32 Offset, uint32 NumBytes)
 	RHITHREAD_GLCOMMAND_EPILOGUE_RETURN(void*);
 #else
 	UE_LOG(LogRHI, Fatal, TEXT("Staging Buffers are only available in OpenGL3 or later"));
+	return nullptr;
 #endif
 }
 
@@ -338,4 +340,4 @@ void* FOpenGLDynamicRHI::LockStagingBuffer_RenderThread(class FRHICommandListImm
 void FOpenGLDynamicRHI::UnlockStagingBuffer_RenderThread(class FRHICommandListImmediate& RHICmdList, FStagingBufferRHIParamRef StagingBuffer)
 {
 	RHIUnlockStagingBuffer(StagingBuffer);
-}
\ No newline at end of file
+}
