Add GL LRU program touch to programs found in small release cache. Fixes rare case in which evicted programs are still in the recently released cache.
#rnx

diff --git a/Engine/Source/Runtime/OpenGLDrv/Private/OpenGLShaders.cpp b/Engine/Source/Runtime/OpenGLDrv/Private/OpenGLShaders.cpp
--- a/Engine/Source/Runtime/OpenGLDrv/Private/OpenGLShaders.cpp
+++ b/Engine/Source/Runtime/OpenGLDrv/Private/OpenGLShaders.cpp
@@ -2304,18 +2304,18 @@ public:
 		}
 	}
 
-	FORCEINLINE_DEBUGGABLE void Touch(const FOpenGLBoundShaderState* OpenGLBoundShaderState)
+	FORCEINLINE_DEBUGGABLE void Touch(FOpenGLLinkedProgram* LinkedProgram)
 	{
-		if(OpenGLBoundShaderState->LinkedProgram->LRUInfo.LRUNode.IsValidId())
+		if(LinkedProgram->LRUInfo.LRUNode.IsValidId())
 		{
-			LRU.MarkAsRecent(OpenGLBoundShaderState->LinkedProgram->LRUInfo.LRUNode);
+			LRU.MarkAsRecent(LinkedProgram->LRUInfo.LRUNode);
 		}
 		else
 		{
 			// This must find the program.
-			ensure(FindEvictedAndUpdateLRU(OpenGLBoundShaderState->LinkedProgram->Config.ProgramKey));
+			ensure(FindEvictedAndUpdateLRU(LinkedProgram->Config.ProgramKey));
 		}
-		FDelayedEvictionContainer::OnProgramTouched(OpenGLBoundShaderState->LinkedProgram);
+		FDelayedEvictionContainer::OnProgramTouched(LinkedProgram);
 	}
 
 	void Empty()
@@ -2381,11 +2381,11 @@ public:
 		return bUseLRUCache;
 	}
 
-	FORCEINLINE_DEBUGGABLE void Touch(const FOpenGLBoundShaderState* OpenGLBoundShaderState)
+	FORCEINLINE_DEBUGGABLE void Touch(FOpenGLLinkedProgram* LinkedProgram)
 	{
 		if (bUseLRUCache)
 		{
-			ProgramCacheLRU.Touch(OpenGLBoundShaderState);
+			ProgramCacheLRU.Touch(LinkedProgram);
 		}
 	}
 
@@ -3533,7 +3533,7 @@ FBoundShaderStateRHIRef FOpenGLDynamicRHI::RHICreateBoundShaderState_OnThisThrea
 		// If we've already created a bound shader state with these parameters, reuse it.
 		{
 			FOpenGLBoundShaderState* BoundShaderState = ResourceCast(CachedBoundShaderStateLink->BoundShaderState);
-			GetOpenGLProgramsCache().Touch(BoundShaderState);
+			GetOpenGLProgramsCache().Touch(BoundShaderState->LinkedProgram);
 		}
 		return CachedBoundShaderStateLink->BoundShaderState;
 	}
@@ -3607,6 +3607,7 @@ FBoundShaderStateRHIRef FOpenGLDynamicRHI::RHICreateBoundShaderState_OnThisThrea
 			{
 				StaticLastReleasedPrograms[Index] = 0;
 				LinkedProgram = Prog;
+				GetOpenGLProgramsCache().Touch(LinkedProgram);
 				break;
 			}
 			Index = (Index == LAST_RELEASED_PROGRAMS_CACHE_COUNT-1) ? 0 : Index+1;
