Add additional integer texture formats to ES3.1 feature level.
Fixes warnings and issues with Niagara rendering.

diff --git a/Engine/Source/Runtime/OpenGLDrv/Private/OpenGLDevice.cpp b/Engine/Source/Runtime/OpenGLDrv/Private/OpenGLDevice.cpp
--- a/Engine/Source/Runtime/OpenGLDrv/Private/OpenGLDevice.cpp
+++ b/Engine/Source/Runtime/OpenGLDrv/Private/OpenGLDevice.cpp
@@ -1027,6 +1027,16 @@ static void InitRHICapabilitiesForGL()
 		SetupTextureFormat(PF_R8G8B8A8_UINT, FOpenGLTextureFormat(GL_RGBA8UI, GL_RGBA8UI, GL_RGBA_INTEGER, GL_UNSIGNED_BYTE, false, false));
 	#endif
 
+		if (GMaxRHIFeatureLevel == ERHIFeatureLevel::ES3_1)
+		{
+			// These integer formats are supported with ES3.1:
+			SetupTextureFormat(PF_R32_UINT, FOpenGLTextureFormat(GL_R32UI, GL_R32UI, GL_RED_INTEGER, GL_UNSIGNED_INT, false, false));
+			SetupTextureFormat(PF_R32_SINT, FOpenGLTextureFormat(GL_R32I, GL_R32I, GL_RED_INTEGER, GL_INT, false, false));
+			SetupTextureFormat(PF_R16_UINT, FOpenGLTextureFormat(GL_R16UI, GL_R16UI, GL_RED_INTEGER, GL_UNSIGNED_SHORT, false, false));
+			SetupTextureFormat(PF_R16_SINT, FOpenGLTextureFormat(GL_R16I, GL_R16I, GL_RED_INTEGER, GL_SHORT, false, false));
+			SetupTextureFormat(PF_R8_UINT, FOpenGLTextureFormat(GL_R8UI, GL_R8UI, GL_RED_INTEGER, GL_UNSIGNED_BYTE, false, false));
+		}
+
 		if (FOpenGL::SupportsColorBufferHalfFloat() && FOpenGL::SupportsTextureHalfFloat())
 		{
 #if PLATFORM_ANDROID
