Fix iOS CIS build.

diff --git a/Engine/Source/Runtime/OpenGLDrv/Private/OpenGLDevice.cpp b/Engine/Source/Runtime/OpenGLDrv/Private/OpenGLDevice.cpp
--- a/Engine/Source/Runtime/OpenGLDrv/Private/OpenGLDevice.cpp
+++ b/Engine/Source/Runtime/OpenGLDrv/Private/OpenGLDevice.cpp
@@ -1027,6 +1027,7 @@ static void InitRHICapabilitiesForGL()
 		SetupTextureFormat(PF_R8G8B8A8_UINT, FOpenGLTextureFormat(GL_RGBA8UI, GL_RGBA8UI, GL_RGBA_INTEGER, GL_UNSIGNED_BYTE, false, false));
 	#endif
 
+#if PLATFORM_ANDROID
 		if (GMaxRHIFeatureLevel == ERHIFeatureLevel::ES3_1)
 		{
 			// These integer formats are supported with ES3.1:
@@ -1036,6 +1037,7 @@ static void InitRHICapabilitiesForGL()
 			SetupTextureFormat(PF_R16_SINT, FOpenGLTextureFormat(GL_R16I, GL_R16I, GL_RED_INTEGER, GL_SHORT, false, false));
 			SetupTextureFormat(PF_R8_UINT, FOpenGLTextureFormat(GL_R8UI, GL_R8UI, GL_RED_INTEGER, GL_UNSIGNED_BYTE, false, false));
 		}
+#endif
 
 		if (FOpenGL::SupportsColorBufferHalfFloat() && FOpenGL::SupportsTextureHalfFloat())
 		{
