Add PF_R8G8B8A8_SNORM to GL texture formats for devices supporting ES3.0 and above.

diff --git a/Engine/Source/Runtime/OpenGLDrv/Private/OpenGLDevice.cpp b/Engine/Source/Runtime/OpenGLDrv/Private/OpenGLDevice.cpp
--- a/Engine/Source/Runtime/OpenGLDrv/Private/OpenGLDevice.cpp
+++ b/Engine/Source/Runtime/OpenGLDrv/Private/OpenGLDevice.cpp
@@ -1075,6 +1075,9 @@ static void InitRHICapabilitiesForGL()
 #if PLATFORM_ANDROID 
 		if (FOpenGL::UseES30ShadingLanguage())
 		{
+			// NOTE: This should be GL_RGBA8_SNORM but it doesn't work with glTexBuffer -> mapping to Unorm and unpacking in the shader
+			SetupTextureFormat(PF_R8G8B8A8_SNORM,	FOpenGLTextureFormat( GL_RGBA8,					GL_RGBA8,				GL_RGBA,		GL_UNSIGNED_BYTE,					false,	false));
+
 			SetupTextureFormat( PF_R32_UINT,		FOpenGLTextureFormat( GL_R32UI,					GL_R32UI,				GL_RED_INTEGER,	GL_UNSIGNED_INT,					false,	false));
 			SetupTextureFormat( PF_R32_SINT,		FOpenGLTextureFormat( GL_R32I,					GL_R32I,				GL_RED_INTEGER,	GL_INT,								false,	false));
 			SetupTextureFormat( PF_R16_UINT,		FOpenGLTextureFormat( GL_R16UI,					GL_R16UI,				GL_RED_INTEGER,	GL_UNSIGNED_SHORT,					false,	false));
