Fallback to naive frame pacer on frames in which eglGetFrameTimestampsANDROID reports errors.
Revert to naive frame pacer entirely after 100 consecutive failed frames.
Tweaked EGL extension enumeration output slightly to help differentiate.

diff --git a/Engine/Source/Runtime/OpenGLDrv/Private/Android/AndroidEGL.cpp b/Engine/Source/Runtime/OpenGLDrv/Private/Android/AndroidEGL.cpp
--- a/Engine/Source/Runtime/OpenGLDrv/Private/Android/AndroidEGL.cpp
+++ b/Engine/Source/Runtime/OpenGLDrv/Private/Android/AndroidEGL.cpp
@@ -23,6 +23,12 @@ static TAutoConsoleVariable<int32> CVarAllowFrameTimestamps(
 	1,
 	TEXT("True to allow the use use eglGetFrameTimestampsANDROID et al for frame pacing or spew."));
 
+static TAutoConsoleVariable<int32> CVarTimeStampErrorRetryCount(
+	TEXT("a.TimeStampErrorRetryCount"),
+	100,
+	TEXT("Number of consecutive frames eglGetFrameTimestampsANDROID can fail before reverting to the naive frame pacer.\n")
+	TEXT("Retry count is reset after any successful eglGetFrameTimestampsANDROID calls or a change to sync interval."));
+
 #define ENABLE_CONFIG_FILTER 1
 #define ENABLE_EGL_DEBUG 0
 #define ENABLE_VERIFY_EGL 0
@@ -960,13 +966,29 @@ static TAutoConsoleVariable<int32> CVarDisableOpenGLGPUSync(
 	TEXT("When true, android OpenGL will not prevent the GPU from running more than one frame behind. This will allow higher performance on some devices but increase input latency."),
 	ECVF_RenderThreadSafe);
 
+static bool GGetTimeStampsSucceededThisFrame = true;
+static uint32 GGetTimeStampsRetryCount = 0;
+
+static bool CanUseGetFrameTimestamps()
+{
+	return CVarUseGetFrameTimestamps.GetValueOnAnyThread()
+		&& eglGetFrameTimestampsANDROID_p
+		&& eglGetNextFrameIdANDROID_p
+		&& eglPresentationTimeANDROID_p
+		&& (GGetTimeStampsRetryCount < CVarTimeStampErrorRetryCount.GetValueOnAnyThread());
+}
+static bool CanUseGetFrameTimestampsForThisFrame()
+{
+	return CanUseGetFrameTimestamps() && GGetTimeStampsSucceededThisFrame;
+}
+
 bool ShouldUseGPUFencesToLimitLatency()
 {
 	if (ShouldUseChoreographer())
 	{
 		return false; // this method does its own GPU fences as part of the swap.
 	}
-	if (CVarUseGetFrameTimestamps.GetValueOnAnyThread())
+	if (CanUseGetFrameTimestampsForThisFrame())
 	{
 		return true; // this method requires a GPU fence to give steady results
 	}
@@ -976,7 +998,6 @@ bool ShouldUseGPUFencesToLimitLatency()
 
 bool AndroidEGL::SwapBuffers(int32 SyncInterval)
 {
-	//SyncInterval = 3;
 #if !UE_BUILD_SHIPPING
 	if (CVarStallSwap.GetValueOnAnyThread() > 0.0f)
 	{
@@ -1036,6 +1057,8 @@ bool AndroidEGL::SwapBuffers(int32 SyncInterval)
 	bool bPrintMethod = false;
 	if (PImplData->DesiredSyncIntervalRelativeTo60Hz != SyncInterval)
 	{
+		GGetTimeStampsRetryCount = 0;
+
 		bPrintMethod = true;
 		PImplData->DesiredSyncIntervalRelativeTo60Hz = SyncInterval;
 		PImplData->DriverRefreshRate = 60.0f;
@@ -1189,50 +1212,58 @@ bool AndroidEGL::SwapBuffers(int32 SyncInterval)
 	{
 		TheChoreographerFramePacer.StopChoreographer();
 	}
-	if (!bUseChoreographer && eglPresentationTimeANDROID_p && CVarUseGetFrameTimestamps.GetValueOnAnyThread())
+	if (!bUseChoreographer && CanUseGetFrameTimestamps())
 	{
-		if (eglGetFrameTimestampsANDROID_p && eglGetNextFrameIdANDROID_p)
-		{
-			UE_CLOG(bPrintMethod, LogRHI, Display, TEXT("Using eglGetFrameTimestampsANDROID method for frame pacing"));
+		UE_CLOG(bPrintMethod, LogRHI, Display, TEXT("Using eglGetFrameTimestampsANDROID method for frame pacing"));
 
-			//static bool bPrintOnce = true;
-			if (FrameIDs[(int32(NextFrameIDSlot) - 1) % NUM_FRAMES_TO_MONITOR])
-				// not supported   && eglGetFrameTimestampsSupportedANDROID_p && eglGetFrameTimestampsSupportedANDROID_p(PImplData->eglDisplay, PImplData->eglSurface, EGL_FIRST_COMPOSITION_START_TIME_ANDROID))
+		//static bool bPrintOnce = true;
+		if (FrameIDs[(int32(NextFrameIDSlot) - 1) % NUM_FRAMES_TO_MONITOR])
+			// not supported   && eglGetFrameTimestampsSupportedANDROID_p && eglGetFrameTimestampsSupportedANDROID_p(PImplData->eglDisplay, PImplData->eglSurface, EGL_FIRST_COMPOSITION_START_TIME_ANDROID))
+		{
+			//UE_CLOG(bPrintOnce, LogRHI, Log, TEXT("eglGetFrameTimestampsSupportedANDROID retured true for EGL_FIRST_COMPOSITION_START_TIME_ANDROID"));
+			EGLint TimestampList = EGL_FIRST_COMPOSITION_START_TIME_ANDROID;
+			//EGLint TimestampList = EGL_COMPOSITION_LATCH_TIME_ANDROID;
+			//EGLint TimestampList = EGL_LAST_COMPOSITION_START_TIME_ANDROID;
+			//EGLint TimestampList = EGL_DISPLAY_PRESENT_TIME_ANDROID;
+			EGLnsecsANDROID Result = 0;
+			int32 DeltaFrameIndex = 1;
+			for (int32 Index = int32(NextFrameIDSlot) - 1; Index >= int32(NextFrameIDSlot) - NUM_FRAMES_TO_MONITOR && Index >= 0; Index--)
 			{
-				//UE_CLOG(bPrintOnce, LogRHI, Log, TEXT("eglGetFrameTimestampsSupportedANDROID retured true for EGL_FIRST_COMPOSITION_START_TIME_ANDROID"));
-				EGLint TimestampList = EGL_FIRST_COMPOSITION_START_TIME_ANDROID;
-				//EGLint TimestampList = EGL_COMPOSITION_LATCH_TIME_ANDROID;
-				//EGLint TimestampList = EGL_LAST_COMPOSITION_START_TIME_ANDROID;
-				//EGLint TimestampList = EGL_DISPLAY_PRESENT_TIME_ANDROID;
-				EGLnsecsANDROID Result = 0;
-				int32 DeltaFrameIndex = 1;
-				for (int32 Index = int32(NextFrameIDSlot) - 1; Index >= int32(NextFrameIDSlot) - NUM_FRAMES_TO_MONITOR && Index >= 0; Index--)
+				Result = 0;
+				if (FrameIDs[Index % NUM_FRAMES_TO_MONITOR])
 				{
-					Result = 0;
-					if (FrameIDs[Index % NUM_FRAMES_TO_MONITOR])
-					{
-						eglGetFrameTimestampsANDROID_p(PImplData->eglDisplay, PImplData->eglSurface, FrameIDs[Index % NUM_FRAMES_TO_MONITOR], 1, &TimestampList, &Result);
-					}
-					if (Result > 0)
-					{
-						break;
-					}
-					DeltaFrameIndex++;
+					eglGetFrameTimestampsANDROID_p(PImplData->eglDisplay, PImplData->eglSurface, FrameIDs[Index % NUM_FRAMES_TO_MONITOR], 1, &TimestampList, &Result);
 				}
 				if (Result > 0)
 				{
-					EGLnsecsANDROID FudgeFactor = 0; //  8333 * 1000;
-					EGLnsecsANDROID DeltaNanos = EGLnsecsANDROID(PImplData->DesiredSyncIntervalRelativeToDevice) * EGLnsecsANDROID(DeltaFrameIndex) * PImplData->DriverRefreshNanos;
-					EGLnsecsANDROID PresentationTime = Result + DeltaNanos + FudgeFactor;
-					eglPresentationTimeANDROID_p(PImplData->eglDisplay, PImplData->eglSurface, PresentationTime);
+					break;
 				}
+				DeltaFrameIndex++;
+			}
+
+			GGetTimeStampsSucceededThisFrame = Result > 0;
+			if (GGetTimeStampsSucceededThisFrame)
+			{
+				EGLnsecsANDROID FudgeFactor = 0; //  8333 * 1000;
+				EGLnsecsANDROID DeltaNanos = EGLnsecsANDROID(PImplData->DesiredSyncIntervalRelativeToDevice) * EGLnsecsANDROID(DeltaFrameIndex) * PImplData->DriverRefreshNanos;
+				EGLnsecsANDROID PresentationTime = Result + DeltaNanos + FudgeFactor;
+				eglPresentationTimeANDROID_p(PImplData->eglDisplay, PImplData->eglSurface, PresentationTime);
+				GGetTimeStampsRetryCount = 0;
 			}
 			else
 			{
-				//UE_CLOG(bPrintOnce, LogRHI, Log, TEXT("eglGetFrameTimestampsSupportedANDROID doesn't exist or retured false for EGL_FIRST_COMPOSITION_START_TIME_ANDROID, discarding eglGetNextFrameIdANDROID_p and eglGetFrameTimestampsANDROID_p"));
+				GGetTimeStampsRetryCount++;
+				if (GGetTimeStampsRetryCount == CVarTimeStampErrorRetryCount.GetValueOnAnyThread())
+				{
+					UE_LOG(LogRHI, Log, TEXT("eglGetFrameTimestampsANDROID_p failed for %d consecutive frames, reverting to naive frame pacer."), GGetTimeStampsRetryCount);
+				}
 			}
-			//bPrintOnce = false;
 		}
+		else
+		{
+			//UE_CLOG(bPrintOnce, LogRHI, Log, TEXT("eglGetFrameTimestampsSupportedANDROID doesn't exist or retured false for EGL_FIRST_COMPOSITION_START_TIME_ANDROID, discarding eglGetNextFrameIdANDROID_p and eglGetFrameTimestampsANDROID_p"));
+		}
+		//bPrintOnce = false;
 	}
 
 	PImplData->LastTimeEmulatedSync = FPlatformTime::Seconds();
@@ -1241,7 +1272,7 @@ bool AndroidEGL::SwapBuffers(int32 SyncInterval)
 		QUICK_SCOPE_CYCLE_COUNTER(STAT_eglSwapBuffers);
 
 		FrameIDs[(NextFrameIDSlot) % NUM_FRAMES_TO_MONITOR] = 0;
-		if (eglGetNextFrameIdANDROID_p && (CVarUseGetFrameTimestamps.GetValueOnAnyThread() || CVarSpewGetFrameTimestamps.GetValueOnAnyThread()))
+		if (eglGetNextFrameIdANDROID_p && (CanUseGetFrameTimestamps() || CVarSpewGetFrameTimestamps.GetValueOnAnyThread()))
 		{
 			eglGetNextFrameIdANDROID_p(PImplData->eglDisplay, PImplData->eglSurface, &FrameIDs[(NextFrameIDSlot) % NUM_FRAMES_TO_MONITOR]);
 		}
diff --git a/Engine/Source/Runtime/OpenGLDrv/Private/Android/AndroidOpenGL.cpp b/Engine/Source/Runtime/OpenGLDrv/Private/Android/AndroidOpenGL.cpp
--- a/Engine/Source/Runtime/OpenGLDrv/Private/Android/AndroidOpenGL.cpp
+++ b/Engine/Source/Runtime/OpenGLDrv/Private/Android/AndroidOpenGL.cpp
@@ -349,14 +349,16 @@ void FPlatformOpenGLDevice::LoadEXT()
 	eglGetCompositorTimingSupportedANDROID_p = (PFNeglQueryTimestampSupportedANDROID)((void*)eglGetProcAddress("eglGetCompositorTimingSupportedANDROID"));
 	eglGetFrameTimestampsSupportedANDROID_p = (PFNeglQueryTimestampSupportedANDROID)((void*)eglGetProcAddress("eglGetFrameTimestampsSupportedANDROID"));
 
-
-	UE_LOG(LogRHI, Log, TEXT("Extension %s %s"), TEXT("eglPresentationTimeANDROID"), eglPresentationTimeANDROID_p  ? TEXT("Present") : TEXT("Not Present"));
-	UE_LOG(LogRHI, Log, TEXT("Extension %s %s"), TEXT("eglGetNextFrameIdANDROID"), eglGetNextFrameIdANDROID_p ? TEXT("Present") : TEXT("Not Present"));
-	UE_LOG(LogRHI, Log, TEXT("Extension %s %s"), TEXT("eglGetCompositorTimingANDROID"), eglGetCompositorTimingANDROID_p  ? TEXT("Present") : TEXT("Not Present"));
-	UE_LOG(LogRHI, Log, TEXT("Extension %s %s"), TEXT("eglGetFrameTimestampsANDROID"), eglGetFrameTimestampsANDROID_p  ? TEXT("Present") : TEXT("Not Present"));
-	UE_LOG(LogRHI, Log, TEXT("Extension %s %s"), TEXT("eglQueryTimestampSupportedANDROID"), eglQueryTimestampSupportedANDROID_p ? TEXT("Present") : TEXT("Not Present"));
-	UE_LOG(LogRHI, Log, TEXT("Extension %s %s"), TEXT("eglGetCompositorTimingSupportedANDROID"), eglGetCompositorTimingSupportedANDROID_p ? TEXT("Present") : TEXT("Not Present"));
-	UE_LOG(LogRHI, Log, TEXT("Extension %s %s"), TEXT("eglGetFrameTimestampsSupportedANDROID"), eglGetFrameTimestampsSupportedANDROID_p ? TEXT("Present") : TEXT("Not Present"));
+	const TCHAR* NotAvailable = TEXT("NOT Available");
+	const TCHAR* Present = TEXT("Present");
+
+	UE_LOG(LogRHI, Log, TEXT("Extension %s %s"), TEXT("eglPresentationTimeANDROID"), eglPresentationTimeANDROID_p  ? Present : NotAvailable);
+	UE_LOG(LogRHI, Log, TEXT("Extension %s %s"), TEXT("eglGetNextFrameIdANDROID"), eglGetNextFrameIdANDROID_p ? Present : NotAvailable);
+	UE_LOG(LogRHI, Log, TEXT("Extension %s %s"), TEXT("eglGetCompositorTimingANDROID"), eglGetCompositorTimingANDROID_p  ? Present : NotAvailable);
+	UE_LOG(LogRHI, Log, TEXT("Extension %s %s"), TEXT("eglGetFrameTimestampsANDROID"), eglGetFrameTimestampsANDROID_p  ? Present : NotAvailable);
+	UE_LOG(LogRHI, Log, TEXT("Extension %s %s"), TEXT("eglQueryTimestampSupportedANDROID"), eglQueryTimestampSupportedANDROID_p ? Present : NotAvailable);
+	UE_LOG(LogRHI, Log, TEXT("Extension %s %s"), TEXT("eglGetCompositorTimingSupportedANDROID"), eglGetCompositorTimingSupportedANDROID_p ? Present : NotAvailable);
+	UE_LOG(LogRHI, Log, TEXT("Extension %s %s"), TEXT("eglGetFrameTimestampsSupportedANDROID"), eglGetFrameTimestampsSupportedANDROID_p ? Present : NotAvailable);
 
 	glDebugMessageControlKHR = (PFNGLDEBUGMESSAGECONTROLKHRPROC)((void*)eglGetProcAddress("glDebugMessageControlKHR"));
 
