When building OGL binary file cache use the existing cache file to avoid rebuilding matching programs.
r.OpenGL.UseExistingBinaryFileCache=0 to disable.
+ Additional improvements to binary program cache error detection.

diff --git a/Engine/Source/Runtime/OpenGLDrv/Private/OpenGLShaders.cpp b/Engine/Source/Runtime/OpenGLDrv/Private/OpenGLShaders.cpp
--- a/Engine/Source/Runtime/OpenGLDrv/Private/OpenGLShaders.cpp
+++ b/Engine/Source/Runtime/OpenGLDrv/Private/OpenGLShaders.cpp
@@ -78,6 +78,14 @@ static TAutoConsoleVariable<int32> CVarIgnoreLinkFailure(
 	TEXT("1: Ignore link failures. this may allow a program to continue but could lead to undefined rendering behaviour."),
 	ECVF_RenderThreadSafe);
 
+static TAutoConsoleVariable<int32> CVarUseExistingBinaryFileCache(
+	TEXT("r.OpenGL.UseExistingBinaryFileCache"),
+	1,
+	TEXT("When generating a new binary cache (such as when Shader Pipeline Cache Version Guid changes) use the existing binary file cache to speed up generation of the new cache.\n")
+	TEXT("0: Always rebuild binary file cache when Pipeline Cache Version Guid changes.\n")
+	TEXT("1: When Pipeline Cache Version Guid changes re-use programs from the existing binary cache where possible (default)."),
+	ECVF_RenderThreadSafe);
+
 #if PLATFORM_ANDROID
 bool GOpenGLShaderHackLastCompileSuccess = false;
 #endif
@@ -96,7 +104,7 @@ static FAutoConsoleCommand ConsoleCommandPrintProgramStats(
 								TEXT("r.OpenGL.PrintProgramStats"),
 								TEXT("Print to log current program binary stats"),
 								FConsoleCommandDelegate::CreateStatic(PrintProgramStats)
-								);
+);
 
 static void SetNewProgramStats(GLuint Program)
 {
@@ -211,7 +219,7 @@ static bool VerifyLinkedProgram(GLuint Program, VerifyProgramPipelineFailurePoli
 			CompileLog = (ANSICHAR *)FMemory::Malloc(LogLength);
 			glGetProgramInfoLog(Program, LogLength, NULL, CompileLog);
 		}
-		LogError(FString::Printf(TEXT("Failed to link program. Current total programs: %d binary bytes: %d\n  log:\n%s"),
+		LogError(FString::Printf(TEXT("Failed to link program. Current total programs: %d program binary bytes: %d\n  log:\n%s"),
 			GNumPrograms,
 			GCurrentDriverProgramBinaryAllocation,
 			ANSI_TO_TCHAR(CompileLog)), FailurePolicy);
@@ -546,6 +554,13 @@ namespace
 	}
 }
 
+
+FOpenGLProgramBinaryCache::FPreviousGLProgramBinaryCacheInfo::FPreviousGLProgramBinaryCacheInfo() : NumberOfOldEntriesReused(0) {}
+FOpenGLProgramBinaryCache::FPreviousGLProgramBinaryCacheInfo::FPreviousGLProgramBinaryCacheInfo(FOpenGLProgramBinaryCache::FPreviousGLProgramBinaryCacheInfo&&) = default;
+FOpenGLProgramBinaryCache::FPreviousGLProgramBinaryCacheInfo& FOpenGLProgramBinaryCache::FPreviousGLProgramBinaryCacheInfo::operator = (FOpenGLProgramBinaryCache::FPreviousGLProgramBinaryCacheInfo&&) = default;
+FOpenGLProgramBinaryCache::FPreviousGLProgramBinaryCacheInfo::~FPreviousGLProgramBinaryCacheInfo() = default;
+
+
 inline uint32 GetTypeHash(FAnsiCharArray const& CharArray)
 {
 	return FCrc::MemCrc32(CharArray.GetData(), CharArray.Num() * sizeof(ANSICHAR));
@@ -2196,7 +2211,7 @@ class FGLProgramCacheLRU
 
 	void EvictFromLRU(FOpenGLLinkedProgram* LinkedProgram)
 	{
-		SCOPE_CYCLE_COUNTER(STAT_OpenGLShaderLRUEvictTime);
+	 	SCOPE_CYCLE_COUNTER(STAT_OpenGLShaderLRUEvictTime);
 		LinkedProgram->LRUInfo.LRUNode = FSetElementId();
 
 		if (LinkedProgram->LRUInfo.EvictBucket >= 0)
@@ -2447,7 +2462,7 @@ public:
 		if( ProgramCacheLRU.LRU.Num() )
 		{
 			ProgramCacheLRU.EvictMostRecentFromLRU();
-		}
+	 	}
 	}
 
 	void EvictProgram(const FOpenGLProgramKey& ProgramKey)
@@ -3075,12 +3090,13 @@ static FOpenGLLinkedProgram* LinkProgram( const FOpenGLLinkedProgramConfiguratio
 	check( (Config.Shaders[CrossCompiler::SHADER_STAGE_VERTEX].Resource == 0) != (Config.Shaders[CrossCompiler::SHADER_STAGE_COMPUTE].Resource == 0));
 	check( (Config.Shaders[CrossCompiler::SHADER_STAGE_PIXEL].Resource == 0) != (Config.Shaders[CrossCompiler::SHADER_STAGE_COMPUTE].Resource == 0));
 
+	TArray<uint8> CachedProgramBinary;
 	GLuint Program = 0;
 	bool bShouldLinkProgram = true;
 	if (FOpenGLProgramBinaryCache::IsEnabled())
 	{
 		// Try to create program from a saved binary
-		bShouldLinkProgram = !FOpenGLProgramBinaryCache::UseCachedProgram(Program, Config.ProgramKey);
+		bShouldLinkProgram = !FOpenGLProgramBinaryCache::UseCachedProgram(Program, Config.ProgramKey, CachedProgramBinary);
 		if (bShouldLinkProgram)
 		{
 			// In case there is no saved binary in the cache, compile required shaders we have deferred before
@@ -3157,7 +3173,8 @@ static FOpenGLLinkedProgram* LinkProgram( const FOpenGLLinkedProgramConfiguratio
 
 			if (FOpenGLProgramBinaryCache::IsEnabled())
 			{
-				FOpenGLProgramBinaryCache::CacheProgram(Program, Config.ProgramKey);
+				check(CachedProgramBinary.Num() == 0);
+				FOpenGLProgramBinaryCache::CacheProgram(Program, Config.ProgramKey, CachedProgramBinary);
 			}
 		}
 	}
@@ -3177,6 +3194,12 @@ static FOpenGLLinkedProgram* LinkProgram( const FOpenGLLinkedProgramConfiguratio
 	bool bUsingTessellation = Config.Shaders[CrossCompiler::SHADER_STAGE_HULL].Resource && Config.Shaders[CrossCompiler::SHADER_STAGE_DOMAIN].Resource;
 	FOpenGLLinkedProgram* LinkedProgram = new FOpenGLLinkedProgram(Config, Program, bUsingTessellation);
 
+	if (GetOpenGLProgramsCache().IsUsingLRU() && CVarLRUKeepProgramBinaryResident.GetValueOnAnyThread() && CachedProgramBinary.Num())
+	{
+ 		// Store the binary data in LRUInfo, this avoids requesting a program binary from the driver when this program is evicted. 
+ 		INC_MEMORY_STAT_BY(STAT_OpenGLShaderLRUProgramMemory, CachedProgramBinary.Num());
+ 		LinkedProgram->LRUInfo.CachedProgramBinary = MoveTemp(CachedProgramBinary);
+	}
  	ConfigureStageStates(LinkedProgram);
 
 #if ENABLE_UNIFORM_BUFFER_LAYOUT_VERIFICATION	
@@ -4444,7 +4467,7 @@ bool FOpenGLProgramBinaryCache::IsBuildingCache()
 {
 	if(CachePtr != nullptr)
 	{
-		return CachePtr->BinaryFileState == EBinaryFileState::BuildingCacheFile;
+		return CachePtr->IsBuildingCache_internal();
 	}
 	return false;
 }
@@ -4513,7 +4536,7 @@ void FOpenGLProgramBinaryCache::OnShaderPipelineCacheOpened(FString const& Name
 {
 	UE_LOG(LogRHI, Log, TEXT("Scanning Binary program cache, using Shader Pipeline Cache version %s"), *VersionGuid.ToString());
 	ScanProgramCacheFile(VersionGuid);
-	if(BinaryFileState == EBinaryFileState::BuildingCacheFile)
+	if(IsBuildingCache_internal())
 	{
 		ShaderCachePrecompileContext.SetPrecompilationIsSlowTask();
 	}
@@ -4529,9 +4552,9 @@ void FOpenGLProgramBinaryCache::OnShaderPipelineCachePrecompilationComplete(uint
 	OnShaderPipelineCacheOpenedDelegate.Reset();
 	OnShaderPipelineCachePrecompilationCompleteDelegate.Reset();
 
-	check(BinaryFileState == EBinaryFileState::BuildingCacheFile || BinaryFileState == EBinaryFileState::ValidCacheFile);
+	check(IsBuildingCache_internal() || BinaryFileState == EBinaryFileState::ValidCacheFile);
 
-	if (BinaryFileState == EBinaryFileState::BuildingCacheFile)
+	if (IsBuildingCache_internal())
 	{
 		CloseWriteHandle();
 
@@ -4613,12 +4636,13 @@ void FOpenGLProgramBinaryCache::ScanProgramCacheFile(const FGuid& ShaderPipeline
 	check(BinaryFileState == EBinaryFileState::Uninitialized);
 
 	bool bBinaryFileIsValid = false;
+	bool bBinaryFileIsValidAndGuidMatch = false;
 
 	// Try to move the file to a temporary filename before the scan, so we won't try to read it again if it's corrupted
 	PlatformFile.DeleteFile(*ProgramCacheFilenameTemp);
 	PlatformFile.MoveFile(*ProgramCacheFilenameTemp, *ProgramCacheFilename);
 
-	FArchive* FileReader = IFileManager::Get().CreateFileReader(*ProgramCacheFilenameTemp);
+	TUniquePtr<FArchive> FileReader(IFileManager::Get().CreateFileReader(*ProgramCacheFilenameTemp));
 	if (FileReader)
 	{
 		UE_LOG(LogRHI, Log, TEXT("OnShaderScanProgramCacheFile : Opened %s"), *ProgramCacheFilenameTemp);
@@ -4633,11 +4657,21 @@ void FOpenGLProgramBinaryCache::ScanProgramCacheFile(const FGuid& ShaderPipeline
 			Ar << bCacheUsesCompressedBinaries;
 
 			const bool bUseCompressedProgramBinaries = CVarStoreCompressedBinaries.GetValueOnAnyThread() != 0;
-			bBinaryFileIsValid = (!ShaderPipelineCacheVersionGuid.IsValid() || ShaderPipelineCacheVersionGuid == BinaryCacheGuid) && (bUseCompressedProgramBinaries == bCacheUsesCompressedBinaries);
+			bBinaryFileIsValid = (bUseCompressedProgramBinaries == bCacheUsesCompressedBinaries);
+			bBinaryFileIsValidAndGuidMatch = bBinaryFileIsValid && (!ShaderPipelineCacheVersionGuid.IsValid() || ShaderPipelineCacheVersionGuid == BinaryCacheGuid);
+			
+			if (CVarUseExistingBinaryFileCache.GetValueOnAnyThread() == 0 && bBinaryFileIsValidAndGuidMatch == false)
+			{
+				// If we dont want to use the existing binary cache and the guids have changed then rebuild the binary file.
+				bBinaryFileIsValid = false;
+			}
 		}
-
+		
 		if (bBinaryFileIsValid)
 		{
+			const uint32 ProgramBinaryStart = Ar.Tell();
+
+			// Search the file for the end record.
 			bool bFoundEndRecord = false;
 			int32 ProgramIndex = 0;
 			while (!Ar.AtEnd())
@@ -4648,42 +4682,8 @@ void FOpenGLProgramBinaryCache::ScanProgramCacheFile(const FGuid& ShaderPipeline
 				uint32 ProgramBinarySize = 0;
 				Ar << ProgramKey;
 				Ar << ProgramBinarySize;
-
-				if (ProgramBinarySize > 0)
-				{
-					FGLProgramBinaryFileCacheEntry* NewEntry = new FGLProgramBinaryFileCacheEntry();
-					NewEntry->FileInfo.ShaderHasheSet = ProgramKey;
-					NewEntry->ProgramIndex = ProgramIndex++;
-
-					uint32 ProgramBinaryOffset = Ar.Tell();
-					NewEntry->FileInfo.ProgramSize = ProgramBinarySize;
-					NewEntry->FileInfo.ProgramOffset = ProgramBinaryOffset;
-
-					ProgramEntryContainer.Emplace(TUniquePtr<FGLProgramBinaryFileCacheEntry>(NewEntry));
-
-					// check to see if any of the shaders are already loaded and so we should serialize the binary
-					bool bAllShadersLoaded = true;
-					for (int32 i = 0; i < CrossCompiler::NUM_NON_COMPUTE_SHADER_STAGES && bAllShadersLoaded; i++)
-					{
-						bAllShadersLoaded = ProgramKey.ShaderHashes[i] == FSHAHash() || ShaderIsLoaded(ProgramKey.ShaderHashes[i]);
-					}
-					if (bAllShadersLoaded)
-					{
-						FPlatformMisc::LowLevelOutputDebugStringf(TEXT("*** All shaders for program %s already loaded\n"), *ProgramKey.ToString());
-						NewEntry->ProgramBinaryData.AddUninitialized(ProgramBinarySize);
-						Ar.Serialize(NewEntry->ProgramBinaryData.GetData(), ProgramBinarySize);
-						NewEntry->GLProgramState = FGLProgramBinaryFileCacheEntry::EGLProgramState::ProgramLoaded;
-						CompleteLoadedGLProgramRequest_internal(NewEntry);
-					}
-					else
-					{
-						NewEntry->GLProgramState = FGLProgramBinaryFileCacheEntry::EGLProgramState::ProgramStored;
-					}
-					Ar.Seek(ProgramBinaryOffset + ProgramBinarySize);
-
-					AddProgramFileEntryToMap(NewEntry);
-				}
-				else
+				uint32 ProgramBinaryOffset = Ar.Tell();
+				if (ProgramBinarySize == 0) 
 				{
 					if (ProgramKey == FOpenGLProgramKey())
 					{
@@ -4695,33 +4695,93 @@ void FOpenGLProgramBinaryCache::ScanProgramCacheFile(const FGuid& ShaderPipeline
 						UE_LOG(LogRHI, Warning, TEXT("FOpenGLProgramBinaryCache::ScanProgramCacheFile : encountered 0 sized program during binary program cache scan"));
 					}
 				}
+				Ar.Seek(ProgramBinaryOffset + ProgramBinarySize);
 			}
 
-			if (!bFoundEndRecord)
+			if(bFoundEndRecord)
 			{
-				// failed to find sentinel record, this file was not finalized.
-				bBinaryFileIsValid = false;
-				UE_LOG(LogRHI, Warning, TEXT("ScanProgramCacheFile - incomplete binary cache file encountered. Rebuilding binary program cache."));
-			}
+				Ar.Seek(ProgramBinaryStart);
+				while (!Ar.AtEnd())
+				{
+					FOpenGLProgramKey ProgramKey;
+					uint32 ProgramBinarySize = 0;
+					Ar << ProgramKey;
+					Ar << ProgramBinarySize;
 
-			UE_LOG(LogRHI, Log, TEXT("Program Binary cache: Found %d cached programs, end record found: %d"), ProgramIndex, (uint32)bFoundEndRecord);
-			FileReader->Close();
-			delete FileReader;
+					if (ProgramBinarySize > 0)
+					{
+						FGLProgramBinaryFileCacheEntry* NewEntry = new FGLProgramBinaryFileCacheEntry();
+						NewEntry->FileInfo.ShaderHasheSet = ProgramKey;
+						NewEntry->ProgramIndex = ProgramIndex++;
+
+						uint32 ProgramBinaryOffset = Ar.Tell();
+						NewEntry->FileInfo.ProgramSize = ProgramBinarySize;
+						NewEntry->FileInfo.ProgramOffset = ProgramBinaryOffset;
+
+						if (bBinaryFileIsValidAndGuidMatch)
+						{
+							ProgramEntryContainer.Emplace(TUniquePtr<FGLProgramBinaryFileCacheEntry>(NewEntry));
+
+							// check to see if any of the shaders are already loaded and so we should serialize the binary
+							bool bAllShadersLoaded = true;
+							for (int32 i = 0; i < CrossCompiler::NUM_NON_COMPUTE_SHADER_STAGES && bAllShadersLoaded; i++)
+							{
+								bAllShadersLoaded = ProgramKey.ShaderHashes[i] == FSHAHash() || ShaderIsLoaded(ProgramKey.ShaderHashes[i]);
+							}
+							if (bAllShadersLoaded)
+							{
+								FPlatformMisc::LowLevelOutputDebugStringf(TEXT("*** All shaders for program %s already loaded\n"), *ProgramKey.ToString());
+								NewEntry->ProgramBinaryData.AddUninitialized(ProgramBinarySize);
+								Ar.Serialize(NewEntry->ProgramBinaryData.GetData(), ProgramBinarySize);
+								NewEntry->GLProgramState = FGLProgramBinaryFileCacheEntry::EGLProgramState::ProgramLoaded;
+								CompleteLoadedGLProgramRequest_internal(NewEntry);
+							}
+							else
+							{
+								NewEntry->GLProgramState = FGLProgramBinaryFileCacheEntry::EGLProgramState::ProgramStored;
+							}
+							AddProgramFileEntryToMap(NewEntry);
+						}
+						else
+						{
+							check(!PreviousBinaryCacheInfo.ProgramToOldBinaryCacheMap.Contains(ProgramKey));
+							PreviousBinaryCacheInfo.ProgramToOldBinaryCacheMap.Emplace(ProgramKey, TUniquePtr<FGLProgramBinaryFileCacheEntry>(NewEntry));
+						}
+						Ar.Seek(ProgramBinaryOffset + ProgramBinarySize);
+					}
+				}
 
-			if(bBinaryFileIsValid)
+				if (bBinaryFileIsValidAndGuidMatch)
+				{
+					UE_LOG(LogRHI, Log, TEXT("Program Binary cache: Found %d cached programs, end record found: %d"), ProgramIndex, (uint32)bFoundEndRecord);
+					FileReader->Close();
+					// Rename the file back after a successful scan.
+					PlatformFile.MoveFile(*ProgramCacheFilename, *ProgramCacheFilenameTemp);
+				}
+				else
+				{
+					UE_LOG(LogRHI, Log, TEXT("Program Binary cache: ShaderPipelineCache changed, regenerating for new pipeline cache. Existing cache contains %d programs, using it to populate."), PreviousBinaryCacheInfo.ProgramToOldBinaryCacheMap.Num());
+					// Not closing the scan source file, we're using it to move shaders from the old cache.
+					PreviousBinaryCacheInfo.OldCacheArchive = MoveTemp(FileReader);
+					PreviousBinaryCacheInfo.OldCacheFilename = ProgramCacheFilenameTemp;
+				}
+			}
+			else
 			{
-				// Rename the file back after a successful scan.
-				PlatformFile.MoveFile(*ProgramCacheFilename, *ProgramCacheFilenameTemp);
+				// failed to find sentinel record, the file was not finalized.
+				UE_LOG(LogRHI, Warning, TEXT("ScanProgramCacheFile - incomplete binary cache file encountered. Rebuilding binary program cache."));
+				FileReader->Close();
+				bBinaryFileIsValid = false;
+				bBinaryFileIsValidAndGuidMatch = false;
 			}
 		}
-		else
+		
+		if(!bBinaryFileIsValid)
 		{
 			UE_LOG(LogRHI, Log, TEXT("OnShaderScanProgramCacheFile : binary file version invalid"));
-			FileReader->Close();
-			delete FileReader;
 		}
 
-		if (bBinaryFileIsValid)
+		if (bBinaryFileIsValidAndGuidMatch)
 		{
 			OpenAsyncReadHandle();
 			BinaryFileState = EBinaryFileState::ValidCacheFile;
@@ -4744,10 +4804,13 @@ void FOpenGLProgramBinaryCache::ScanProgramCacheFile(const FGuid& ShaderPipeline
 			UE_LOG(LogRHI, Warning, TEXT("Failed to create directory for a program binary cache. Cache will be disabled: %s"), *CachePath);
 			return;
 		}
+	}
 
-		if (OpenWriteHandle(true))
+	if(!bBinaryFileIsValid || !bBinaryFileIsValidAndGuidMatch)
+	{
+		if (OpenWriteHandle())
 		{
-			BinaryFileState = EBinaryFileState::BuildingCacheFile;
+			BinaryFileState = bBinaryFileIsValid && !bBinaryFileIsValidAndGuidMatch ? EBinaryFileState::BuildingCacheFileWithMove : EBinaryFileState::BuildingCacheFile;
 
 			// save header
 			FArchive& Ar = *BinaryCacheWriteFileHandle;
@@ -4758,6 +4821,13 @@ void FOpenGLProgramBinaryCache::ScanProgramCacheFile(const FGuid& ShaderPipeline
 			bool bWritingCompressedBinaries = (CVarStoreCompressedBinaries.GetValueOnAnyThread() != 0);
 			Ar << bWritingCompressedBinaries;
 		}
+		else
+		{
+			// Binary cache file cannot be used, failed to open output file.
+			BinaryFileState = EBinaryFileState::Uninitialized;
+			RHIGetPanicDelegate().ExecuteIfBound(FName("FailedBinaryProgramArchiveOpen"));
+			UE_LOG(LogRHI, Fatal, TEXT("ScanProgramCacheFile - Failed to open binary cache."));
+		}
 	}
 }
 
@@ -4768,7 +4838,7 @@ void FOpenGLProgramBinaryCache::AddProgramFileEntryToMap(FGLProgramBinaryFileCac
 	check(!ProgramToBinaryMap.Contains(ProgramKey));
 	ProgramToBinaryMap.Add(ProgramKey, NewEntry);
 
-	UE_LOG(LogRHI, Log, TEXT("AddProgramFileEntryToMap : Adding program: %s"), *ProgramKey.ToString());
+	UE_LOG(LogRHI, Verbose, TEXT("AddProgramFileEntryToMap : Adding program: %s"), *ProgramKey.ToString());
 
 	for (int i = 0; i < CrossCompiler::NUM_NON_COMPUTE_SHADER_STAGES; ++i)
 	{
@@ -4788,7 +4858,7 @@ void FOpenGLProgramBinaryCache::AddProgramFileEntryToMap(FGLProgramBinaryFileCac
 	}
 }
 
-bool FOpenGLProgramBinaryCache::OpenWriteHandle(bool bTruncate)
+bool FOpenGLProgramBinaryCache::OpenWriteHandle()
 {
 	check(BinaryCacheWriteFileHandle == nullptr);
 	check(BinaryCacheAsyncReadFileHandle == nullptr);
@@ -4797,20 +4867,28 @@ bool FOpenGLProgramBinaryCache::OpenWriteHandle(bool bTruncate)
 	FString ProgramCacheFilename = GetProgramBinaryCacheFilePath();
 	FString ProgramCacheFilenameWrite = ProgramCacheFilename + TEXT(".write");
 
-	// if not truncating, move the live file to the temp filename
-	if (!bTruncate)
-	{
-		IPlatformFile& PlatformFile = FPlatformFileManager::Get().GetPlatformFile();
-		PlatformFile.DeleteFile(*ProgramCacheFilenameWrite);
-		PlatformFile.MoveFile(*ProgramCacheFilenameWrite, *ProgramCacheFilename);
-	}
+	BinaryCacheWriteFileHandle = IFileManager::Get().CreateFileWriter(*ProgramCacheFilenameWrite, EFileWrite::FILEWRITE_None);
+
+	UE_CLOG(BinaryCacheWriteFileHandle == nullptr, LogRHI, Warning, TEXT("Failed to open OGL binary cache output file."));
 
-	BinaryCacheWriteFileHandle = IFileManager::Get().CreateFileWriter(*ProgramCacheFilenameWrite, bTruncate ? EFileWrite::FILEWRITE_None : EFileWrite::FILEWRITE_Append);
 	return BinaryCacheWriteFileHandle != nullptr;
 }
 
 void FOpenGLProgramBinaryCache::CloseWriteHandle()
 {
+	if(BinaryFileState == EBinaryFileState::BuildingCacheFileWithMove)
+	{
+		UE_LOG(LogRHI, Log, TEXT("FOpenGLProgramBinaryCache: Deleting previous binary program cache (%s), reused %d programs from a total of %d."), *PreviousBinaryCacheInfo.OldCacheFilename, PreviousBinaryCacheInfo.NumberOfOldEntriesReused, ProgramToBinaryMap.Num());
+
+		// clean up references to old cache.
+		PreviousBinaryCacheInfo.OldCacheArchive->Close();
+		PreviousBinaryCacheInfo.OldCacheArchive = nullptr;
+		IPlatformFile& PlatformFile = FPlatformFileManager::Get().GetPlatformFile();
+		PlatformFile.DeleteFile(*PreviousBinaryCacheInfo.OldCacheFilename);
+		PreviousBinaryCacheInfo.OldCacheFilename.Empty();
+		PreviousBinaryCacheInfo.ProgramToOldBinaryCacheMap.Empty();
+	}
+
 	check(BinaryCacheWriteFileHandle != nullptr);
 
 	AppendProgramBinaryFileEofEntry(*BinaryCacheWriteFileHandle);
@@ -4865,20 +4943,20 @@ void FOpenGLProgramBinaryCache::CloseAsyncReadHandle()
 }*/
 
 // Called when a new program has been created by OGL RHI, creates the binary cache if it's invalid and then appends the new program details to the file and runtime containers.
-void FOpenGLProgramBinaryCache::AppendProgramToBinaryCache(const FOpenGLProgramKey& ProgramKey, GLuint Program)
+void FOpenGLProgramBinaryCache::AppendGLProgramToBinaryCache(const FOpenGLProgramKey& ProgramKey, GLuint Program, TArray<uint8>& CachedProgramBinaryOUT)
 {
-	if (BinaryFileState != EBinaryFileState::BuildingCacheFile)
+	if (IsBuildingCache_internal() == false)
 	{
 		return;
 	}
 
 	FScopeLock Lock(&GProgramBinaryCacheCS);
 
-	AddUniqueProgramToBinaryCache(BinaryCacheWriteFileHandle, ProgramKey, Program);
+	AddUniqueGLProgramToBinaryCache(BinaryCacheWriteFileHandle, ProgramKey, Program, CachedProgramBinaryOUT);
 }
 
 // Add the program to the binary cache if it does not already exist.
-void FOpenGLProgramBinaryCache::AddUniqueProgramToBinaryCache(FArchive* FileWriter, const FOpenGLProgramKey& ProgramKey, GLuint Program)
+void FOpenGLProgramBinaryCache::AddUniqueGLProgramToBinaryCache(FArchive* FileWriter, const FOpenGLProgramKey& ProgramKey, GLuint Program, TArray<uint8>& CachedProgramBinaryOUT)
 {
 	// Add to runtime and disk.
 	const FOpenGLProgramKey& ProgramHash = ProgramKey;
@@ -4886,60 +4964,57 @@ void FOpenGLProgramBinaryCache::AddUniqueProgramToBinaryCache(FArchive* FileWrit
 	// Check we dont already have this: Something could be in the cache but still reach this point if OnSharedShaderCodeRequest(s) have not occurred.
 	if (!ProgramToBinaryMap.Contains(ProgramHash))
 	{
-		uint32 ProgramBinaryOffset, ProgramBinarySize;
+		uint32 ProgramBinaryOffset = 0, ProgramBinarySize = 0;
 
-		if(AppendProgramBinaryFile(*FileWriter, ProgramKey, Program, ProgramBinaryOffset, ProgramBinarySize))
+		FOpenGLProgramKey SerializedProgramKey = ProgramKey;
+		if (ensure(GetProgramBinaryFromGLProgram(Program, CachedProgramBinaryOUT)))
 		{
-			FGLProgramBinaryFileCacheEntry* NewIndexEntry = new FGLProgramBinaryFileCacheEntry();
-			ProgramEntryContainer.Emplace(TUniquePtr<FGLProgramBinaryFileCacheEntry>(NewIndexEntry));
-
-
-			// Store the program file descriptor in the runtime program/shader container:
-			NewIndexEntry->GLProgramState = FGLProgramBinaryFileCacheEntry::EGLProgramState::ProgramComplete;
-			NewIndexEntry->FileInfo.ProgramOffset = ProgramBinaryOffset;
-			NewIndexEntry->FileInfo.ProgramSize = ProgramBinarySize;
-			NewIndexEntry->ProgramIndex = ProgramToBinaryMap.Num();
-			NewIndexEntry->FileInfo.ShaderHasheSet = ProgramKey;
-			AddProgramFileEntryToMap(NewIndexEntry);
+			AddProgramBinaryDataToBinaryCache(*FileWriter, CachedProgramBinaryOUT, ProgramKey);
+		}
+		else
+		{
+			// we've encountered a problem with this program and there's nothing to write.
+			// This likely means the device will never be able to use this program.
+			// Panic!
+			RHIGetPanicDelegate().ExecuteIfBound(FName("FailedBinaryProgramWrite"));
+			UE_LOG(LogRHI, Fatal, TEXT("AppendProgramBinaryFile Binary program returned 0 bytes!"));
+			// Panic!
 		}
 	}
 }
 
-// Serialize out the program binary data
-bool FOpenGLProgramBinaryCache::AppendProgramBinaryFile(FArchive& Ar, const FOpenGLProgramKey& ProgramKey, GLuint Program, uint32& ProgramBinaryOffsetOUT, uint32& ProgramBinarySizeOUT)
+// Serialize out the program binary data and add to runtime structures.
+void FOpenGLProgramBinaryCache::AddProgramBinaryDataToBinaryCache(FArchive& Ar, TArray<uint8>& BinaryProgramData, const FOpenGLProgramKey& ProgramKey)
 {
- 	FOpenGLProgramKey SerializedProgramKey = ProgramKey;
-	TArray<uint8> BinaryProgramData;
-	if (ensure(GetProgramBinaryFromGLProgram(Program, BinaryProgramData)))
+	// Serialize to output file:
+	FOpenGLProgramKey SerializedProgramKey = ProgramKey;
+	uint32 ProgramBinarySize = (uint32)BinaryProgramData.Num();
+	Ar << SerializedProgramKey;
+	uint32 ProgramBinaryOffset = Ar.Tell();
+	Ar << ProgramBinarySize;
+	Ar.Serialize(BinaryProgramData.GetData(), ProgramBinarySize);
+	if(CVarStoreCompressedBinaries.GetValueOnAnyThread())
 	{
-		ProgramBinarySizeOUT = (uint32)BinaryProgramData.Num();
-		Ar << SerializedProgramKey;
-		ProgramBinaryOffsetOUT = Ar.Tell();
-		Ar << ProgramBinarySizeOUT;
-		Ar.Serialize(BinaryProgramData.GetData(), ProgramBinarySizeOUT);
+		static uint32 TotalUncompressed = 0;
+		static uint32 TotalCompressed = 0;
 
-		if(CVarStoreCompressedBinaries.GetValueOnAnyThread())
-		{
-			static uint32 TotalUncompressed = 0;
-			static uint32 TotalCompressed = 0;
-
-			FCompressedProgramBinaryHeader* Header = (FCompressedProgramBinaryHeader*)BinaryProgramData.GetData();
-			TotalUncompressed += Header->UncompressedSize;
-			TotalCompressed += BinaryProgramData.Num();
+		FCompressedProgramBinaryHeader* Header = (FCompressedProgramBinaryHeader*)BinaryProgramData.GetData();
+		TotalUncompressed += Header->UncompressedSize;
+		TotalCompressed += BinaryProgramData.Num();
 
-			UE_LOG(LogRHI, Verbose, TEXT("AppendProgramBinaryFile: total Uncompressed: %d, total Compressed %d, Total saved so far: %d"), TotalUncompressed, TotalCompressed, TotalUncompressed - TotalCompressed);
-		}
+		UE_LOG(LogRHI, Verbose, TEXT("AppendProgramBinaryFile: total Uncompressed: %d, total Compressed %d, Total saved so far: %d"), TotalUncompressed, TotalCompressed, TotalUncompressed - TotalCompressed);
 	}
-	else
-	{
-		// we've encountered a problem with this program and there's nothing to write.
-		// This likely means the device will never be able to use this program.
-		// Panic!
-		RHIGetPanicDelegate().ExecuteIfBound(FName("FailedBinaryProgramWrite"));
-		UE_LOG(LogRHI, Fatal, TEXT("AppendProgramBinaryFile Binary program returned 0 bytes!"));
-		// Panic!
-	}
-	return (ProgramBinarySizeOUT > 0);
+
+	FGLProgramBinaryFileCacheEntry* NewIndexEntry = new FGLProgramBinaryFileCacheEntry();
+	ProgramEntryContainer.Emplace(TUniquePtr<FGLProgramBinaryFileCacheEntry>(NewIndexEntry));
+
+	// Store the program file descriptor in the runtime program/shader container:
+	NewIndexEntry->GLProgramState = FGLProgramBinaryFileCacheEntry::EGLProgramState::ProgramStored;
+	NewIndexEntry->FileInfo.ProgramOffset = ProgramBinaryOffset;
+	NewIndexEntry->FileInfo.ProgramSize = ProgramBinarySize;
+	NewIndexEntry->ProgramIndex = ProgramToBinaryMap.Num();
+	NewIndexEntry->FileInfo.ShaderHasheSet = ProgramKey;
+	AddProgramFileEntryToMap(NewIndexEntry);
 }
 
 void FOpenGLProgramBinaryCache::AppendProgramBinaryFileEofEntry(FArchive& Ar)
@@ -4977,51 +5052,83 @@ bool FOpenGLProgramBinaryCache::DeferShaderCompilation(GLuint Shader, const TArr
 	return false;
 }
 
-void FOpenGLProgramBinaryCache::CacheProgram(GLuint Program, const FOpenGLProgramKey& ProgramKey)
+void FOpenGLProgramBinaryCache::CacheProgram(GLuint Program, const FOpenGLProgramKey& ProgramKey, TArray<uint8>& CachedProgramBinaryOUT)
 {
 	if (CachePtr)
 	{
-		CachePtr->AppendProgramToBinaryCache(ProgramKey, Program);
+		CachePtr->AppendGLProgramToBinaryCache(ProgramKey, Program, CachedProgramBinaryOUT);
 	}
 }
 
-bool FOpenGLProgramBinaryCache::UseCachedProgram(GLuint& ProgramOUT, const FOpenGLProgramKey& ProgramKey)
+bool FOpenGLProgramBinaryCache::UseCachedProgram(GLuint& ProgramOUT, const FOpenGLProgramKey& ProgramKey, TArray<uint8>& CachedProgramBinaryOUT)
 {
 	if (CachePtr)
 	{
-		FGLProgramBinaryFileCacheEntry** ProgramBinRefPtr = nullptr;
+		return CachePtr->UseCachedProgram_internal(ProgramOUT, ProgramKey, CachedProgramBinaryOUT);
+	}
+	return false;
+}
 
-		FScopeLock Lock(&GProgramBinaryCacheCS);
+bool FOpenGLProgramBinaryCache::UseCachedProgram_internal(GLuint& ProgramOUT, const FOpenGLProgramKey& ProgramKey, TArray<uint8>& CachedProgramBinaryOUT)
+{
+	FGLProgramBinaryFileCacheEntry** ProgramBinRefPtr = nullptr;
 
-		ProgramBinRefPtr = CachePtr->ProgramToBinaryMap.Find(ProgramKey);
-		if (ProgramBinRefPtr)
-		{
-			FGLProgramBinaryFileCacheEntry* FoundProgram = *ProgramBinRefPtr;
-			check(FoundProgram->FileInfo.ShaderHasheSet == ProgramKey);
+	FScopeLock Lock(&GProgramBinaryCacheCS);
 
-			TSharedPtr<IAsyncReadRequest, ESPMode::ThreadSafe> LocalReadRequest = FoundProgram->ReadRequest.Pin();
-			bool bHasReadRequest = LocalReadRequest.IsValid();
-			check(!bHasReadRequest);
+	ProgramBinRefPtr = ProgramToBinaryMap.Find(ProgramKey);
 
-			// by this point the program must be either available or no attempt to load from shader library has occurred.
-			checkf(FoundProgram->GLProgramState == FGLProgramBinaryFileCacheEntry::EGLProgramState::ProgramStored
-				|| FoundProgram->GLProgramState == FGLProgramBinaryFileCacheEntry::EGLProgramState::ProgramAvailable,
-				TEXT("Unexpected program state: %d"), (int32)FoundProgram->GLProgramState);
+	if (ProgramBinRefPtr)
+	{
+		FGLProgramBinaryFileCacheEntry* FoundProgram = *ProgramBinRefPtr;
+		check(FoundProgram->FileInfo.ShaderHasheSet == ProgramKey);
 
-			if (FoundProgram->GLProgramState == FGLProgramBinaryFileCacheEntry::EGLProgramState::ProgramAvailable)
-			{
-				UE_LOG(LogRHI, Log, TEXT("UseCachedProgram : Program (%s) GLid = %x is ready!"), *ProgramKey.ToString(), FoundProgram->GLProgramId);
-				ProgramOUT = FoundProgram->GLProgramId;
+		TSharedPtr<IAsyncReadRequest, ESPMode::ThreadSafe> LocalReadRequest = FoundProgram->ReadRequest.Pin();
+		bool bHasReadRequest = LocalReadRequest.IsValid();
+		check(!bHasReadRequest);
 
-				// GLProgram has been handed over. 
-				FoundProgram->GLProgramId = 0;
-				FoundProgram->GLProgramState = FGLProgramBinaryFileCacheEntry::EGLProgramState::ProgramComplete;
-				return true;
-			}
-			else
+		// by this point the program must be either available or no attempt to load from shader library has occurred.
+		checkf(FoundProgram->GLProgramState == FGLProgramBinaryFileCacheEntry::EGLProgramState::ProgramStored
+			|| FoundProgram->GLProgramState == FGLProgramBinaryFileCacheEntry::EGLProgramState::ProgramAvailable,
+			TEXT("Unexpected program state: %d"), (int32)FoundProgram->GLProgramState);
+
+		if (FoundProgram->GLProgramState == FGLProgramBinaryFileCacheEntry::EGLProgramState::ProgramAvailable)
+		{
+			UE_LOG(LogRHI, Log, TEXT("UseCachedProgram : Program (%s) GLid = %x is ready!"), *ProgramKey.ToString(), FoundProgram->GLProgramId);
+			ProgramOUT = FoundProgram->GLProgramId;
+
+			// GLProgram has been handed over. 
+			FoundProgram->GLProgramId = 0;
+			FoundProgram->GLProgramState = FGLProgramBinaryFileCacheEntry::EGLProgramState::ProgramComplete;
+			return true;
+		}
+		else
+		{
+			UE_LOG(LogRHI, Log, TEXT("UseCachedProgram : %s was not ready when needed!! (state %d)"), *ProgramKey.ToString(), (uint32)FoundProgram->GLProgramState);
+		}
+	}
+	else if (BinaryFileState == EBinaryFileState::BuildingCacheFileWithMove)
+	{
+		// We're building the new cache using the original cache to warm:
+		TUniquePtr<FGLProgramBinaryFileCacheEntry>* FoundExistingBinary = PreviousBinaryCacheInfo.ProgramToOldBinaryCacheMap.Find(ProgramKey);
+		if (FoundExistingBinary)
+		{
+			TUniquePtr<FGLProgramBinaryFileCacheEntry>& ExistingBinary = *FoundExistingBinary;
+			// read old binary:
+			CachedProgramBinaryOUT.SetNumUninitialized(ExistingBinary->FileInfo.ProgramSize);
+			PreviousBinaryCacheInfo.OldCacheArchive->Seek(ExistingBinary->FileInfo.ProgramOffset);
+			PreviousBinaryCacheInfo.OldCacheArchive->Serialize(CachedProgramBinaryOUT.GetData(), ExistingBinary->FileInfo.ProgramSize);
+			bool bSuccess = CreateGLProgramFromBinary(ProgramOUT, CachedProgramBinaryOUT);
+			if (!bSuccess)
 			{
-				UE_LOG(LogRHI, Log, TEXT("UseCachedProgram : %s was not ready when needed!! (state %d)"), *ProgramKey.ToString(), (uint32)FoundProgram->GLProgramState);
+				RHIGetPanicDelegate().ExecuteIfBound(FName("FailedBinaryProgramCreateFromOldCache"));
+				UE_LOG(LogRHI, Fatal, TEXT("UseCachedProgram : Failed to create GL program from binary data while BuildingCacheFileWithMove!"));
 			}
+			SetNewProgramStats(ProgramOUT);
+			// Now write to new cache, we're returning true here so no attempt will be made to add it back to the cache later.
+			AddProgramBinaryDataToBinaryCache(*BinaryCacheWriteFileHandle, CachedProgramBinaryOUT, ProgramKey);
+
+			PreviousBinaryCacheInfo.NumberOfOldEntriesReused++;
+			return true;
 		}
 	}
 	return false;
diff --git a/Engine/Source/Runtime/OpenGLDrv/Public/OpenGLShaderResources.h b/Engine/Source/Runtime/OpenGLDrv/Public/OpenGLShaderResources.h
--- a/Engine/Source/Runtime/OpenGLDrv/Public/OpenGLShaderResources.h
+++ b/Engine/Source/Runtime/OpenGLDrv/Public/OpenGLShaderResources.h
@@ -485,10 +485,10 @@ public:
 	static void CompilePendingShaders(const FOpenGLLinkedProgramConfiguration& Config);
 	
 	/** Try to find and load program binary from cache */
-	static bool UseCachedProgram(GLuint& ProgramOUT, const FOpenGLProgramKey& ProgramKey);
+	static bool UseCachedProgram(GLuint& ProgramOUT, const FOpenGLProgramKey& ProgramKey, TArray<uint8>& CachedProgramBinaryOUT);
 	
 	/** Store program binary on disk in case ProgramBinaryCache is enabled */
-	static void CacheProgram(GLuint Program, const FOpenGLProgramKey& ProgramKey);
+	static void CacheProgram(GLuint Program, const FOpenGLProgramKey& ProgramKey, TArray<uint8>& CachedProgramBinaryOUT);
 
 	static void OnShaderLibraryRequestShaderCode(const FSHAHash& Hash, FArchive* Ar);
 
@@ -545,11 +545,13 @@ private:
 	void OpenAsyncReadHandle();
 	void CloseAsyncReadHandle();
 
-	bool OpenWriteHandle(bool bTruncate = false);
+	bool OpenWriteHandle();
 	void CloseWriteHandle();
 
-	void AppendProgramToBinaryCache(const FOpenGLProgramKey& ProgramKey, GLuint Program);
-	void AddUniqueProgramToBinaryCache(FArchive* FileWriter, const FOpenGLProgramKey& ProgramKey, GLuint Program);
+	void AppendGLProgramToBinaryCache(const FOpenGLProgramKey& ProgramKey, GLuint Program, TArray<uint8>& CachedProgramBinaryOUT);
+	void AddUniqueGLProgramToBinaryCache(FArchive* FileWriter, const FOpenGLProgramKey& ProgramKey, GLuint Program, TArray<uint8>& CachedProgramBinaryOUT);
+
+	void AddProgramBinaryDataToBinaryCache(FArchive& Ar, TArray<uint8>& BinaryProgramData, const FOpenGLProgramKey& ProgramKey);
 
 	void ReleaseGLProgram_internal(FOpenGLLinkedProgramConfiguration& Config, GLuint Program);
 
@@ -559,6 +561,8 @@ private:
 		return FoundShaderToBinary && FoundShaderToBinary->bLoaded;
 	}
 
+	bool UseCachedProgram_internal(GLuint& ProgramOUT, const FOpenGLProgramKey& ProgramKey, TArray<uint8>& CachedProgramBinaryOUT);
+
 	void OnShaderLibraryRequestShaderCode_internal(const FSHAHash& Hash, FArchive* Ar);
 
 	void BeginProgramReadRequest(FGLProgramBinaryFileCacheEntry* IndexEntry, FArchive* Ar);
@@ -602,6 +606,7 @@ private:
 		TArray<FGLProgramBinaryFileCacheEntry*> AssociatedPrograms;
 	};
 
+	// Map of shader hash to a list of programs which reference it.
 	TMap<FSHAHash, FGLShaderToPrograms> ShaderToProgramsMap;
 
 	// programs loaded via async and now ready for creation on GL context owning thread.
@@ -613,10 +618,33 @@ private:
 
 	enum class EBinaryFileState : uint8
 	{
-		Uninitialized,		// No binary file is yet established and we should not read or write to it.
-		BuildingCacheFile,	// We are precompiling shaders from the PSO and storing them in a new binary cache. Do not attempt to read.
-		ValidCacheFile,		// We have a valid cache file we can use for reading. Do not attempt to write.
+		Uninitialized,					// No binary file is yet established and we should not read or write to it.
+		BuildingCacheFile,				// We are precompiling shaders from the PSO and storing them in a new binary cache. Do not attempt to read.
+		BuildingCacheFileWithMove,		// We are precompiling shaders from the PSO and storing them in a new binary cache, shaders matching from the existing cache are moved to the new file. Do not attempt to read.
+		ValidCacheFile,					// We have a valid cache file we can use for reading. Do not attempt to write.
+	};
+
+	bool IsBuildingCache_internal() const
+	{
+		return BinaryFileState == EBinaryFileState::BuildingCacheFile || BinaryFileState == EBinaryFileState::BuildingCacheFileWithMove;
+	}
+
+	struct FPreviousGLProgramBinaryCacheInfo
+	{
+		FPreviousGLProgramBinaryCacheInfo();
+		FPreviousGLProgramBinaryCacheInfo(FPreviousGLProgramBinaryCacheInfo&&);
+		FPreviousGLProgramBinaryCacheInfo& operator=(FPreviousGLProgramBinaryCacheInfo&&);
+		~FPreviousGLProgramBinaryCacheInfo();
+
+
+		FString OldCacheFilename;
+		TUniquePtr<FArchive> OldCacheArchive;
+		TMap<FOpenGLProgramKey, TUniquePtr<FGLProgramBinaryFileCacheEntry> > ProgramToOldBinaryCacheMap; // program key to program entry for old cache used when generating new cache.
+
+		// for logging
+		uint32 NumberOfOldEntriesReused;
 	};
+	FPreviousGLProgramBinaryCacheInfo PreviousBinaryCacheInfo;
 
 	EBinaryFileState BinaryFileState;
 };
