Copied cl 5433386
nDisplay openGL crash with multi GPU

#rnx

diff --git a/Engine/Source/Runtime/OpenGLDrv/Private/Windows/OpenGLWindows.cpp b/Engine/Source/Runtime/OpenGLDrv/Private/Windows/OpenGLWindows.cpp
--- a/Engine/Source/Runtime/OpenGLDrv/Private/Windows/OpenGLWindows.cpp
+++ b/Engine/Source/Runtime/OpenGLDrv/Private/Windows/OpenGLWindows.cpp
@@ -178,13 +178,18 @@ static void PlatformCreateDummyGLWindow(FPlatformOpenGLContext* OutContext)
 		check(ClassAtom);
 	}
 
+	int32 WinX = 0;
+	int32 WinY = 0;
+	FParse::Value(FCommandLine::Get(), TEXT("WinX="), WinX);
+	FParse::Value(FCommandLine::Get(), TEXT("WinY="), WinY);
+
 	// Create a dummy window.
 	OutContext->WindowHandle = CreateWindowEx(
 		WS_EX_WINDOWEDGE,
 		WindowClassName,
 		NULL,
 		WS_POPUP,
-		0, 0, 1, 1,
+		WinX, WinY, 1, 1,
 		NULL, NULL, NULL, NULL);
 	check(OutContext->WindowHandle);
 	OutContext->bReleaseWindowOnDestroy = true;
