fix cooked builds not working with opengl on windows - sm4 was incorrectly selected while sm5 was cooked.

diff --git a/Engine/Source/Runtime/OpenGLDrv/Private/Windows/OpenGLWindows.cpp b/Engine/Source/Runtime/OpenGLDrv/Private/Windows/OpenGLWindows.cpp
--- a/Engine/Source/Runtime/OpenGLDrv/Private/Windows/OpenGLWindows.cpp
+++ b/Engine/Source/Runtime/OpenGLDrv/Private/Windows/OpenGLWindows.cpp
@@ -225,7 +225,7 @@ static void PlatformOpenGLVersionFromCommandLine(int& OutMajorVersion, int& OutM
 	bool bGL4 = PlatformOpenGL4();
 	if (!bGL3 && !bGL4)
 	{
-		if (GRequestedFeatureLevel == ERHIFeatureLevel::SM5)
+		if (GRequestedFeatureLevel >= ERHIFeatureLevel::SM5)
 		{
 			bGL4 = true;
 		}
