Set GL_UNPACK_ALGINMENT to 1 for glTexImage* and glTexSubImage* commands.

#rnx

diff --git a/Engine/Source/Runtime/OpenGLDrv/Private/OpenGLTexture.cpp b/Engine/Source/Runtime/OpenGLDrv/Private/OpenGLTexture.cpp
--- a/Engine/Source/Runtime/OpenGLDrv/Private/OpenGLTexture.cpp
+++ b/Engine/Source/Runtime/OpenGLDrv/Private/OpenGLTexture.cpp
@@ -565,6 +565,8 @@ void FOpenGLDynamicRHI::InitializeGLTexture(FRHITexture* Texture, uint32 SizeX,
 				uint32 NumBlocksY = AlignArbitrary(FMath::Max<uint32>(1,(SizeY >> MipIndex)), BlockSizeY) / BlockSizeY;
 				uint32 NumLayers = FMath::Max<uint32>(1,ArraySize);
 				
+				glPixelStorei(GL_UNPACK_ALIGNMENT, 1);
+
 				if(bArrayTexture )
 				{
 					if(bCubeTexture)
@@ -614,6 +616,8 @@ void FOpenGLDynamicRHI::InitializeGLTexture(FRHITexture* Texture, uint32 SizeX,
 						MipOffset += NumBlocksX * NumBlocksY * NumLayers * GPixelFormats[Format].BlockBytes;
 					}
 				}
+
+				glPixelStorei(GL_UNPACK_ALIGNMENT, 4);
 			}
 
 			BulkData->Discard();
@@ -1271,6 +1275,7 @@ void TOpenGLTexture<RHIResourceType>::Unlock(uint32 MipIndex,uint32 ArrayIndex)
 		}
 		else
 		{
+			glPixelStorei(GL_UNPACK_ALIGNMENT, 1);
 			if (GetAllocatedStorageForMip(MipIndex,ArrayIndex))
 			{
 				glTexSubImage2D(
@@ -1298,6 +1303,7 @@ void TOpenGLTexture<RHIResourceType>::Unlock(uint32 MipIndex,uint32 ArrayIndex)
 					LockedBuffer);
 				SetAllocatedStorageForMip(MipIndex,ArrayIndex);
 			}
+			glPixelStorei(GL_UNPACK_ALIGNMENT, 4);
 		}
 		if(bUseClientStorage)
 		{
@@ -1662,6 +1668,7 @@ FTexture2DArrayRHIRef FOpenGLDynamicRHI::RHICreateTexture2DArray(uint32 SizeX,ui
 
 	if (Data)
 	{
+		glPixelStorei(GL_UNPACK_ALIGNMENT, 1);
 		for(uint32 MipIndex = 0; MipIndex < NumMips; MipIndex++)
 		{
 			FOpenGL::TexSubImage3D(
@@ -1682,6 +1689,7 @@ FTexture2DArrayRHIRef FOpenGLDynamicRHI::RHICreateTexture2DArray(uint32 SizeX,ui
 			uint32 SysMemSlicePitch =  FMath::Max<uint32>(1,SizeY >> MipIndex) * SysMemPitch;
 			MipOffset               += SizeZ * SysMemSlicePitch;
 		}
+		glPixelStorei(GL_UNPACK_ALIGNMENT, 4);
 
 		Info.BulkData->Discard();
 	}
@@ -1808,6 +1816,7 @@ FTexture3DRHIRef FOpenGLDynamicRHI::RHICreateTexture3D(uint32 SizeX,uint32 SizeY
 			}
 			else
 			{
+				glPixelStorei(GL_UNPACK_ALIGNMENT, 1);
 				FOpenGL::TexSubImage3D(
 					/*Target=*/ Target,
 					/*Level=*/ MipIndex,
@@ -1816,6 +1825,7 @@ FTexture3DRHIRef FOpenGLDynamicRHI::RHICreateTexture3D(uint32 SizeX,uint32 SizeY
 					/*Format=*/ GLFormat.Format,
 					/*Type=*/ GLFormat.Type,
 					/*Data=*/ Data + MipOffset);
+				glPixelStorei(GL_UNPACK_ALIGNMENT, 4);
 			}
 
 			MipOffset += MipSize;
