Use GL_QCOM_shader_framebuffer_fetch_noncoherent extension on supported devices for better GPU performance
Remove shader_framebuffer_fetch exception for Adreno51x devices (old drivers issue)

diff --git a/Engine/Config/BaseDeviceProfiles.ini b/Engine/Config/BaseDeviceProfiles.ini
--- a/Engine/Config/BaseDeviceProfiles.ini
+++ b/Engine/Config/BaseDeviceProfiles.ini
@@ -496,7 +496,6 @@ BaseProfileName=Android_Adreno330
 [Android_Adreno51x DeviceProfile]
 DeviceType=Android
 BaseProfileName=Android_High
-+CVars=r.OpenGL.StripExtensions=GL_EXT_shader_framebuffer_fetch
 
 [Android_Adreno5xx DeviceProfile]
 DeviceType=Android
diff --git a/Engine/Source/Runtime/OpenGLDrv/Private/Android/AndroidOpenGL.cpp b/Engine/Source/Runtime/OpenGLDrv/Private/Android/AndroidOpenGL.cpp
--- a/Engine/Source/Runtime/OpenGLDrv/Private/Android/AndroidOpenGL.cpp
+++ b/Engine/Source/Runtime/OpenGLDrv/Private/Android/AndroidOpenGL.cpp
@@ -110,6 +110,7 @@ PFNeglQueryTimestampSupportedANDROID eglQueryTimestampSupportedANDROID_p = NULL;
 PFNeglQueryTimestampSupportedANDROID eglGetCompositorTimingSupportedANDROID_p = NULL;
 PFNeglQueryTimestampSupportedANDROID eglGetFrameTimestampsSupportedANDROID_p = NULL;
 
+PFNGLFRAMEBUFFERFETCHBARRIERQCOMPROC glFramebufferFetchBarrierQCOM = NULL;
 
 int32 FAndroidOpenGL::GLMajorVerion = 0;
 int32 FAndroidOpenGL::GLMinorVersion = 0;
@@ -814,6 +815,16 @@ FAndroidOpenGL::EFeatureLevelSupport FAndroidOpenGL::CurrentFeatureLevelSupport
 extern bool AndroidThunkCpp_GetMetaDataBoolean(const FString& Key);
 extern FString AndroidThunkCpp_GetMetaDataString(const FString& Key);
 
+void FAndroidOpenGL::SetupDefaultGLContextState(const FString& ExtensionsString)
+{
+	// Enable QCOM non-coherent framebuffer fetch if supported
+	if (ExtensionsString.Contains(TEXT("GL_QCOM_shader_framebuffer_fetch_noncoherent")) && 
+		ExtensionsString.Contains(TEXT("GL_EXT_shader_framebuffer_fetch")))
+	{
+		glEnable(GL_FRAMEBUFFER_FETCH_NONCOHERENT_QCOM);
+	}
+}
+
 void FAndroidOpenGL::ProcessExtensions(const FString& ExtensionsString)
 {
 	FOpenGLES2::ProcessExtensions(ExtensionsString);
@@ -1281,6 +1292,17 @@ void FAndroidOpenGL::ProcessExtensions(const FString& ExtensionsString)
 		}
 		bSupportsCopyImage = (glCopyImageSubData != nullptr);
 	}
+
+	// Qualcomm non-coherent framebuffer_fetch
+	if (ExtensionsString.Contains(TEXT("GL_QCOM_shader_framebuffer_fetch_noncoherent")) && 
+		ExtensionsString.Contains(TEXT("GL_EXT_shader_framebuffer_fetch")))
+	{
+		glFramebufferFetchBarrierQCOM = (PFNGLFRAMEBUFFERFETCHBARRIERQCOMPROC)((void*)eglGetProcAddress("glFramebufferFetchBarrierQCOM"));
+		if (glFramebufferFetchBarrierQCOM != nullptr)
+		{
+			UE_LOG(LogRHI, Log, TEXT("Using QCOM_shader_framebuffer_fetch_noncoherent"));
+		}
+	}
 }
 
 FString FAndroidMisc::GetGPUFamily()
diff --git a/Engine/Source/Runtime/OpenGLDrv/Private/Android/AndroidOpenGL.h b/Engine/Source/Runtime/OpenGLDrv/Private/Android/AndroidOpenGL.h
--- a/Engine/Source/Runtime/OpenGLDrv/Private/Android/AndroidOpenGL.h
+++ b/Engine/Source/Runtime/OpenGLDrv/Private/Android/AndroidOpenGL.h
@@ -223,6 +223,12 @@ extern "C"
 	extern PFNeglQueryTimestampSupportedANDROID eglGetFrameTimestampsSupportedANDROID_p;
 }
 
+#ifndef GL_FRAMEBUFFER_FETCH_NONCOHERENT_QCOM
+#define GL_FRAMEBUFFER_FETCH_NONCOHERENT_QCOM	0x96A2
+#endif
+typedef void (GL_APIENTRYP PFNGLFRAMEBUFFERFETCHBARRIERQCOMPROC) (void);
+extern PFNGLFRAMEBUFFERFETCHBARRIERQCOMPROC	glFramebufferFetchBarrierQCOM;
+
 struct FAndroidOpenGL : public FOpenGLES2
 {
 	static FORCEINLINE bool IsES31Usable()
@@ -666,8 +672,18 @@ struct FAndroidOpenGL : public FOpenGLES2
 
 	static FORCEINLINE GLint GetMaxComputeTextureImageUnits() { check(MaxComputeTextureImageUnits != -1); return MaxComputeTextureImageUnits; }
 	static FORCEINLINE GLint GetMaxComputeUniformComponents() { check(MaxComputeUniformComponents != -1); return MaxComputeUniformComponents; }
-	
+
+	static FORCEINLINE void FrameBufferFetchBarrier()
+	{
+		if (glFramebufferFetchBarrierQCOM)
+		{
+			glFramebufferFetchBarrierQCOM();
+		}
+	}
+			
 	static void ProcessExtensions(const FString& ExtensionsString);
+	static void SetupDefaultGLContextState(const FString& ExtensionsString);
+
 
 	// whether to use ES 3.0 function glTexStorage2D to allocate storage for GL_HALF_FLOAT_OES render target textures
 	static bool bUseHalfFloatTexStorage;
diff --git a/Engine/Source/Runtime/OpenGLDrv/Private/OpenGLDrv.cpp b/Engine/Source/Runtime/OpenGLDrv/Private/OpenGLDrv.cpp
--- a/Engine/Source/Runtime/OpenGLDrv/Private/OpenGLDrv.cpp
+++ b/Engine/Source/Runtime/OpenGLDrv/Private/OpenGLDrv.cpp
@@ -640,6 +640,9 @@ void InitDefaultGLContextState(void)
 		glClipControl(GL_UPPER_LEFT, GL_ZERO_TO_ONE);
 	}
 #endif
+
+	// optional per platform setup
+	FOpenGL::SetupDefaultGLContextState(ExtensionsString);
 }
 
 #undef LOCTEXT_NAMESPACE
diff --git a/Engine/Source/Runtime/OpenGLDrv/Private/OpenGLRenderTarget.cpp b/Engine/Source/Runtime/OpenGLDrv/Private/OpenGLRenderTarget.cpp
--- a/Engine/Source/Runtime/OpenGLDrv/Private/OpenGLRenderTarget.cpp
+++ b/Engine/Source/Runtime/OpenGLDrv/Private/OpenGLRenderTarget.cpp
@@ -1194,3 +1194,14 @@ void FOpenGLDynamicRHI::RHIEndRenderPass()
 		}
 	}
 }
+
+void FOpenGLDynamicRHI::RHINextSubpass()
+{
+	IRHICommandContext::RHINextSubpass();
+	
+	if (RenderPassInfo.SubpassHint == ESubpassHint::DepthReadSubpass)
+	{
+		FOpenGL::FrameBufferFetchBarrier();
+	}
+}
+
diff --git a/Engine/Source/Runtime/OpenGLDrv/Public/OpenGL.h b/Engine/Source/Runtime/OpenGLDrv/Public/OpenGL.h
--- a/Engine/Source/Runtime/OpenGLDrv/Public/OpenGL.h
+++ b/Engine/Source/Runtime/OpenGLDrv/Public/OpenGL.h
@@ -76,6 +76,7 @@ public:
 
 	static void ProcessQueryGLInt();
 	static void ProcessExtensions(const FString& ExtensionsString);
+	static void SetupDefaultGLContextState(const FString& ExtensionsString) {};
 
 	static FORCEINLINE bool SupportsMapBuffer()							{ return true; }
 	static FORCEINLINE bool SupportsDepthTexture()						{ return true; }
@@ -357,7 +358,8 @@ public:
 	static FORCEINLINE void GetProgramBinary(GLuint Program, GLsizei BufSize, GLsizei *Length, GLenum *BinaryFormat, void *Binary) UGL_OPTIONAL_VOID
 	static FORCEINLINE void ProgramBinary(GLuint Program, GLenum BinaryFormat, const void *Binary, GLsizei Length) UGL_OPTIONAL_VOID
 
-
+	static FORCEINLINE void FrameBufferFetchBarrier() UGL_OPTIONAL_VOID
+	
 	static FPlatformOpenGLDevice*	CreateDevice() UGL_REQUIRED(NULL)
 	static FPlatformOpenGLContext*	CreateContext( FPlatformOpenGLDevice* Device, void* WindowHandle ) UGL_REQUIRED(NULL)
 
diff --git a/Engine/Source/Runtime/OpenGLDrv/Public/OpenGLDrv.h b/Engine/Source/Runtime/OpenGLDrv/Public/OpenGLDrv.h
--- a/Engine/Source/Runtime/OpenGLDrv/Public/OpenGLDrv.h
+++ b/Engine/Source/Runtime/OpenGLDrv/Public/OpenGLDrv.h
@@ -445,7 +445,8 @@ public:
 
 	virtual void RHIBeginRenderPass(const FRHIRenderPassInfo& InInfo, const TCHAR* InName) final override;
 	virtual void RHIEndRenderPass() final override;
-
+	virtual void RHINextSubpass() final override;
+	
 	virtual void RHISetComputeShader(FRHIComputeShader* ComputeShader) final override;
 	virtual void RHIDispatchComputeShader(uint32 ThreadGroupCountX, uint32 ThreadGroupCountY, uint32 ThreadGroupCountZ) final override;
 	virtual void RHIDispatchIndirectComputeShader(FRHIVertexBuffer* ArgumentBuffer, uint32 ArgumentOffset) final override;
