Show proper error dialog for case of neither OpenGL ES 2 or 3.1 fallback support available when device does not support Vulkan

#ue4
#android

diff --git a/Engine/Source/Runtime/OpenGLDrv/Private/Android/AndroidOpenGL.cpp b/Engine/Source/Runtime/OpenGLDrv/Private/Android/AndroidOpenGL.cpp
--- a/Engine/Source/Runtime/OpenGLDrv/Private/Android/AndroidOpenGL.cpp
+++ b/Engine/Source/Runtime/OpenGLDrv/Private/Android/AndroidOpenGL.cpp
@@ -313,10 +313,24 @@ bool PlatformInitOpenGL()
 			// If we're here and there's no ES2 data then we're in trouble.
 			if (!bBuildForES2)
 			{
-				Message.Append(TEXT("This device only supports OpenGL ES 2 but the app was not packaged with ES2 support."));
-				FPlatformMisc::LowLevelOutputDebugString(*Message);
-				FAndroidMisc::MessageBoxExt(EAppMsgType::Ok, *Message, TEXT("Unable to run on this device!"));
-				checkf(bBuildForES2, TEXT("This device only supports OpenGL ES 2 but the app was not packaged with ES2 support."));
+				if (bES31Supported)
+				{
+					Message.Append(TEXT("This device does not support Vulkan but the app was not packaged with either OpenGL ES 2 or ES 3.1 support."));
+					if (FAndroidMisc::GetAndroidBuildVersion() < 26)
+					{
+						Message.Append(TEXT(" Updating to a newer Android version may resolve this issue."));
+					}
+					FPlatformMisc::LowLevelOutputDebugString(*Message);
+					FAndroidMisc::MessageBoxExt(EAppMsgType::Ok, *Message, TEXT("Unable to run on this device!"));
+					checkf(bBuildForES2, TEXT("This device does not support Vulkan but the app was not packaged with either OpenGL ES 2 or ES 3.1 support."));
+				}
+				else
+				{
+					Message.Append(TEXT("This device only supports OpenGL ES 2 but the app was not packaged with ES2 support."));
+					FPlatformMisc::LowLevelOutputDebugString(*Message);
+					FAndroidMisc::MessageBoxExt(EAppMsgType::Ok, *Message, TEXT("Unable to run on this device!"));
+					checkf(bBuildForES2, TEXT("This device does not support Vulkan but the app was not packaged with either OpenGL ES 2 or ES 3.1 support."));
+				}
 			}
 		}
 	}
