Fixed texture invalidation in OpenGL RHI when texture is deleted.

diff --git a/Engine/Source/Runtime/OpenGLDrv/Private/OpenGLTexture.cpp b/Engine/Source/Runtime/OpenGLDrv/Private/OpenGLTexture.cpp
--- a/Engine/Source/Runtime/OpenGLDrv/Private/OpenGLTexture.cpp
+++ b/Engine/Source/Runtime/OpenGLDrv/Private/OpenGLTexture.cpp
@@ -2398,7 +2398,7 @@ void FOpenGLDynamicRHI::RHIUpdateTexture3D(FRHITexture3D* TextureRHI, uint32 Mip
 void FOpenGLDynamicRHI::InvalidateTextureResourceInCache(GLuint Resource)
 {
 	VERIFY_GL_SCOPE();
-	if (SharedContextState.Textures || RenderingContextState.Textures)
+	if (SharedContextState.Textures || RenderingContextState.Textures || PendingState.Textures)
 	{
 		for (int32 SamplerIndex = 0; SamplerIndex < FOpenGL::GetMaxCombinedTextureImageUnits(); ++SamplerIndex)
 		{
@@ -2413,6 +2413,12 @@ void FOpenGLDynamicRHI::InvalidateTextureResourceInCache(GLuint Resource)
 				RenderingContextState.Textures[SamplerIndex].Target = GL_NONE;
 				RenderingContextState.Textures[SamplerIndex].Resource = 0;
 			}
+
+			if (PendingState.Textures && PendingState.Textures[SamplerIndex].Resource == Resource)
+			{
+				PendingState.Textures[SamplerIndex].Target = GL_NONE;
+				PendingState.Textures[SamplerIndex].Resource = 0;
+			}
 		}
 	}
 	
@@ -2439,6 +2445,12 @@ void FOpenGLDynamicRHI::InvalidateUAVResourceInCache(GLuint Resource)
 			RenderingContextState.UAVs[UAVIndex].Format = GL_NONE;
 			RenderingContextState.UAVs[UAVIndex].Resource = 0;
 		}
+
+		if (PendingState.UAVs[UAVIndex].Resource == Resource)
+		{
+			PendingState.UAVs[UAVIndex].Format = GL_NONE;
+			PendingState.UAVs[UAVIndex].Resource = 0;
+		}
 	}
 }
 
diff --git a/Engine/Source/Runtime/OpenGLDrv/Public/OpenGLState.h b/Engine/Source/Runtime/OpenGLDrv/Public/OpenGLState.h
--- a/Engine/Source/Runtime/OpenGLDrv/Public/OpenGLState.h
+++ b/Engine/Source/Runtime/OpenGLDrv/Public/OpenGLState.h
@@ -285,8 +285,12 @@ struct FOpenGLCommonState
 		CleanupResources();
 	}
 
+	// NumCombinedTextures must be greater than or equal to FOpenGL::GetMaxCombinedTextureImageUnits()
+	// NumComputeUAVUnits must be greater than or equal to OGL_MAX_COMPUTE_STAGE_UAV_UNITS
 	virtual void InitializeResources(int32 NumCombinedTextures, int32 NumComputeUAVUnits)
 	{
+		check(NumCombinedTextures >= FOpenGL::GetMaxCombinedTextureImageUnits());
+		check(NumComputeUAVUnits >= OGL_MAX_COMPUTE_STAGE_UAV_UNITS);
 		check(!Textures && !SamplerStates && !UAVs);
 		Textures = new FTextureStage[NumCombinedTextures];
 		SamplerStates = new FOpenGLSamplerState*[NumCombinedTextures];
