Added separate path for RHICreateRenderQueryPool function in OpenGL RHI as workaround.

#rnx

diff --git a/Engine/Source/Runtime/OpenGLDrv/Private/OpenGLQuery.cpp b/Engine/Source/Runtime/OpenGLDrv/Private/OpenGLQuery.cpp
--- a/Engine/Source/Runtime/OpenGLDrv/Private/OpenGLQuery.cpp
+++ b/Engine/Source/Runtime/OpenGLDrv/Private/OpenGLQuery.cpp
@@ -614,6 +614,12 @@ bool FOpenGLDynamicRHI::RHIGetRenderQueryResult(FRHIRenderQuery* QueryRHI,uint64
 
 
 
+FRenderQueryPoolRHIRef FOpenGLDynamicRHI::RHICreateRenderQueryPool(ERenderQueryType QueryType, uint32 /*NumQueries*/)
+{
+	// Workaround to fix UE-77873: pass <QueryType> to primary implementation and ignore <NumQueries> for pre-allocated queries
+	return FDynamicRHI::RHICreateRenderQueryPool(QueryType);
+}
+
 extern void OnQueryCreation( FOpenGLRenderQuery* Query );
 extern void OnQueryDeletion( FOpenGLRenderQuery* Query );
 
diff --git a/Engine/Source/Runtime/OpenGLDrv/Public/OpenGLDrv.h b/Engine/Source/Runtime/OpenGLDrv/Public/OpenGLDrv.h
--- a/Engine/Source/Runtime/OpenGLDrv/Public/OpenGLDrv.h
+++ b/Engine/Source/Runtime/OpenGLDrv/Public/OpenGLDrv.h
@@ -538,6 +538,7 @@ public:
 	virtual void* LockStagingBuffer_RenderThread(class FRHICommandListImmediate& RHICmdList, FRHIStagingBuffer* StagingBuffer, uint32 Offset, uint32 SizeRHI) final override;
 	virtual	void UnlockStagingBuffer_RenderThread(class FRHICommandListImmediate& RHICmdList, FRHIStagingBuffer* StagingBuffer) final override;
 
+	virtual FRenderQueryPoolRHIRef RHICreateRenderQueryPool(ERenderQueryType QueryType, uint32 NumQueries = UINT32_MAX) final override;
 	virtual FStagingBufferRHIRef RHICreateStagingBuffer() final override;
 	virtual FGPUFenceRHIRef RHICreateGPUFence(const FName &Name) final override;
 
