Correctly discard attachments for a default GL framebuffer

diff --git a/Engine/Source/Runtime/OpenGLDrv/Private/OpenGLCommands.cpp b/Engine/Source/Runtime/OpenGLDrv/Private/OpenGLCommands.cpp
--- a/Engine/Source/Runtime/OpenGLDrv/Private/OpenGLCommands.cpp
+++ b/Engine/Source/Runtime/OpenGLDrv/Private/OpenGLCommands.cpp
@@ -1770,34 +1770,47 @@ void FOpenGLDynamicRHI::RHIDiscardRenderTargets(bool Depth, bool Stencil, uint32
 	if (FOpenGL::SupportsDiscardFrameBuffer())
 	{
 		VERIFY_GL_SCOPE();
+		const bool bDefaultFramebuffer = (PendingState.Framebuffer == 0);
 		uint32 ColorBitMask = ColorBitMaskIn;
 		// 8 Color + Depth + Stencil = 10
 		GLenum Attachments[MaxSimultaneousRenderTargets + 2];
 		uint32 I = 0;
 		if (Depth)
 		{
-			Attachments[I] = GL_DEPTH_ATTACHMENT;
+			Attachments[I] = bDefaultFramebuffer ? GL_DEPTH : GL_DEPTH_ATTACHMENT;
 			I++;
 		}
 		if (Stencil)
 		{
-			Attachments[I] = GL_STENCIL_ATTACHMENT;
+			Attachments[I] = bDefaultFramebuffer ? GL_STENCIL : GL_STENCIL_ATTACHMENT;
 			I++;
 		}
 
-		ColorBitMask &= (1 << MaxSimultaneousRenderTargets) - 1;
-		uint32 J = 0;
-		while (ColorBitMask)
+		if (bDefaultFramebuffer)
 		{
-			if (ColorBitMask & 1)
+			if (ColorBitMask)
 			{
-				Attachments[I] = GL_COLOR_ATTACHMENT0 + J;
+				Attachments[I] = GL_COLOR;
 				I++;
 			}
+		}
+		else
+		{
+			ColorBitMask &= (1 << MaxSimultaneousRenderTargets) - 1;
+			uint32 J = 0;
+			while (ColorBitMask)
+			{
+				if (ColorBitMask & 1)
+				{
+					Attachments[I] = GL_COLOR_ATTACHMENT0 + J;
+					I++;
+				}
 
-			ColorBitMask >>= 1;
-			++J;
+				ColorBitMask >>= 1;
+				++J;
+			}
 		}
+
 		FOpenGL::DiscardFramebufferEXT(GL_FRAMEBUFFER, I, Attachments);
 	}
 }
diff --git a/Engine/Source/Runtime/OpenGLDrv/Private/OpenGLRenderTarget.cpp b/Engine/Source/Runtime/OpenGLDrv/Private/OpenGLRenderTarget.cpp
--- a/Engine/Source/Runtime/OpenGLDrv/Private/OpenGLRenderTarget.cpp
+++ b/Engine/Source/Runtime/OpenGLDrv/Private/OpenGLRenderTarget.cpp
@@ -1185,7 +1185,7 @@ void FOpenGLDynamicRHI::RHIEndRenderPass()
 	if (RenderPassInfo.DepthStencilRenderTarget.DepthStencilTarget)
 	{
 		ERenderTargetActions DepthActions = GetDepthActions(RenderPassInfo.DepthStencilRenderTarget.Action);
-		ERenderTargetActions StencilActions = GetDepthActions(RenderPassInfo.DepthStencilRenderTarget.Action);
+		ERenderTargetActions StencilActions = GetStencilActions(RenderPassInfo.DepthStencilRenderTarget.Action);
 		bool bDiscardDepth = GetStoreAction(DepthActions) == ERenderTargetStoreAction::ENoAction;
 		bool bDiscardStencil = GetStoreAction(StencilActions) == ERenderTargetStoreAction::ENoAction;
 		if (bDiscardDepth || bDiscardStencil)
