Use glCopyImageSubData() to copy textures when it's available.

This is primarily a workaround for glCopyTexSubImage2D() crashing on Galaxy S10/Adreno, but it's also simpler and more efficient. Also rewrote the fallback glCopyTexSubImage* path so that it copies the requested number of mips, not just the first one.

diff --git a/Engine/Source/Runtime/OpenGLDrv/Private/OpenGLTexture.cpp b/Engine/Source/Runtime/OpenGLDrv/Private/OpenGLTexture.cpp
--- a/Engine/Source/Runtime/OpenGLDrv/Private/OpenGLTexture.cpp
+++ b/Engine/Source/Runtime/OpenGLDrv/Private/OpenGLTexture.cpp
@@ -1871,7 +1871,8 @@ FShaderResourceViewRHIRef FOpenGLDynamicRHI::RHICreateShaderResourceView(FRHITex
 {
 	const uint32 MipLevel = CreateInfo.MipLevel;
 	const uint32 NumMipLevels = CreateInfo.NumMipLevels;
-	const uint8 Format = CreateInfo.Format;
+	EPixelFormat TextureBaseFormat = Texture->GetFormat();
+	const uint8 Format = (CreateInfo.Format == PF_Unknown) ? TextureBaseFormat : CreateInfo.Format;
 
 	FOpenGLShaderResourceViewProxy *ViewProxy = new FOpenGLShaderResourceViewProxy([this, Texture, MipLevel, NumMipLevels, Format](FRHIShaderResourceView* OwnerRHI) -> FOpenGLShaderResourceView*
 	{
@@ -2433,7 +2434,7 @@ void FOpenGLDynamicRHI::InvalidateTextureResourceInCache(GLuint Resource)
 
 void FOpenGLDynamicRHI::InvalidateUAVResourceInCache(GLuint Resource)
 {
-	for (int32 UAVIndex = 0; UAVIndex < OGL_MAX_COMPUTE_STAGE_UAV_UNITS; ++UAVIndex)
+	for (int32 UAVIndex = 0; UAVIndex < FOpenGL::GetMaxCombinedUAVUnits(); ++UAVIndex)
 	{
 		if (SharedContextState.UAVs[UAVIndex].Resource == Resource)
 		{
@@ -2579,13 +2580,86 @@ void FOpenGLDynamicRHI::RHICopyTexture(FRHITexture* SourceTextureRHI, FRHITextur
 	FOpenGLTextureBase* SourceTexture = GetOpenGLTextureFromRHITexture(SourceTextureRHI);
 	FOpenGLTextureBase* DestTexture = GetOpenGLTextureFromRHITexture(DestTextureRHI);
 
-	check(SourceTexture->Target == DestTexture->Target);
+	checkf(SourceTexture->Target == DestTexture->Target, TEXT("Cannot copy between different texture targets, SourceTexture Target=%x, Format=%d, Flags=%x; DestTexture Target=%x, Format=%d, Flags=%x"),
+		SourceTexture->Target, SourceTextureRHI->GetFormat(), SourceTextureRHI->GetFlags(),
+		DestTexture->Target, DestTextureRHI->GetFormat(), DestTextureRHI->GetFlags()
+	);
+	
+	checkf((SourceTextureRHI->GetFlags() & TexCreate_SRGB) == (DestTextureRHI->GetFlags() & TexCreate_SRGB), TEXT("Cannot copy between sRGB and linear, SourceTexture Format=%d, Flags=%x; DestTexture Format=%d, Flags=%x"),
+		SourceTextureRHI->GetFormat(), SourceTextureRHI->GetFlags(),
+		DestTextureRHI->GetFormat(), DestTextureRHI->GetFlags()
+	);
 
-	// Use a texture stage that's not likely to be used for draws, to avoid waiting
-	FOpenGLContextState& ContextState = GetContextStateForCurrentContext();
-	CachedSetupTextureStage(ContextState, FOpenGL::GetMaxCombinedTextureImageUnits() - 1, DestTexture->Target, DestTexture->Resource, 0, DestTextureRHI->GetNumMips());
-	CachedBindPixelUnpackBuffer(ContextState, 0);
+	GLsizei Width, Height, Depth;
+
+	if (CopyInfo.Size == FIntVector::ZeroValue)
+	{
+		// Copy whole texture when zero vector is specified for region size.
+		FIntVector SrcTexSize = SourceTextureRHI->GetSizeXYZ();
+		Width = FMath::Max(1, SrcTexSize.X >> CopyInfo.SourceMipIndex);
+		Height = FMath::Max(1, SrcTexSize.Y >> CopyInfo.SourceMipIndex);
+		switch (SourceTexture->Target)
+		{
+		case GL_TEXTURE_3D:			Depth = FMath::Max(1, SrcTexSize.Z >> CopyInfo.SourceMipIndex); break;
+		case GL_TEXTURE_CUBE_MAP:	Depth = 6; break;
+		default:					Depth = 1; break;
+		}
+		ensure(CopyInfo.SourcePosition == FIntVector::ZeroValue);
+	}
+	else
+	{
+		Width = CopyInfo.Size.X;
+		Height = CopyInfo.Size.Y;
+		Depth = (SourceTexture->Target == GL_TEXTURE_3D) ? CopyInfo.Size.Z : 1;
+	}
+
+	GLint SrcMip = CopyInfo.SourceMipIndex;
+	GLint DestMip = CopyInfo.DestMipIndex;
 
+	if (FOpenGL::SupportsCopyImage())
+	{
+		GLint SrcZOffset, DestZOffset;
+		switch (SourceTexture->Target)
+		{
+		case GL_TEXTURE_3D:
+		case GL_TEXTURE_CUBE_MAP:
+			// For cube maps, the Z offsets select the starting faces.
+			SrcZOffset = CopyInfo.SourcePosition.Z;
+			DestZOffset = CopyInfo.DestPosition.Z;
+			break;
+		case GL_TEXTURE_1D_ARRAY:
+		case GL_TEXTURE_2D_ARRAY:
+			// For texture arrays, the Z offsets and depth actually refer to the range of slices to copy.
+			SrcZOffset = CopyInfo.SourceSliceIndex;
+			DestZOffset = CopyInfo.DestSliceIndex;
+			Depth = CopyInfo.NumSlices;
+			break;
+		default:
+			SrcZOffset = 0;
+			DestZOffset = 0;
+			break;
+		}
+
+		for (uint32 MipIndex = 0; MipIndex < CopyInfo.NumMips; ++MipIndex)
+		{
+			FOpenGL::CopyImageSubData(SourceTexture->Resource, SourceTexture->Target, SrcMip, CopyInfo.SourcePosition.X, CopyInfo.SourcePosition.Y, SrcZOffset,
+				DestTexture->Resource, DestTexture->Target, DestMip, CopyInfo.DestPosition.X, CopyInfo.DestPosition.Y, DestZOffset,
+				Width, Height, Depth);
+
+			++SrcMip;
+			++DestMip;
+
+			Width = FMath::Max(1, Width >> 1);
+			Height = FMath::Max(1, Height >> 1);
+			if(DestTexture->Target == GL_TEXTURE_3D)
+			{
+				Depth = FMath::Max(1, Depth >> 1);
+			}
+		}
+
+		return;
+	}
+	
 	// Convert sub texture regions to GL types
 	GLint XOffset = CopyInfo.DestPosition.X;
 	GLint YOffset = CopyInfo.DestPosition.Y;
@@ -2593,46 +2667,54 @@ void FOpenGLDynamicRHI::RHICopyTexture(FRHITexture* SourceTextureRHI, FRHITextur
 	GLint X = CopyInfo.SourcePosition.X;
 	GLint Y = CopyInfo.SourcePosition.Y;
 	GLint Z = CopyInfo.SourcePosition.Z;
-	GLsizei Width = CopyInfo.Size.X;
-	GLsizei Height = CopyInfo.Size.Y;
-	GLsizei Depth = CopyInfo.Size.Z;
+
+	// Use a texture stage that's not likely to be used for draws, to avoid waiting
+	FOpenGLContextState& ContextState = GetContextStateForCurrentContext();
+	CachedSetupTextureStage(ContextState, FOpenGL::GetMaxCombinedTextureImageUnits() - 1, DestTexture->Target, DestTexture->Resource, 0, DestTextureRHI->GetNumMips());
+	CachedBindPixelUnpackBuffer(ContextState, 0);
 
 	// Bind source texture to an FBO to read from
-	for (GLsizei Layer = 0; Layer < Depth; ++Layer)
+	for (uint32 SliceIndex = 0; SliceIndex < CopyInfo.NumSlices; ++SliceIndex)
 	{
-		FOpenGLTextureBase* RenderTargets[1] = { SourceTexture };
-		uint32 MipLevels[1] = { CopyInfo.SourceMipIndex };
-		uint32 ArrayIndices[1] = { CopyInfo.SourceSliceIndex + static_cast<uint32>(Layer) };
+		for (uint32 MipIndex = 0; MipIndex < CopyInfo.NumMips; ++MipIndex)
+		{
+			FOpenGLTextureBase* RenderTargets[1] = { SourceTexture };
+			uint32 MipLevels[1] = { static_cast<uint32>(SrcMip) };
+			uint32 ArrayIndices[1] = { CopyInfo.SourceSliceIndex + SliceIndex };
 
-		GLuint SourceFBO = GetOpenGLFramebuffer(1, RenderTargets, ArrayIndices, MipLevels, nullptr);
-		check(SourceFBO != 0);
+			GLuint SourceFBO = GetOpenGLFramebuffer(1, RenderTargets, ArrayIndices, MipLevels, nullptr);
+			check(SourceFBO != 0);
 
-		glBindFramebuffer(GL_FRAMEBUFFER, SourceFBO);
+			glBindFramebuffer(GL_FRAMEBUFFER, SourceFBO);
 
-		FOpenGL::ReadBuffer(GL_COLOR_ATTACHMENT0);
+			FOpenGL::ReadBuffer(GL_COLOR_ATTACHMENT0);
 
-		switch (DestTexture->Target)
-		{
-		case GL_TEXTURE_1D:
-			FOpenGL::CopyTexSubImage1D(DestTexture->Target, CopyInfo.DestMipIndex, XOffset, X, 0, Width);
-			break;
-		case GL_TEXTURE_1D_ARRAY:
-		case GL_TEXTURE_2D:
-		case GL_TEXTURE_CUBE_MAP_POSITIVE_X:
-		case GL_TEXTURE_CUBE_MAP_NEGATIVE_X:
-		case GL_TEXTURE_CUBE_MAP_POSITIVE_Y:
-		case GL_TEXTURE_CUBE_MAP_NEGATIVE_Y:
-		case GL_TEXTURE_CUBE_MAP_POSITIVE_Z:
-		case GL_TEXTURE_CUBE_MAP_NEGATIVE_Z:
-		case GL_TEXTURE_RECTANGLE:
-			FOpenGL::CopyTexSubImage2D(DestTexture->Target, CopyInfo.DestMipIndex, XOffset, YOffset, X, Y, Width, Height);
-			break;
-		case GL_TEXTURE_3D:
-		case GL_TEXTURE_2D_ARRAY:
-		case GL_TEXTURE_CUBE_MAP_ARRAY:
-			FOpenGL::CopyTexSubImage3D(DestTexture->Target, CopyInfo.DestMipIndex, XOffset, YOffset, ZOffset + Layer, X, Y, Width, Depth);
-			break;
+			switch (DestTexture->Target)
+			{
+			case GL_TEXTURE_1D:
+				FOpenGL::CopyTexSubImage1D(DestTexture->Target, DestMip, XOffset, X, 0, Width);
+				break;
+			case GL_TEXTURE_1D_ARRAY:
+				FOpenGL::CopyTexSubImage2D(DestTexture->Target, DestMip, XOffset, CopyInfo.DestSliceIndex + SliceIndex, X, 0, Width, 1);
+				break;
+			case GL_TEXTURE_2D:
+			case GL_TEXTURE_RECTANGLE:
+				FOpenGL::CopyTexSubImage2D(DestTexture->Target, DestMip, XOffset, YOffset, X, Y, Width, Height);
+				break;
+			case GL_TEXTURE_2D_ARRAY:
+				FOpenGL::CopyTexSubImage3D(DestTexture->Target, DestMip, XOffset, YOffset, CopyInfo.DestSliceIndex + SliceIndex, X, Y, Width, Height);
+				break;
+			case GL_TEXTURE_3D:
+				FOpenGL::CopyTexSubImage3D(DestTexture->Target, DestMip, XOffset, YOffset, ZOffset, X, Y, Width, Height);
+				break;
+			}
 		}
+
+		++SrcMip;
+		++DestMip;
+
+		Width = FMath::Max(1, Width >> 1);
+		Height = FMath::Max(1, Height >> 1);
 	}
 
 	ContextState.Framebuffer = (GLuint)-1;
