Fixed build.

diff --git a/Engine/Source/Runtime/OpenGLDrv/Private/OpenGLTexture.cpp b/Engine/Source/Runtime/OpenGLDrv/Private/OpenGLTexture.cpp
--- a/Engine/Source/Runtime/OpenGLDrv/Private/OpenGLTexture.cpp
+++ b/Engine/Source/Runtime/OpenGLDrv/Private/OpenGLTexture.cpp
@@ -2434,7 +2434,7 @@ void FOpenGLDynamicRHI::InvalidateTextureResourceInCache(GLuint Resource)
 
 void FOpenGLDynamicRHI::InvalidateUAVResourceInCache(GLuint Resource)
 {
-	for (int32 UAVIndex = 0; UAVIndex < FOpenGL::GetMaxCombinedUAVUnits(); ++UAVIndex)
+	for (int32 UAVIndex = 0; UAVIndex < OGL_MAX_COMPUTE_STAGE_UAV_UNITS; ++UAVIndex)
 	{
 		if (SharedContextState.UAVs[UAVIndex].Resource == Resource)
 		{
