Support 3D immutable textures on Android GL, required for sky atmosphere

diff --git a/Engine/Source/Runtime/OpenGLDrv/Private/Android/AndroidOpenGL.cpp b/Engine/Source/Runtime/OpenGLDrv/Private/Android/AndroidOpenGL.cpp
--- a/Engine/Source/Runtime/OpenGLDrv/Private/Android/AndroidOpenGL.cpp
+++ b/Engine/Source/Runtime/OpenGLDrv/Private/Android/AndroidOpenGL.cpp
@@ -45,6 +45,7 @@ PFNGLMAPBUFFEROESPROC					glMapBufferOESa = NULL;
 PFNGLUNMAPBUFFEROESPROC					glUnmapBufferOESa = NULL;
 
 PFNGLTEXSTORAGE2DPROC					glTexStorage2D = NULL;
+PFNGLTEXSTORAGE3DPROC					glTexStorage3D = NULL;
 
 // KHR_debug
 PFNGLDEBUGMESSAGECONTROLKHRPROC			glDebugMessageControlKHR = NULL;
@@ -1081,6 +1082,8 @@ void FAndroidOpenGL::ProcessExtensions(const FString& ExtensionsString)
 		glBindSampler = (PFNGLBINDSAMPLERPROC)((void*)eglGetProcAddress("glBindSampler"));
 		glProgramParameteri = (PFNGLPROGRAMPARAMETERIPROC)((void*)eglGetProcAddress("glProgramParameteri"));
 
+		glTexStorage3D = (PFNGLTEXSTORAGE3DPROC)((void*)eglGetProcAddress("glTexStorage3D"));
+		
 		// Required by the ES3 spec
 		bSupportsInstancing = true;
 		bSupportsTextureFloat = true;
diff --git a/Engine/Source/Runtime/OpenGLDrv/Private/Android/AndroidOpenGL.h b/Engine/Source/Runtime/OpenGLDrv/Private/Android/AndroidOpenGL.h
--- a/Engine/Source/Runtime/OpenGLDrv/Private/Android/AndroidOpenGL.h
+++ b/Engine/Source/Runtime/OpenGLDrv/Private/Android/AndroidOpenGL.h
@@ -125,6 +125,7 @@ extern PFNGLLABELOBJECTEXTPROC			glLabelObjectEXT;
 extern PFNGLGETOBJECTLABELEXTPROC		glGetObjectLabelEXT;
 extern PFNGLPOPGROUPMARKEREXTPROC 		glPopGroupMarkerEXT;
 extern PFNGLTEXSTORAGE2DPROC			glTexStorage2D;
+extern PFNGLTEXSTORAGE3DPROC			glTexStorage3D;
 extern PFNGLDEBUGMESSAGECONTROLKHRPROC	glDebugMessageControlKHR;
 extern PFNGLDEBUGMESSAGEINSERTKHRPROC	glDebugMessageInsertKHR;
 extern PFNGLDEBUGMESSAGECALLBACKKHRPROC	glDebugMessageCallbackKHR;
@@ -410,23 +411,30 @@ struct FAndroidOpenGL : public FOpenGLES2
 	
 	static FORCEINLINE void TexStorage3D(GLenum Target, GLint Levels, GLint InternalFormat, GLsizei Width, GLsizei Height, GLsizei Depth, GLenum Format, GLenum Type)
 	{
-		const bool bArrayTexture = Target == GL_TEXTURE_2D_ARRAY || Target == GL_TEXTURE_CUBE_MAP_ARRAY;
-		for(uint32 MipIndex = 0; MipIndex < uint32(Levels); MipIndex++)
+		if (glTexStorage3D)
 		{
-			glTexImage3D(
-				Target,
-				MipIndex,
-				InternalFormat,
-				FMath::Max<uint32>(1,(Width >> MipIndex)),
-				FMath::Max<uint32>(1,(Height >> MipIndex)),
-				(bArrayTexture) ? Depth : FMath::Max<uint32>(1,(Depth >> MipIndex)),
-				0,
-				Format,
-				Type,
-				NULL
-				);
-
-			VERIFY_GL(TexImage_3D);
+			glTexStorage3D( Target, Levels, InternalFormat, Width, Height, Depth);
+		}
+		else
+		{
+			const bool bArrayTexture = Target == GL_TEXTURE_2D_ARRAY || Target == GL_TEXTURE_CUBE_MAP_ARRAY;
+			for(uint32 MipIndex = 0; MipIndex < uint32(Levels); MipIndex++)
+			{
+				glTexImage3D(
+					Target,
+					MipIndex,
+					InternalFormat,
+					FMath::Max<uint32>(1,(Width >> MipIndex)),
+					FMath::Max<uint32>(1,(Height >> MipIndex)),
+					(bArrayTexture) ? Depth : FMath::Max<uint32>(1,(Depth >> MipIndex)),
+					0,
+					Format,
+					Type,
+					NULL
+					);
+
+				VERIFY_GL(TexImage_3D);
+			}
 		}
 	}
 	
