Added missing emulated UB constants for compute on OpenGL

diff --git a/Engine/Source/Runtime/OpenGLDrv/Private/OpenGLCommands.cpp b/Engine/Source/Runtime/OpenGLDrv/Private/OpenGLCommands.cpp
--- a/Engine/Source/Runtime/OpenGLDrv/Private/OpenGLCommands.cpp
+++ b/Engine/Source/Runtime/OpenGLDrv/Private/OpenGLCommands.cpp
@@ -2425,9 +2425,13 @@ void FOpenGLDynamicRHI::CommitComputeShaderConstants(FOpenGLComputeShader* Compu
 	check(FOpenGL::SupportsComputeShaders());
 
 	const int32 Stage = CrossCompiler::SHADER_STAGE_COMPUTE;
+	FOpenGLShaderParameterCache& StageShaderParameters = PendingState.ShaderParameters[Stage];
 
-	FOpenGLShaderParameterCache& StageShaderParameters = PendingState.ShaderParameters[CrossCompiler::SHADER_STAGE_COMPUTE];
-	StageShaderParameters.CommitPackedGlobals(ComputeShader->LinkedProgram, CrossCompiler::SHADER_STAGE_COMPUTE);
+	if (GUseEmulatedUniformBuffers)
+	{
+		StageShaderParameters.CommitPackedUniformBuffers(ComputeShader->LinkedProgram, Stage, PendingState.BoundUniformBuffers[Stage], ComputeShader->UniformBuffersCopyInfo);
+	}
+	StageShaderParameters.CommitPackedGlobals(ComputeShader->LinkedProgram, Stage);
 	PendingState.LinkedProgramAndDirtyFlag = nullptr;
 }
 
