Fixed binding of 3D textures as UAVs in OpenGL RHI

diff --git a/Engine/Source/Runtime/OpenGLDrv/Private/OpenGLCommands.cpp b/Engine/Source/Runtime/OpenGLDrv/Private/OpenGLCommands.cpp
--- a/Engine/Source/Runtime/OpenGLDrv/Private/OpenGLCommands.cpp
+++ b/Engine/Source/Runtime/OpenGLDrv/Private/OpenGLCommands.cpp
@@ -448,11 +448,13 @@ void FOpenGLDynamicRHI::RHISetUAVParameter(FRHIComputeShader* ComputeShaderRHI,u
 	if(UnorderedAccessViewRHI)
 	{
 		FOpenGLUnorderedAccessView* UnorderedAccessView = ResourceCast(UnorderedAccessViewRHI);
-		InternalSetShaderUAV(FOpenGL::GetFirstComputeUAVUnit() + UAVIndex, UnorderedAccessView->Format , UnorderedAccessView->Resource);
+		bool bLayered = UnorderedAccessView->IsLayered();
+		GLint Layer = UnorderedAccessView->GetLayer();
+		InternalSetShaderUAV(FOpenGL::GetFirstComputeUAVUnit() + UAVIndex, UnorderedAccessView->Format, UnorderedAccessView->Resource, bLayered, Layer);
 	}
 	else
 	{
-		InternalSetShaderUAV(FOpenGL::GetFirstComputeUAVUnit() + UAVIndex, GL_R32F, 0);
+		InternalSetShaderUAV(FOpenGL::GetFirstComputeUAVUnit() + UAVIndex, GL_R32F, 0, false, 0);
 	}
 }
 
@@ -745,11 +747,13 @@ void FOpenGLDynamicRHI::SetupTexturesForDraw( FOpenGLContextState& ContextState
 	SetupTexturesForDraw(ContextState, PendingState.BoundShaderState, FOpenGL::GetMaxCombinedTextureImageUnits());
 }
 
-void FOpenGLDynamicRHI::InternalSetShaderUAV(GLint UAVIndex, GLenum Format, GLuint Resource)
+void FOpenGLDynamicRHI::InternalSetShaderUAV(GLint UAVIndex, GLenum Format, GLuint Resource, bool bLayered, GLint Layer)
 {
 	VERIFY_GL_SCOPE();
 	PendingState.UAVs[UAVIndex].Format = Format;
 	PendingState.UAVs[UAVIndex].Resource = Resource;
+	PendingState.UAVs[UAVIndex].Layer = Layer;
+	PendingState.UAVs[UAVIndex].bLayered = bLayered;
 }
 
 void FOpenGLDynamicRHI::SetupUAVsForDraw( FOpenGLContextState& ContextState, const FOpenGLComputeShader* ComputeShader, int32 MaxUAVsNeeded )
@@ -760,29 +764,35 @@ void FOpenGLDynamicRHI::SetupUAVsForDraw( FOpenGLContextState& ContextState, con
 	{
 		if (!ComputeShader->NeedsUAVStage(UAVStageIndex))
 		{
-			CachedSetupUAVStage(ContextState, UAVStageIndex, GL_R32F, 0 );
+			CachedSetupUAVStage(ContextState, UAVStageIndex, GL_R32F, 0, false, 0);
 		}
 		else
 		{
-			CachedSetupUAVStage(ContextState, UAVStageIndex, PendingState.UAVs[UAVStageIndex].Format, PendingState.UAVs[UAVStageIndex].Resource );
+			const FUAVStage& UAVStage = PendingState.UAVs[UAVStageIndex];
+			CachedSetupUAVStage(ContextState, UAVStageIndex, UAVStage.Format, UAVStage.Resource, UAVStage.bLayered, UAVStage.Layer);
 		}
 	}
 	
 }
 
 
-void FOpenGLDynamicRHI::CachedSetupUAVStage( FOpenGLContextState& ContextState, GLint UAVIndex, GLenum Format, GLuint Resource)
+void FOpenGLDynamicRHI::CachedSetupUAVStage( FOpenGLContextState& ContextState, GLint UAVIndex, GLenum Format, GLuint Resource, bool bLayered, GLint Layer)
 {
-	if( ContextState.UAVs[UAVIndex].Format == Format && ContextState.UAVs[UAVIndex].Resource == Resource)
+	if (ContextState.UAVs[UAVIndex].Format == Format && 
+		ContextState.UAVs[UAVIndex].Resource == Resource &&
+		ContextState.UAVs[UAVIndex].Layer == Layer &&
+		ContextState.UAVs[UAVIndex].bLayered == bLayered)
 	{
 		// Nothing's changed, no need to update
 		return;
 	}
 
-	FOpenGL::BindImageTexture(UAVIndex, Resource, 0, GL_FALSE, 0, GL_READ_WRITE, Format);
+	FOpenGL::BindImageTexture(UAVIndex, Resource, 0, bLayered ? GL_TRUE : GL_FALSE, Layer, GL_READ_WRITE, Format);
 	
 	ContextState.UAVs[UAVIndex].Format = Format;
 	ContextState.UAVs[UAVIndex].Resource = Resource;
+	ContextState.UAVs[UAVIndex].Layer = Layer;
+	ContextState.UAVs[UAVIndex].bLayered = bLayered;
 }
 
 void FOpenGLDynamicRHI::UpdateSRV(FOpenGLShaderResourceView* SRV)
diff --git a/Engine/Source/Runtime/OpenGLDrv/Private/OpenGLUAV.cpp b/Engine/Source/Runtime/OpenGLDrv/Private/OpenGLUAV.cpp
--- a/Engine/Source/Runtime/OpenGLDrv/Private/OpenGLUAV.cpp
+++ b/Engine/Source/Runtime/OpenGLDrv/Private/OpenGLUAV.cpp
@@ -138,6 +138,7 @@ FOpenGLTextureUnorderedAccessView::FOpenGLTextureUnorderedAccessView(FRHITexture
 
 	this->Resource = Texture->Resource;
 	this->Format = GLFormat.InternalFormat[0];
+	this->bLayered = (Texture->Target == GL_TEXTURE_3D);
 }
 
 
diff --git a/Engine/Source/Runtime/OpenGLDrv/Public/OpenGLDrv.h b/Engine/Source/Runtime/OpenGLDrv/Public/OpenGLDrv.h
--- a/Engine/Source/Runtime/OpenGLDrv/Public/OpenGLDrv.h
+++ b/Engine/Source/Runtime/OpenGLDrv/Public/OpenGLDrv.h
@@ -564,7 +564,7 @@ public:
 	}
 
 	void CachedSetupTextureStageInner(FOpenGLContextState& ContextState, GLint TextureIndex, GLenum Target, GLuint Resource, GLint BaseMip, GLint NumMips);
-	void CachedSetupUAVStage(FOpenGLContextState& ContextState, GLint UAVIndex, GLenum Format, GLuint Resource);
+	void CachedSetupUAVStage(FOpenGLContextState& ContextState, GLint UAVIndex, GLenum Format, GLuint Resource, bool bLayered, GLint Layer);
 	void UpdateSRV(FOpenGLShaderResourceView* SRV);
 	FOpenGLContextState& GetContextStateForCurrentContext(bool bAssertIfInvalid = true);
 
@@ -1143,7 +1143,7 @@ private:
 public:
 	/** Remember what RHI user wants set on a specific OpenGL texture stage, translating from Stage and TextureIndex for stage pair. */
 	void InternalSetShaderTexture(FOpenGLTextureBase* Texture, FOpenGLShaderResourceView* SRV, GLint TextureIndex, GLenum Target, GLuint Resource, int NumMips, int LimitMip);
-	void InternalSetShaderUAV(GLint UAVIndex, GLenum Format, GLuint Resource);
+	void InternalSetShaderUAV(GLint UAVIndex, GLenum Format, GLuint Resource, bool bLayered, GLint Layer);
 	void InternalSetSamplerStates(GLint TextureIndex, FOpenGLSamplerState* SamplerState);
 
 private:
diff --git a/Engine/Source/Runtime/OpenGLDrv/Public/OpenGLResources.h b/Engine/Source/Runtime/OpenGLDrv/Public/OpenGLResources.h
--- a/Engine/Source/Runtime/OpenGLDrv/Public/OpenGLResources.h
+++ b/Engine/Source/Runtime/OpenGLDrv/Public/OpenGLResources.h
@@ -1812,6 +1812,17 @@ public:
 	{
 		return 0;
 	}
+
+	virtual bool IsLayered() const
+	{
+		return false;
+	}
+
+	virtual GLint GetLayer() const
+	{
+		return 0;
+	}
+
 };
 
 class FOpenGLTextureUnorderedAccessView : public FOpenGLUnorderedAccessView
@@ -1821,6 +1832,12 @@ public:
 	FOpenGLTextureUnorderedAccessView(FRHITexture* InTexture);
 
 	FTextureRHIRef TextureRHI; // to keep the texture alive
+	bool bLayered;
+
+	virtual bool IsLayered() const override
+	{
+		return bLayered;
+	}
 };
 
 
diff --git a/Engine/Source/Runtime/OpenGLDrv/Public/OpenGLState.h b/Engine/Source/Runtime/OpenGLDrv/Public/OpenGLState.h
--- a/Engine/Source/Runtime/OpenGLDrv/Public/OpenGLState.h
+++ b/Engine/Source/Runtime/OpenGLDrv/Public/OpenGLState.h
@@ -220,10 +220,14 @@ struct FUAVStage
 {
 	GLenum Format;
 	GLuint Resource;
+	GLint Layer;
+	bool bLayered;
 	
 	FUAVStage()
 	:	Format(GL_NONE)
 	,	Resource(0)
+	,	Layer(0)
+	,	bLayered(false)
 	{
 	}
 };
