Changed FGlobalDynamicReadBuffer buffer allocation so that they are BUF_Volatile.
Optimzed OpenGL vertex buffer lock / unlock in that usecase.

diff --git a/Engine/Source/Runtime/OpenGLDrv/Public/OpenGLResources.h b/Engine/Source/Runtime/OpenGLDrv/Public/OpenGLResources.h
--- a/Engine/Source/Runtime/OpenGLDrv/Public/OpenGLResources.h
+++ b/Engine/Source/Runtime/OpenGLDrv/Public/OpenGLResources.h
@@ -519,27 +519,29 @@ public:
 		{
 			auto DeleteGLResources = [Resource=Resource, RealSize= RealSize, bStreamDraw= (bool)bStreamDraw, LockBuffer = LockBuffer, bLockBufferWasAllocated=bLockBufferWasAllocated]()
 			{
-		VERIFY_GL_SCOPE();
+				VERIFY_GL_SCOPE();
 				if (BaseType::OnDelete(Resource, RealSize, bStreamDraw, 0))
-		{
-			FOpenGL::DeleteBuffers(1, &Resource);
-		}
-		if (LockBuffer != NULL)
-		{
-			if (bLockBufferWasAllocated)
-			{
-				FMemory::Free(LockBuffer);
-			}
-			else
-			{
-				UE_LOG(LogRHI,Warning,TEXT("Destroying TOpenGLBuffer without returning memory to the driver; possibly called RHIMapStagingSurface() but didn't call RHIUnmapStagingSurface()? Resource %u"), Resource);
-			}
-		}
+				{
+					FOpenGL::DeleteBuffers(1, &Resource);
+				}
+				if (LockBuffer != NULL)
+				{
+					if (bLockBufferWasAllocated)
+					{
+						FMemory::Free(LockBuffer);
+					}
+					else
+					{
+						UE_LOG(LogRHI,Warning,TEXT("Destroying TOpenGLBuffer without returning memory to the driver; possibly called RHIMapStagingSurface() but didn't call RHIUnmapStagingSurface()? Resource %u"), Resource);
+					}
+				}
 			};
 
 			RunOnGLRenderContextThread(MoveTemp(DeleteGLResources));
 			LockBuffer = nullptr;
 			DecrementBufferMemory(Type, BaseType::IsStructuredBuffer(), RealSize);
+
+			ReleaseCachedBuffer();
 		}
 
 	}
@@ -632,7 +634,18 @@ public:
 			// Allocate a temp buffer to write into
 			LockOffset = InOffset;
 			LockSize = InSize;
-			LockBuffer = FMemory::Malloc( InSize );
+			if (CachedBuffer && InSize <= CachedBufferSize)
+			{
+				LockBuffer = CachedBuffer;
+				CachedBuffer = nullptr;
+				// Keep CachedBufferSize to keep the actual size allocated.
+			}
+			else
+			{
+				ReleaseCachedBuffer();
+				LockBuffer = FMemory::Malloc( InSize );
+				CachedBufferSize = InSize; // Safegard
+			}
 			Data = static_cast<uint8*>( LockBuffer );
 			bLockBufferWasAllocated = true;
 		}
@@ -697,7 +710,18 @@ public:
 			// Allocate a temp buffer to write into
 			LockOffset = InOffset;
 			LockSize = InSize;
-			LockBuffer = FMemory::Malloc( InSize );
+			if (CachedBuffer && InSize <= CachedBufferSize)
+			{
+				LockBuffer = CachedBuffer;
+				CachedBuffer = nullptr;
+				// Keep CachedBufferSize to keep the actual size allocated.
+			}
+			else
+			{
+				ReleaseCachedBuffer();
+				LockBuffer = FMemory::Malloc( InSize );
+				CachedBufferSize = InSize; // Safegard
+			}
 			Data = static_cast<uint8*>( LockBuffer );
 			bLockBufferWasAllocated = true;
 		}
@@ -758,9 +782,22 @@ public:
 #endif
 				}
 				check(bLockBufferWasAllocated);
-				FMemory::Free(LockBuffer);
+
+				if ((this->GetUsage() & BUF_Volatile) != 0)
+				{
+					ReleaseCachedBuffer(); // Safegard
+
+					CachedBuffer = LockBuffer;
+					// Possibly > LockSize when reusing cached allocation.
+					CachedBufferSize = FMath::Max<GLuint>(CachedBufferSize, LockSize);
+				}
+				else
+				{
+					FMemory::Free(LockBuffer);
+				}
 				LockBuffer = NULL;
 				bLockBufferWasAllocated = false;
+				LockSize = 0;
 			}
 			ModificationCount += (bIsLockReadOnly ? 0 : 1);
 			bIsLocked = false;
@@ -786,6 +823,15 @@ public:
 	bool IsLockReadOnly() const { return bIsLockReadOnly; }
 	void* GetLockedBuffer() const { return LockBuffer; }
 
+	void ReleaseCachedBuffer()
+	{
+		if (CachedBuffer)
+		{
+			FMemory::Free(CachedBuffer);
+			CachedBufferSize = 0;
+		}
+		// Don't reset CachedBufferSize if !CachedBuffer since it could be the locked buffer allocation size.
+	}
 private:
 
 	uint32 bIsLocked : 1;
@@ -797,6 +843,11 @@ private:
 	GLuint LockOffset;
 	void* LockBuffer;
 
+	// A cached allocation that can be reused. The same allocation can never be in CachedBuffer and LockBuffer at the same time.
+	void* CachedBuffer = nullptr;
+	// The size of the cached buffer allocation. Can be non zero even though CachedBuffer is  null, to preserve the allocation size.
+	GLuint CachedBufferSize = 0;
+
 	uint32 RealSize;	// sometimes (for example, for uniform buffer pool) we allocate more in OpenGL than is requested of us.
 
 	FOpenGLAssertRHIThreadFence TransitionFence;
diff --git a/Engine/Source/Runtime/RenderCore/Private/DynamicBufferAllocator.cpp b/Engine/Source/Runtime/RenderCore/Private/DynamicBufferAllocator.cpp
--- a/Engine/Source/Runtime/RenderCore/Private/DynamicBufferAllocator.cpp
+++ b/Engine/Source/Runtime/RenderCore/Private/DynamicBufferAllocator.cpp
@@ -130,7 +130,7 @@ FGlobalDynamicReadBuffer::FAllocation FGlobalDynamicReadBuffer::AllocateFloat(ui
 			const uint32 NewBufferSize = FMath::Max(AlignedNum, (uint32)GMinReadBufferRenderingBufferSize);
 			Buffer = new FDynamicAllocReadBuffer();
 			FloatBufferPool->Buffers.Add(Buffer);
-			Buffer->Initialize(sizeof(float), NewBufferSize, PF_R32_FLOAT, BUF_Dynamic);
+			Buffer->Initialize(sizeof(float), NewBufferSize, PF_R32_FLOAT, BUF_Volatile);
 		}
 
 		// Lock the buffer if needed.
@@ -186,7 +186,7 @@ FGlobalDynamicReadBuffer::FAllocation FGlobalDynamicReadBuffer::AllocateInt32(ui
 			const uint32 NewBufferSize = FMath::Max(AlignedNum, (uint32)GMinReadBufferRenderingBufferSize);
 			Buffer = new FDynamicAllocReadBuffer();
 			Int32BufferPool->Buffers.Add(Buffer);
-			Buffer->Initialize(sizeof(int32), NewBufferSize, PF_R32_SINT, BUF_Dynamic);
+			Buffer->Initialize(sizeof(int32), NewBufferSize, PF_R32_SINT, BUF_Volatile);
 		}
 
 		// Lock the buffer if needed.
