Fixed possible double delete of vertexbuffer Cached allocation

diff --git a/Engine/Source/Runtime/OpenGLDrv/Public/OpenGLResources.h b/Engine/Source/Runtime/OpenGLDrv/Public/OpenGLResources.h
--- a/Engine/Source/Runtime/OpenGLDrv/Public/OpenGLResources.h
+++ b/Engine/Source/Runtime/OpenGLDrv/Public/OpenGLResources.h
@@ -828,6 +828,7 @@ public:
 		if (CachedBuffer)
 		{
 			FMemory::Free(CachedBuffer);
+			CachedBuffer = nullptr;
 			CachedBufferSize = 0;
 		}
 		// Don't reset CachedBufferSize if !CachedBuffer since it could be the locked buffer allocation size.
