Fixes for loading compute shaders when binary cache and LRU are active.

diff --git a/Engine/Source/Runtime/OpenGLDrv/Private/OpenGLCommands.cpp b/Engine/Source/Runtime/OpenGLDrv/Private/OpenGLCommands.cpp
--- a/Engine/Source/Runtime/OpenGLDrv/Private/OpenGLCommands.cpp
+++ b/Engine/Source/Runtime/OpenGLDrv/Private/OpenGLCommands.cpp
@@ -3229,6 +3229,10 @@ void FOpenGLDynamicRHI::RHIDispatchComputeShader(uint32 ThreadGroupCountX, uint3
 
 		FOpenGLComputeShader* ComputeShader = ResourceCast(ComputeShaderRHI);
 
+		if (ComputeShader->LinkedProgram == nullptr)
+		{
+			ComputeShader->LinkedProgram = GetLinkedComputeProgram(ComputeShaderRHI);
+		}
 		FOpenGLContextState& ContextState = GetContextStateForCurrentContext();
 
 		GPUProfilingData.RegisterGPUDispatch(FIntVector(ThreadGroupCountX, ThreadGroupCountY, ThreadGroupCountZ));	
@@ -3259,6 +3263,11 @@ void FOpenGLDynamicRHI::RHIDispatchIndirectComputeShader(FRHIVertexBuffer* Argum
 		check(ComputeShaderRHI);
 
 		FOpenGLComputeShader* ComputeShader = ResourceCast(ComputeShaderRHI);
+		if (ComputeShader->LinkedProgram == nullptr)
+		{
+			ComputeShader->LinkedProgram = GetLinkedComputeProgram(ComputeShaderRHI);
+		}
+
 		FOpenGLVertexBuffer* ArgumentBuffer = ResourceCast(ArgumentBufferRHI);
 
 		FOpenGLContextState& ContextState = GetContextStateForCurrentContext();
diff --git a/Engine/Source/Runtime/OpenGLDrv/Private/OpenGLShaders.cpp b/Engine/Source/Runtime/OpenGLDrv/Private/OpenGLShaders.cpp
--- a/Engine/Source/Runtime/OpenGLDrv/Private/OpenGLShaders.cpp
+++ b/Engine/Source/Runtime/OpenGLDrv/Private/OpenGLShaders.cpp
@@ -3196,10 +3196,11 @@ static FOpenGLLinkedProgram* LinkProgram( const FOpenGLLinkedProgramConfiguratio
 	return LinkedProgram;
 }
 
-static bool LinkComputeShader(FOpenGLComputeShader* ComputeShader)
+static bool LinkComputeShader(FRHIComputeShader* ComputeShaderRHI, FOpenGLComputeShader* ComputeShader)
 {
 	check(ComputeShader);
 	check(ComputeShader->Resource != 0);
+	check(ComputeShaderRHI->GetHash() != FSHAHash());
 
 	const ANSICHAR* GlslCode = NULL;
 	if (!ComputeShader->bSuccessfullyCompiled)
@@ -3214,77 +3215,64 @@ static bool LinkComputeShader(FOpenGLComputeShader* ComputeShader)
 
 	Config.Shaders[CrossCompiler::SHADER_STAGE_COMPUTE].Resource = ComputeShader->Resource;
 	Config.Shaders[CrossCompiler::SHADER_STAGE_COMPUTE].Bindings = ComputeShader->Bindings;
-	Config.ProgramKey.ShaderHashes[CrossCompiler::SHADER_STAGE_COMPUTE] = ComputeShader->GetHash();
+	Config.ProgramKey.ShaderHashes[CrossCompiler::SHADER_STAGE_COMPUTE] = ComputeShaderRHI->GetHash();
 
-	ComputeShader->LinkedProgram = LinkProgram( Config, false );
+	ComputeShader->LinkedProgram = GetOpenGLProgramsCache().Find(Config, true);
 
-	if (ComputeShader->LinkedProgram == NULL)
+	if (ComputeShader->LinkedProgram == nullptr)
 	{
-#if DEBUG_GL_SHADERS
-		if (ComputeShader->bSuccessfullyCompiled)
+		ComputeShader->LinkedProgram = LinkProgram(Config, false);
+		if(ComputeShader->LinkedProgram == nullptr)
 		{
-			UE_LOG(LogRHI,Error,TEXT("Compute Shader:\n%s"),ANSI_TO_TCHAR(ComputeShader->GlslCode.GetData()));
+	#if DEBUG_GL_SHADERS
+			if (ComputeShader->bSuccessfullyCompiled)
+			{
+				UE_LOG(LogRHI, Error, TEXT("Compute Shader:\n%s"), ANSI_TO_TCHAR(ComputeShader->GlslCode.GetData()));
+			}
+	#endif //DEBUG_GL_SHADERS
+			checkf(ComputeShader->LinkedProgram, TEXT("Compute shader failed to compile & link."));
+
+			FName LinkFailurePanic = FName("FailedComputeProgramLink");
+			RHIGetPanicDelegate().ExecuteIfBound(LinkFailurePanic);
+			UE_LOG(LogRHI, Fatal, TEXT("Failed to link compute program [%s]. Current total programs: %d"), *Config.ProgramKey.ToString(), GNumPrograms);
+			return false;
 		}
-#endif //DEBUG_GL_SHADERS
-		checkf(ComputeShader->LinkedProgram, TEXT("Compute shader failed to compile & link."));
-		return false;
+		GetOpenGLProgramsCache().Add(Config.ProgramKey, ComputeShader->LinkedProgram);
 	}
 
 	return true;
 }
 
-//
-// specialization for compute
-//
-template<>
-FRHIComputeShader* CreateProxyShader<FRHIComputeShader, FOpenGLComputeShaderProxy>(const TArray<uint8>& Code)
+FOpenGLLinkedProgram* FOpenGLDynamicRHI::GetLinkedComputeProgram(FRHIComputeShader* ComputeShaderRHI)
 {
-	FRHICommandListImmediate& RHICmdList = FRHICommandListExecutor::GetImmediateCommandList();
-	if (ShouldRunGLRenderContextOpOnThisThread(RHICmdList))
-	{
-		return new FOpenGLComputeShaderProxy([&](FRHIComputeShader* OwnerRHI)
-		{
-			FOpenGLComputeShader* ComputeShader = CompileOpenGLShader<FOpenGLComputeShader>(Code, FSHAHash(), OwnerRHI);
-			LinkComputeShader(ComputeShader);
-			return ComputeShader;
-		});
-	}
-	else
-	{
-		// take a copy of the code for RHIT version.
-		TArray<uint8> CodeCopy = Code;
-		return new FOpenGLComputeShaderProxy([Code = MoveTemp(CodeCopy)](FRHIComputeShader* OwnerRHI)
-		{
-			FOpenGLComputeShader* ComputeShader = CompileOpenGLShader<FOpenGLComputeShader>(Code, FSHAHash(), OwnerRHI);
-			LinkComputeShader(ComputeShader);
-			return ComputeShader;
-		});
-	}
-}
+	VERIFY_GL_SCOPE();
+	check(ComputeShaderRHI->GetHash() != FSHAHash());
+	FOpenGLComputeShader* ComputeShader = ResourceCast(ComputeShaderRHI);
 
-template<>
-FRHIComputeShader* CreateProxyShader<FRHIComputeShader, FOpenGLComputeShaderProxy>(FRHIShaderLibrary* Library, FSHAHash Hash)
-{
-	FRHICommandListImmediate& RHICmdList = FRHICommandListExecutor::GetImmediateCommandList();
-	if (ShouldRunGLRenderContextOpOnThisThread(RHICmdList))
+	FOpenGLLinkedProgramConfiguration Config;
+
+	Config.Shaders[CrossCompiler::SHADER_STAGE_COMPUTE].Resource = ComputeShader->Resource;
+	Config.Shaders[CrossCompiler::SHADER_STAGE_COMPUTE].Bindings = ComputeShader->Bindings;
+	Config.ProgramKey.ShaderHashes[CrossCompiler::SHADER_STAGE_COMPUTE] = ComputeShaderRHI->GetHash();
+
+	FOpenGLLinkedProgram* LinkedProgram = GetOpenGLProgramsCache().Find(Config, true);
+	if (LinkedProgram == nullptr)
 	{
-		return new FOpenGLComputeShaderProxy([&](FRHIComputeShader* OwnerRHI)
-		{
-			FOpenGLComputeShader* ComputeShader = CompileOpenGLShader<FOpenGLComputeShader>(Library, Hash, OwnerRHI);
-			LinkComputeShader(ComputeShader);
-			return ComputeShader;
-		});
+		// Not in the cache. Create and add the program here.
+		// We can now link the compute shader, by now the shader hash has been set.
+		LinkComputeShader(ComputeShaderRHI, ComputeShader);
+		check(ComputeShader->LinkedProgram);
+		LinkedProgram = ComputeShader->LinkedProgram;
 	}
-	else
+	else if (!LinkedProgram->bConfigIsInitalized)
 	{
-		// take a copy of the code for RHIT version.
-		return new FOpenGLComputeShaderProxy([Library, Hash](FRHIComputeShader* OwnerRHI)
-		{
-			FOpenGLComputeShader* ComputeShader = CompileOpenGLShader<FOpenGLComputeShader>(Library, Hash, OwnerRHI);
-			LinkComputeShader(ComputeShader);
-			return ComputeShader;
-		});
+		// this has been loaded via binary program cache, properly initialize it here:
+		LinkedProgram->SetConfig(Config);
+		// We now have the config for this program, we must configure the program for use.
+		ConfigureGLProgramStageStates(LinkedProgram);
 	}
+	check(LinkedProgram->bConfigIsInitalized);
+	return LinkedProgram;
 }
 
 FComputeShaderRHIRef FOpenGLDynamicRHI::RHICreateComputeShader(FRHIShaderLibrary* Library, FSHAHash Hash)
@@ -4147,6 +4135,8 @@ void FOpenGLDynamicRHI::BindPendingComputeShaderState(FOpenGLContextState& Conte
 	VERIFY_GL_SCOPE();
 	bool ForceUniformBindingUpdate = false;
 
+	GetOpenGLProgramsCache().Touch(ComputeShader->LinkedProgram);
+
 	GLuint PendingProgram = ComputeShader->LinkedProgram->Program;
 	if (ContextState.Program != PendingProgram)
 	{
@@ -4769,7 +4759,7 @@ void FOpenGLProgramBinaryCache::ScanProgramCacheFile(const FGuid& ShaderPipeline
 
 							// check to see if any of the shaders are already loaded and so we should serialize the binary
 							bool bAllShadersLoaded = true;
-							for (int32 i = 0; i < CrossCompiler::NUM_NON_COMPUTE_SHADER_STAGES && bAllShadersLoaded; i++)
+							for (int32 i = 0; i < CrossCompiler::NUM_SHADER_STAGES && bAllShadersLoaded; i++)
 							{
 								bAllShadersLoaded = ProgramKey.ShaderHashes[i] == FSHAHash() || ShaderIsLoaded(ProgramKey.ShaderHashes[i]);
 							}
@@ -5134,7 +5124,7 @@ bool FOpenGLProgramBinaryCache::UseCachedProgram_internal(GLuint& ProgramOUT, co
 		// by this point the program must be either available or no attempt to load from shader library has occurred.
 		checkf(FoundProgram->GLProgramState == FGLProgramBinaryFileCacheEntry::EGLProgramState::ProgramStored
 			|| FoundProgram->GLProgramState == FGLProgramBinaryFileCacheEntry::EGLProgramState::ProgramAvailable,
-			TEXT("Unexpected program state: %d"), (int32)FoundProgram->GLProgramState);
+			TEXT("Unexpected program state:  (%s) == %d"), *ProgramKey.ToString(), (int32)FoundProgram->GLProgramState);
 
 		if (FoundProgram->GLProgramState == FGLProgramBinaryFileCacheEntry::EGLProgramState::ProgramAvailable)
 		{
diff --git a/Engine/Source/Runtime/OpenGLDrv/Public/OpenGLDrv.h b/Engine/Source/Runtime/OpenGLDrv/Public/OpenGLDrv.h
--- a/Engine/Source/Runtime/OpenGLDrv/Public/OpenGLDrv.h
+++ b/Engine/Source/Runtime/OpenGLDrv/Public/OpenGLDrv.h
@@ -1032,6 +1032,7 @@ public:
 			false);
 	}
 
+	class FOpenGLLinkedProgram* GetLinkedComputeProgram(FRHIComputeShader* ComputeShaderRHI);
 
 	FBoundShaderStateRHIRef RHICreateBoundShaderState_OnThisThread(FRHIVertexDeclaration* VertexDeclaration, FRHIVertexShader* VertexShader, FRHIHullShader* HullShader, FRHIDomainShader* DomainShader, FRHIPixelShader* PixelShader, FRHIGeometryShader* GeometryShader, bool FromPSOFileCache);
 	void RHIPerFrameRHIFlushComplete();
diff --git a/Engine/Source/Runtime/OpenGLDrv/Public/OpenGLShaderResources.h b/Engine/Source/Runtime/OpenGLDrv/Public/OpenGLShaderResources.h
--- a/Engine/Source/Runtime/OpenGLDrv/Public/OpenGLShaderResources.h
+++ b/Engine/Source/Runtime/OpenGLDrv/Public/OpenGLShaderResources.h
@@ -313,7 +313,7 @@ typedef TOpenGLShader<FRefCountedObject, GL_TESS_CONTROL_SHADER, SF_Hull> FOpenG
 typedef TOpenGLShader<FRefCountedObject, GL_TESS_EVALUATION_SHADER, SF_Domain> FOpenGLDomainShader;
 
 
-class FOpenGLComputeShader : public TOpenGLShader<FRHIComputeShader, GL_COMPUTE_SHADER, SF_Compute>
+class FOpenGLComputeShader : public TOpenGLShader<FRefCountedObject, GL_COMPUTE_SHADER, SF_Compute>
 {
 public:
 	FOpenGLComputeShader():
@@ -424,9 +424,22 @@ struct FOpenGLProgramKey
 	FString ToString() const
 	{
 		FString retme;
-		retme = TEXT("Program V_") + ShaderHashes[CrossCompiler::SHADER_STAGE_VERTEX].ToString();
-		retme += TEXT("_P_") + ShaderHashes[CrossCompiler::SHADER_STAGE_PIXEL].ToString();
-		return retme;
+		if(ShaderHashes[CrossCompiler::SHADER_STAGE_VERTEX] != FSHAHash())
+		{
+			retme = TEXT("Program V_") + ShaderHashes[CrossCompiler::SHADER_STAGE_VERTEX].ToString();
+			retme += TEXT("_P_") + ShaderHashes[CrossCompiler::SHADER_STAGE_PIXEL].ToString();
+			return retme;
+		}
+		else if(ShaderHashes[CrossCompiler::SHADER_STAGE_COMPUTE] != FSHAHash())
+		{
+			retme = TEXT("Program C_") + ShaderHashes[CrossCompiler::SHADER_STAGE_COMPUTE].ToString();
+			return retme;
+		}
+		else
+		{
+			retme = TEXT("Program with unset key");
+			return retme;
+		}
 	}
 
 	FSHAHash ShaderHashes[CrossCompiler::NUM_SHADER_STAGES];
