Add R8G8 to GLES3.1's list of supported textures.

diff --git a/Engine/Source/Runtime/OpenGLDrv/Private/OpenGLDevice.cpp b/Engine/Source/Runtime/OpenGLDrv/Private/OpenGLDevice.cpp
--- a/Engine/Source/Runtime/OpenGLDrv/Private/OpenGLDevice.cpp
+++ b/Engine/Source/Runtime/OpenGLDrv/Private/OpenGLDevice.cpp
@@ -1041,6 +1041,7 @@ static void InitRHICapabilitiesForGL()
 			SetupTextureFormat(PF_R16_UINT, FOpenGLTextureFormat(GL_R16UI, GL_R16UI, GL_RED_INTEGER, GL_UNSIGNED_SHORT, false, false));
 			SetupTextureFormat(PF_R16_SINT, FOpenGLTextureFormat(GL_R16I, GL_R16I, GL_RED_INTEGER, GL_SHORT, false, false));
 			SetupTextureFormat(PF_R8_UINT, FOpenGLTextureFormat(GL_R8UI, GL_R8UI, GL_RED_INTEGER, GL_UNSIGNED_BYTE, false, false));
+			SetupTextureFormat(PF_R8G8, FOpenGLTextureFormat(GL_RG8, GL_RG8, GL_RG, GL_UNSIGNED_BYTE, false, false));
 		}
 #endif
 
