Fixed: GL sampler state leak from compute shaders

diff --git a/Engine/Source/Runtime/OpenGLDrv/Private/OpenGLCommands.cpp b/Engine/Source/Runtime/OpenGLDrv/Private/OpenGLCommands.cpp
--- a/Engine/Source/Runtime/OpenGLDrv/Private/OpenGLCommands.cpp
+++ b/Engine/Source/Runtime/OpenGLDrv/Private/OpenGLCommands.cpp
@@ -1175,14 +1175,7 @@ void FOpenGLDynamicRHI::RHISetShaderSampler(FRHIComputeShader* ComputeShaderRHI,
 	check(FOpenGL::SupportsComputeShaders());
 	VERIFY_GL_SCOPE();
 	FOpenGLSamplerState* NewState = ResourceCast(NewStateRHI);
-	if ( OpenGLConsoleVariables::bBindlessTexture == 0 || !FOpenGL::SupportsBindlessTexture())
-	{
-		FOpenGL::BindSampler(FOpenGL::GetFirstComputeTextureUnit() + SamplerIndex, NewState->Resource);
-	}
-	else
-	{
-		PendingState.SamplerStates[FOpenGL::GetFirstComputeTextureUnit() + SamplerIndex] = NewState;
-	}
+	InternalSetSamplerStates(FOpenGL::GetFirstComputeTextureUnit() + SamplerIndex, NewState);
 }
 
 void FOpenGLDynamicRHI::RHISetShaderUniformBuffer(FRHIPixelShader* PixelShaderRHI,uint32 BufferIndex, FRHIUniformBuffer* BufferRHI)
