DR - Remove Draw{Indexed}PrimitiveUP from the engine

diff --git a/Engine/Source/Runtime/Apple/MetalRHI/Private/MetalCommands.cpp b/Engine/Source/Runtime/Apple/MetalRHI/Private/MetalCommands.cpp
--- a/Engine/Source/Runtime/Apple/MetalRHI/Private/MetalCommands.cpp
+++ b/Engine/Source/Runtime/Apple/MetalRHI/Private/MetalCommands.cpp
@@ -829,106 +829,6 @@ void FMetalRHICommandContext::RHIDrawIndexedPrimitiveIndirect(FIndexBufferRHIPar
 	}
 }
 
-
-void FMetalRHICommandContext::RHIBeginDrawPrimitiveUP( uint32 NumPrimitives, uint32 NumVertices, uint32 VertexDataStride, void*& OutVertexData)
-{
-	@autoreleasepool {
-	SCOPE_CYCLE_COUNTER(STAT_MetalDrawCallTime);
-	checkSlow(PendingVertexBuffer.GetPtr() == nil);
-
-	uint32 PrimitiveType = Context->GetCurrentState().GetPrimitiveType();
-	
-	
-	// allocate space
-	PendingVertexBuffer = Context->AllocateFromRingBuffer(VertexDataStride * NumVertices);
-	PendingVertexBuffer.MarkSingleUse();
-
-	// get the pointer to send back for writing
-	OutVertexData = (uint8*)PendingVertexBuffer.GetContents();
-	
-	PendingPrimitiveType = PrimitiveType;
-	PendingNumPrimitives = NumPrimitives;
-	PendingVertexDataStride = VertexDataStride;
-	}
-}
-
-
-void FMetalRHICommandContext::RHIEndDrawPrimitiveUP()
-{
-	@autoreleasepool {
-	SCOPE_CYCLE_COUNTER(STAT_MetalDrawCallTime);
-
-	RHI_DRAW_CALL_STATS(PendingPrimitiveType,PendingNumPrimitives);
-
-	// set the vertex buffer
-	uint32 NumVertices = GetVertexCountForPrimitiveCount(PendingNumPrimitives, PendingPrimitiveType);
-	Context->GetCurrentState().SetVertexStream(0, &PendingVertexBuffer, nil, 0, PendingVertexDataStride * NumVertices);
-#if PLATFORM_SUPPORTS_TESSELLATION_SHADERS
-	if(Context->GetCurrentState().GetUsingTessellation())
-	{
-		Context->GetCurrentState().SetShaderBuffer(EMetalShaderStages::Hull, PendingVertexBuffer, nil, 0, PendingVertexDataStride * NumVertices, UNREAL_TO_METAL_BUFFER_INDEX(0), mtlpp::ResourceUsage::Read);
-		Context->GetCurrentState().SetShaderBuffer(EMetalShaderStages::Domain, PendingVertexBuffer, nil, 0, PendingVertexDataStride * NumVertices, UNREAL_TO_METAL_BUFFER_INDEX(0), mtlpp::ResourceUsage::Read);
-	}
-#endif
-	
-	Context->DrawPrimitive(PendingPrimitiveType, 0, PendingNumPrimitives, Context->GetCurrentState().GetRenderTargetArraySize());
-	
-	// mark temp memory as usable
-	PendingVertexBuffer = nil;
-	}
-}
-
-void FMetalRHICommandContext::RHIBeginDrawIndexedPrimitiveUP( uint32 NumPrimitives, uint32 NumVertices, uint32 VertexDataStride, void*& OutVertexData, uint32 MinVertexIndex, uint32 NumIndices, uint32 IndexDataStride, void*& OutIndexData)
-{
-	@autoreleasepool {
-	SCOPE_CYCLE_COUNTER(STAT_MetalDrawCallTime);
-	checkSlow(PendingVertexBuffer.GetPtr() == nil);
-	checkSlow(PendingIndexBuffer.GetPtr() == nil);
-
-	uint32 PrimitiveType = Context->GetCurrentState().GetPrimitiveType();
-	
-		
-	// allocate space
-	uint32 VertexSize = Align(VertexDataStride * NumVertices, BufferOffsetAlignment);
-	uint32 IndexSize = Align(IndexDataStride * NumIndices, BufferOffsetAlignment);
-	PendingVertexBuffer = Context->AllocateFromRingBuffer(VertexSize);
-	PendingIndexBuffer = Context->AllocateFromRingBuffer(IndexSize);
-	PendingVertexBuffer.MarkSingleUse();
-	PendingIndexBuffer.MarkSingleUse();
-	
-	// get the pointers to send back for writing
-	OutVertexData = (uint8*)PendingVertexBuffer.GetContents();
-	OutIndexData = (uint8*)PendingIndexBuffer.GetContents();
-	
-	PendingPrimitiveType = PrimitiveType;
-	PendingNumPrimitives = NumPrimitives;
-	PendingIndexDataStride = IndexDataStride;
-
-	PendingVertexDataStride = VertexDataStride;
-	}
-}
-
-void FMetalRHICommandContext::RHIEndDrawIndexedPrimitiveUP()
-{
-	@autoreleasepool {
-	SCOPE_CYCLE_COUNTER(STAT_MetalDrawCallTime);
-
-	RHI_DRAW_CALL_STATS(PendingPrimitiveType,PendingNumPrimitives);
-
-	// set the vertex buffer
-	Context->GetCurrentState().SetVertexStream(0, &PendingVertexBuffer, nil, 0, PendingVertexBuffer.GetLength());
-
-	// how many to draw
-	uint32 NumIndices = GetVertexCountForPrimitiveCount(PendingNumPrimitives, PendingPrimitiveType);
-	
-	Context->DrawIndexedPrimitive(PendingIndexBuffer, PendingIndexDataStride, (PendingIndexDataStride == 2) ? mtlpp::IndexType::UInt16 : mtlpp::IndexType::UInt32, PendingPrimitiveType, 0, 0, NumIndices, 0, PendingNumPrimitives, 1);
-	
-	// mark temp memory as usable
-	PendingVertexBuffer = nil;
-	PendingIndexBuffer = nil;
-	}
-}
-
 void FMetalDynamicRHI::SetupRecursiveResources()
 {
     /*
diff --git a/Engine/Source/Runtime/Apple/MetalRHI/Public/MetalRHIContext.h b/Engine/Source/Runtime/Apple/MetalRHI/Public/MetalRHIContext.h
--- a/Engine/Source/Runtime/Apple/MetalRHI/Public/MetalRHIContext.h
+++ b/Engine/Source/Runtime/Apple/MetalRHI/Public/MetalRHIContext.h
@@ -239,40 +239,6 @@ public:
 	virtual void RHIDrawIndexedPrimitive(FIndexBufferRHIParamRef IndexBuffer, int32 BaseVertexIndex, uint32 FirstInstance, uint32 NumVertices, uint32 StartIndex, uint32 NumPrimitives, uint32 NumInstances) final override;
 	
 	virtual void RHIDrawIndexedPrimitiveIndirect(FIndexBufferRHIParamRef IndexBuffer, FVertexBufferRHIParamRef ArgumentBuffer, uint32 ArgumentOffset) final override;
-	
-	/**
-	 * Preallocate memory or get a direct command stream pointer to fill up for immediate rendering . This avoids memcpys below in DrawPrimitiveUP
-	 * @param PrimitiveType The type (triangles, lineloop, etc) of primitive to draw
-	 * @param NumPrimitives The number of primitives in the VertexData buffer
-	 * @param NumVertices The number of vertices to be written
-	 * @param VertexDataStride Size of each vertex
-	 * @param OutVertexData Reference to the allocated vertex memory
-	 */
-	virtual void RHIBeginDrawPrimitiveUP(uint32 NumPrimitives, uint32 NumVertices, uint32 VertexDataStride, void*& OutVertexData) final override;
-	
-	/**
-	 * Draw a primitive using the vertex data populated since RHIBeginDrawPrimitiveUP and clean up any memory as needed
-	 */
-	virtual void RHIEndDrawPrimitiveUP() final override;
-	
-	/**
-	 * Preallocate memory or get a direct command stream pointer to fill up for immediate rendering . This avoids memcpys below in DrawIndexedPrimitiveUP
-	 * @param PrimitiveType The type (triangles, lineloop, etc) of primitive to draw
-	 * @param NumPrimitives The number of primitives in the VertexData buffer
-	 * @param NumVertices The number of vertices to be written
-	 * @param VertexDataStride Size of each vertex
-	 * @param OutVertexData Reference to the allocated vertex memory
-	 * @param MinVertexIndex The lowest vertex index used by the index buffer
-	 * @param NumIndices Number of indices to be written
-	 * @param IndexDataStride Size of each index (either 2 or 4 bytes)
-	 * @param OutIndexData Reference to the allocated index memory
-	 */
-	virtual void RHIBeginDrawIndexedPrimitiveUP(uint32 NumPrimitives, uint32 NumVertices, uint32 VertexDataStride, void*& OutVertexData, uint32 MinVertexIndex, uint32 NumIndices, uint32 IndexDataStride, void*& OutIndexData) final override;
-	
-	/**
-	 * Draw a primitive using the vertex and index data populated since RHIBeginDrawIndexedPrimitiveUP and clean up any memory as needed
-	 */
-	virtual void RHIEndDrawIndexedPrimitiveUP() final override;
 
 	/**
 	* Sets Depth Bounds Testing with the given min/max depth.
diff --git a/Engine/Source/Runtime/D3D12RHI/Private/D3D12CommandContext.cpp b/Engine/Source/Runtime/D3D12RHI/Private/D3D12CommandContext.cpp
--- a/Engine/Source/Runtime/D3D12RHI/Private/D3D12CommandContext.cpp
+++ b/Engine/Source/Runtime/D3D12RHI/Private/D3D12CommandContext.cpp
@@ -688,85 +688,6 @@ FD3D12CommandContextRedirector::FD3D12CommandContextRedirector(class FD3D12Adapt
 	FMemory::Memzero(PhysicalContexts, sizeof(PhysicalContexts[0]) * MAX_NUM_GPUS);
 }
 
-void FD3D12CommandContextRedirector::RHIBeginDrawPrimitiveUP(uint32 NumPrimitives, uint32 NumVertices, uint32 VertexDataStride, void*& OutVertexData)
-{
-	check(!PendingUP);
-
-	PendingUP.NumPrimitives = NumPrimitives;
-	PendingUP.NumVertices = NumVertices;
-	PendingUP.VertexDataStride = VertexDataStride;
-	OutVertexData = PendingUP.VertexData = FMemory::Malloc(NumVertices * VertexDataStride, 16);
-}
-	
-void FD3D12CommandContextRedirector::RHIEndDrawPrimitiveUP()
-{
-	if (PendingUP.VertexData)
-	{
-		for (uint32 GPUIndex : GPUMask)
-		{
-			FD3D12CommandContext* GPUContext = PhysicalContexts[GPUIndex];
-			check(GPUContext);
-
-			void* GPUVertexData = nullptr;
-			GPUContext->RHIBeginDrawPrimitiveUP(PendingUP.NumPrimitives, PendingUP.NumVertices, PendingUP.VertexDataStride, GPUVertexData);
-			if (GPUVertexData)
-			{
-				FMemory::Memcpy(GPUVertexData, PendingUP.VertexData, PendingUP.NumVertices * PendingUP.VertexDataStride);
-			}
-			GPUContext->RHIEndDrawPrimitiveUP();
-		}
-
-		FMemory::Free(PendingUP.VertexData);
-	}
-
-	PendingUP.Reset();
-}
-	
-void FD3D12CommandContextRedirector::RHIBeginDrawIndexedPrimitiveUP(uint32 NumPrimitives, uint32 NumVertices, uint32 VertexDataStride, void*& OutVertexData, uint32 MinVertexIndex, uint32 NumIndices, uint32 IndexDataStride, void*& OutIndexData)
-{
-	check(!PendingUP);
-
-	PendingUP.NumPrimitives = NumPrimitives;
-	PendingUP.NumVertices = NumVertices;
-	PendingUP.VertexDataStride = VertexDataStride;
-	OutVertexData = PendingUP.VertexData = FMemory::Malloc(NumVertices * VertexDataStride);
-
-	PendingUP.MinVertexIndex = MinVertexIndex;
-	PendingUP.NumIndices = NumIndices;
-	PendingUP.IndexDataStride = IndexDataStride;
-	OutIndexData = PendingUP.IndexData = FMemory::Malloc(NumIndices * IndexDataStride);
-}
-	
-void FD3D12CommandContextRedirector::RHIEndDrawIndexedPrimitiveUP()
-{
-	if (PendingUP.VertexData && PendingUP.IndexData)
-	{
-		for (uint32 GPUIndex : GPUMask)
-		{
-			FD3D12CommandContext* GPUContext = PhysicalContexts[GPUIndex];
-			check(GPUContext);
-
-			void* GPUVertexData = nullptr;
-			void* GPUIndexData = nullptr;
-			GPUContext->RHIBeginDrawIndexedPrimitiveUP(PendingUP.NumPrimitives, PendingUP.NumVertices, PendingUP.VertexDataStride, GPUVertexData, PendingUP.MinVertexIndex, PendingUP.NumIndices, PendingUP.IndexDataStride, GPUIndexData);
-			if (GPUVertexData)
-			{
-				FMemory::Memcpy(GPUVertexData, PendingUP.VertexData, PendingUP.NumVertices * PendingUP.VertexDataStride);
-			}
-			if (GPUIndexData)
-			{
-				FMemory::Memcpy(GPUIndexData, PendingUP.IndexData, PendingUP.NumIndices * PendingUP.IndexDataStride);
-			}
-			GPUContext->RHIEndDrawIndexedPrimitiveUP();
-		}
-
-		FMemory::Free(PendingUP.VertexData);
-		FMemory::Free(PendingUP.IndexData);
-	}
-
-	PendingUP.Reset();
-}
-
 void FD3D12CommandContextRedirector::RHITransitionResources(EResourceTransitionAccess TransitionType, EResourceTransitionPipeline TransitionPipeline, FUnorderedAccessViewRHIParamRef* InUAVs, int32 NumUAVs, FComputeFenceRHIParamRef WriteComputeFenceRHI)
 {
 	ContextRedirect(RHITransitionResources(TransitionType, TransitionPipeline, InUAVs, NumUAVs, nullptr));
diff --git a/Engine/Source/Runtime/D3D12RHI/Private/D3D12CommandContext.h b/Engine/Source/Runtime/D3D12RHI/Private/D3D12CommandContext.h
--- a/Engine/Source/Runtime/D3D12RHI/Private/D3D12CommandContext.h
+++ b/Engine/Source/Runtime/D3D12RHI/Private/D3D12CommandContext.h
@@ -55,28 +55,6 @@ protected:
 
 	virtual FD3D12CommandContext* GetContext(uint32 InGPUIndex) = 0;
 
-	// State for begin/end draw primitive UP interface.
-	struct FUserPrimitiveData
-	{
-		FUserPrimitiveData() { FMemory::Memzero(*this); }
-		~FUserPrimitiveData() { check(!VertexData && !IndexData); }
-
-		uint32 NumPrimitives;
-		uint32 NumVertices;
-		uint32 VertexDataStride;
-		void* VertexData;
-
-		uint32 MinVertexIndex;
-		uint32 NumIndices;
-		uint32 IndexDataStride;
-		void* IndexData;
-
-		operator bool() const { return NumVertices != 0; }
-		void Reset() { FMemory::Memzero(*this); }
-	};
-
-	FUserPrimitiveData PendingUP;
-
 protected:
 
 	FRHIGPUMask GPUMask;
@@ -320,10 +298,6 @@ public:
 	virtual void RHIDrawIndexedIndirect(FIndexBufferRHIParamRef IndexBufferRHI, FStructuredBufferRHIParamRef ArgumentsBufferRHI, int32 DrawArgumentsIndex, uint32 NumInstances) final override;
 	virtual void RHIDrawIndexedPrimitive(FIndexBufferRHIParamRef IndexBuffer, int32 BaseVertexIndex, uint32 FirstInstance, uint32 NumVertices, uint32 StartIndex, uint32 NumPrimitives, uint32 NumInstances) final override;
 	virtual void RHIDrawIndexedPrimitiveIndirect(FIndexBufferRHIParamRef IndexBuffer, FVertexBufferRHIParamRef ArgumentBuffer, uint32 ArgumentOffset) final override;
-	virtual void RHIBeginDrawPrimitiveUP(uint32 NumPrimitives, uint32 NumVertices, uint32 VertexDataStride, void*& OutVertexData) final override;
-	virtual void RHIEndDrawPrimitiveUP() final override;
-	virtual void RHIBeginDrawIndexedPrimitiveUP(uint32 NumPrimitives, uint32 NumVertices, uint32 VertexDataStride, void*& OutVertexData, uint32 MinVertexIndex, uint32 NumIndices, uint32 IndexDataStride, void*& OutIndexData) final override;
-	virtual void RHIEndDrawIndexedPrimitiveUP() final override;
 	virtual void RHISetDepthBounds(float MinDepth, float MaxDepth) final override;
 	virtual void RHIUpdateTextureReference(FTextureReferenceRHIParamRef TextureRef, FTextureRHIParamRef NewTexture) final override;
 
@@ -728,14 +702,6 @@ public:
 		ContextRedirect(RHIDrawIndexedPrimitiveIndirect(IndexBuffer, ArgumentBuffer, ArgumentOffset));
 	}
 	
-	virtual void RHIBeginDrawPrimitiveUP(uint32 NumPrimitives, uint32 NumVertices, uint32 VertexDataStride, void*& OutVertexData) final override;
-	
-	virtual void RHIEndDrawPrimitiveUP() final override;
-	
-	virtual void RHIBeginDrawIndexedPrimitiveUP(uint32 NumPrimitives, uint32 NumVertices, uint32 VertexDataStride, void*& OutVertexData, uint32 MinVertexIndex, uint32 NumIndices, uint32 IndexDataStride, void*& OutIndexData) final override;
-	
-	virtual void RHIEndDrawIndexedPrimitiveUP() final override;
-
 	FORCEINLINE virtual void RHISetDepthBounds(float MinDepth, float MaxDepth) final override
 	{
 		ContextRedirect(RHISetDepthBounds(MinDepth, MaxDepth));
diff --git a/Engine/Source/Runtime/D3D12RHI/Private/D3D12Commands.cpp b/Engine/Source/Runtime/D3D12RHI/Private/D3D12Commands.cpp
--- a/Engine/Source/Runtime/D3D12RHI/Private/D3D12Commands.cpp
+++ b/Engine/Source/Runtime/D3D12RHI/Private/D3D12Commands.cpp
@@ -1623,132 +1623,6 @@ void FD3D12CommandContext::RHIDrawIndexedPrimitiveIndirect(FIndexBufferRHIParamR
 	DEBUG_EXECUTE_COMMAND_LIST(this);
 }
 
-/**
-* Preallocate memory or get a direct command stream pointer to fill up for immediate rendering . This avoids memcpys below in DrawPrimitiveUP
-* @param NumPrimitives The number of primitives in the VertexData buffer
-* @param NumVertices The number of vertices to be written
-* @param VertexDataStride Size of each vertex
-* @param OutVertexData Reference to the allocated vertex memory
-*/
-void FD3D12CommandContext::RHIBeginDrawPrimitiveUP(uint32 NumPrimitives, uint32 NumVertices, uint32 VertexDataStride, void*& OutVertexData)
-{
-	checkSlow(PendingUP.NumVertices == 0);
-
-	// Remember the parameters for the draw call.
-	PendingUP.NumPrimitives = NumPrimitives;
-	PendingUP.NumVertices = NumVertices;
-	PendingUP.VertexDataStride = VertexDataStride;
-
-	// Map the dynamic buffer.
-	OutVertexData = DynamicVB.Lock(NumVertices * VertexDataStride);
-}
-
-/**
-* Draw a primitive using the vertex data populated since RHIBeginDrawPrimitiveUP and clean up any memory as needed
-*/
-void FD3D12CommandContext::RHIEndDrawPrimitiveUP()
-{
-	RHI_DRAW_CALL_STATS(StateCache.GetGraphicsPipelinePrimitiveType(), PendingUP.NumPrimitives);
-
-	checkSlow(!bUsingTessellation || StateCache.GetGraphicsPipelinePrimitiveType() == PT_TriangleList);
-
-	if (IsDefaultContext())
-	{
-		GetParentDevice()->RegisterGPUWork(PendingUP.NumPrimitives, PendingUP.NumVertices);
-	}
-
-	// Unmap the dynamic vertex buffer.
-	FD3D12ResourceLocation* BufferLocation = DynamicVB.Unlock();
-	uint32 VBOffset = 0;
-
-	// Issue the draw call.
-	CommitGraphicsResourceTables();
-	CommitNonComputeShaderConstants();
-	StateCache.SetStreamSource(BufferLocation, 0, PendingUP.VertexDataStride, VBOffset);
-	StateCache.SetPrimitiveTopology(GetD3D12PrimitiveType(StateCache.GetGraphicsPipelinePrimitiveType(), bUsingTessellation));
-	StateCache.ApplyState<D3D12PT_Graphics>();
-	numDraws++;
-	CommandListHandle->DrawInstanced(PendingUP.NumVertices, 1, 0, 0);
-#if UE_BUILD_DEBUG	
-	OwningRHI.DrawCount++;
-#endif
-	DEBUG_EXECUTE_COMMAND_LIST(this);
-
-	// Clear these parameters.
-	PendingUP.Reset();
-}
-
-/**
-* Preallocate memory or get a direct command stream pointer to fill up for immediate rendering . This avoids memcpys below in DrawIndexedPrimitiveUP
-* @param NumPrimitives The number of primitives in the VertexData buffer
-* @param NumVertices The number of vertices to be written
-* @param VertexDataStride Size of each vertex
-* @param OutVertexData Reference to the allocated vertex memory
-* @param MinVertexIndex The lowest vertex index used by the index buffer
-* @param NumIndices Number of indices to be written
-* @param IndexDataStride Size of each index (either 2 or 4 bytes)
-* @param OutIndexData Reference to the allocated index memory
-*/
-void FD3D12CommandContext::RHIBeginDrawIndexedPrimitiveUP(uint32 NumPrimitives, uint32 NumVertices, uint32 VertexDataStride, void*& OutVertexData, uint32 MinVertexIndex, uint32 NumIndices, uint32 IndexDataStride, void*& OutIndexData)
-{
-	checkSlow((sizeof(uint16) == IndexDataStride) || (sizeof(uint32) == IndexDataStride));
-
-	// Store off information needed for the draw call.
-	PendingUP.NumPrimitives = NumPrimitives;
-	PendingUP.MinVertexIndex = MinVertexIndex;
-	PendingUP.IndexDataStride = IndexDataStride;
-	PendingUP.NumVertices = NumVertices;
-	PendingUP.NumIndices = NumIndices;
-	PendingUP.VertexDataStride = VertexDataStride;
-
-	// Map dynamic vertex and index buffers.
-	OutVertexData = DynamicVB.Lock(NumVertices * VertexDataStride);
-	OutIndexData = DynamicIB.Lock(NumIndices * IndexDataStride);
-}
-
-/**
-* Draw a primitive using the vertex and index data populated since RHIBeginDrawIndexedPrimitiveUP and clean up any memory as needed
-*/
-void FD3D12CommandContext::RHIEndDrawIndexedPrimitiveUP()
-{
-	// tessellation only supports trilists
-	checkSlow(!bUsingTessellation || StateCache.GetGraphicsPipelinePrimitiveType() == PT_TriangleList);
-
-	RHI_DRAW_CALL_STATS(StateCache.GetGraphicsPipelinePrimitiveType(), PendingUP.NumPrimitives);
-
-	if (IsDefaultContext())
-	{
-		GetParentDevice()->RegisterGPUWork(PendingUP.NumPrimitives, PendingUP.NumVertices);
-	}
-
-	// Unmap the dynamic buffers.
-	FD3D12ResourceLocation* VertexBufferLocation = DynamicVB.Unlock();
-	FD3D12ResourceLocation* IndexBufferLocation = DynamicIB.Unlock();
-	uint32 VBOffset = 0;
-
-	// Issue the draw call.
-	CommitGraphicsResourceTables();
-	CommitNonComputeShaderConstants();
-	StateCache.SetStreamSource(VertexBufferLocation, 0, PendingUP.VertexDataStride, VBOffset);
-	StateCache.SetIndexBuffer(IndexBufferLocation, PendingUP.IndexDataStride == sizeof(uint16) ? DXGI_FORMAT_R16_UINT : DXGI_FORMAT_R32_UINT, 0);
-	StateCache.SetPrimitiveTopology(GetD3D12PrimitiveType(StateCache.GetGraphicsPipelinePrimitiveType(), bUsingTessellation));
-	StateCache.ApplyState<D3D12PT_Graphics>();
-
-	numDraws++;
-	CommandListHandle->DrawIndexedInstanced(PendingUP.NumIndices, 1, 0, PendingUP.MinVertexIndex, 0);
-#if UE_BUILD_DEBUG	
-	OwningRHI.DrawCount++;
-#endif
-	DEBUG_EXECUTE_COMMAND_LIST(this);
-
-	//It's important to release the locations so the fast alloc page can be freed
-	DynamicVB.ReleaseResourceLocation();
-	DynamicIB.ReleaseResourceLocation();
-
-	// Clear these parameters.
-	PendingUP.Reset();
-}
-
 // Raster operations.
 void FD3D12CommandContext::RHIClearMRT(bool bClearColor, int32 NumClearColors, const FLinearColor* ClearColorArray, bool bClearDepth, float Depth, bool bClearStencil, uint32 Stencil)
 {
diff --git a/Engine/Source/Runtime/EmptyRHI/Public/EmptyRHI.h b/Engine/Source/Runtime/EmptyRHI/Public/EmptyRHI.h
--- a/Engine/Source/Runtime/EmptyRHI/Public/EmptyRHI.h
+++ b/Engine/Source/Runtime/EmptyRHI/Public/EmptyRHI.h
@@ -194,10 +194,6 @@ public:
 	virtual void RHIDrawIndexedIndirect(FIndexBufferRHIParamRef IndexBufferRHI, uint32 PrimitiveType, FStructuredBufferRHIParamRef ArgumentsBufferRHI, int32 DrawArgumentsIndex, uint32 NumInstances) final override;
 	virtual void RHIDrawIndexedPrimitive(FIndexBufferRHIParamRef IndexBuffer, uint32 PrimitiveType, int32 BaseVertexIndex, uint32 FirstInstance, uint32 NumVertices, uint32 StartIndex, uint32 NumPrimitives, uint32 NumInstances) final override;
 	virtual void RHIDrawIndexedPrimitiveIndirect(uint32 PrimitiveType, FIndexBufferRHIParamRef IndexBuffer, FVertexBufferRHIParamRef ArgumentBuffer, uint32 ArgumentOffset) final override;
-	virtual void RHIBeginDrawPrimitiveUP(uint32 PrimitiveType, uint32 NumPrimitives, uint32 NumVertices, uint32 VertexDataStride, void*& OutVertexData) final override;
-	virtual void RHIEndDrawPrimitiveUP() final override;
-	virtual void RHIBeginDrawIndexedPrimitiveUP(uint32 PrimitiveType, uint32 NumPrimitives, uint32 NumVertices, uint32 VertexDataStride, void*& OutVertexData, uint32 MinVertexIndex, uint32 NumIndices, uint32 IndexDataStride, void*& OutIndexData) final override;
-	virtual void RHIEndDrawIndexedPrimitiveUP() final override;
 	virtual void RHIEnableDepthBoundsTest(bool bEnable) final override;
 	virtual void RHISetDepthBounds(float MinDepth, float MaxDepth) final override;
 	virtual void RHIPushEvent(const TCHAR* Name, FColor Color) final override;
diff --git a/Engine/Source/Runtime/NullDrv/Public/NullRHI.h b/Engine/Source/Runtime/NullDrv/Public/NullRHI.h
--- a/Engine/Source/Runtime/NullDrv/Public/NullRHI.h
+++ b/Engine/Source/Runtime/NullDrv/Public/NullRHI.h
@@ -757,22 +757,6 @@ public:
 
 	}
 
-	virtual void RHIBeginDrawPrimitiveUP(uint32 NumPrimitives, uint32 NumVertices, uint32 VertexDataStride, void*& OutVertexData) final override
-	{ 
-		OutVertexData = GetStaticBuffer(); 
-	}
-	virtual void RHIEndDrawPrimitiveUP()
-	{
-
-	}
-	virtual void RHIBeginDrawIndexedPrimitiveUP(uint32 NumPrimitives, uint32 NumVertices, uint32 VertexDataStride, void*& OutVertexData, uint32 MinVertexIndex, uint32 NumIndices, uint32 IndexDataStride, void*& OutIndexData) final override
-	{ 
-		OutVertexData = GetStaticBuffer(); 
-		OutIndexData = GetStaticBuffer(); 
-	}
-	virtual void RHIEndDrawIndexedPrimitiveUP() final override
-	{
-	}
 	virtual void RHIBlockUntilGPUIdle() final override
 	{
 	}
diff --git a/Engine/Source/Runtime/OpenGLDrv/Private/OpenGLCommands.cpp b/Engine/Source/Runtime/OpenGLDrv/Private/OpenGLCommands.cpp
--- a/Engine/Source/Runtime/OpenGLDrv/Private/OpenGLCommands.cpp
+++ b/Engine/Source/Runtime/OpenGLDrv/Private/OpenGLCommands.cpp
@@ -2949,270 +2949,6 @@ void FOpenGLDynamicRHI::RHIDrawIndexedPrimitiveIndirect(FIndexBufferRHIParamRef
 	}
 }
 
-/**
- * Preallocate memory or get a direct command stream pointer to fill up for immediate rendering . This avoids memcpys below in DrawPrimitiveUP
- * @param PrimitiveType The type (triangles, lineloop, etc) of primitive to draw
- * @param NumPrimitives The number of primitives in the VertexData buffer
- * @param NumVertices The number of vertices to be written
- * @param VertexDataStride Size of each vertex 
- * @param OutVertexData Reference to the allocated vertex memory
- */
-void FOpenGLDynamicRHI::RHIBeginDrawPrimitiveUP( uint32 NumPrimitives, uint32 NumVertices, uint32 VertexDataStride, void*& OutVertexData)
-{
-	SCOPE_CYCLE_COUNTER_DETAILED(STAT_OpenGLDrawPrimitiveUPTime);
-	VERIFY_GL_SCOPE();
-	check(PendingState.NumPrimitives == 0);
-
-	if(FOpenGL::SupportsFastBufferData()) 
-	{
-		OutVertexData = DynamicVertexBuffers.Lock(NumVertices * VertexDataStride);
-	}
-	else 
-	{
-		uint32 BytesVertex = NumVertices * VertexDataStride;
-		if(BytesVertex > PendingState.UpVertexBufferBytes) 
-		{
-			if(PendingState.UpVertexBuffer) 
-			{
-				FMemory::Free(PendingState.UpVertexBuffer);
-			}
-			PendingState.UpVertexBuffer = FMemory::Malloc(BytesVertex);
-			PendingState.UpVertexBufferBytes = BytesVertex;
-		}
-		OutVertexData = PendingState.UpVertexBuffer;
-		PendingState.UpStride = VertexDataStride;
-	}
-
-	PendingState.PrimitiveType = PrimitiveType;
-	PendingState.NumPrimitives = NumPrimitives;
-	PendingState.NumVertices = NumVertices;
-	if(FOpenGL::SupportsFastBufferData())
-	{
-		PendingState.DynamicVertexStream.VertexBuffer = DynamicVertexBuffers.GetPendingBuffer();
-		PendingState.DynamicVertexStream.Offset = DynamicVertexBuffers.GetPendingOffset();
-		PendingState.DynamicVertexStream.Stride = VertexDataStride;
-	}
-	else
-	{
-		PendingState.DynamicVertexStream.VertexBuffer = 0;
-		PendingState.DynamicVertexStream.Offset = 0;
-		PendingState.DynamicVertexStream.Stride = VertexDataStride;
-	}
-}
-
-/**
- * Draw a primitive using the vertex data populated since RHIBeginDrawPrimitiveUP and clean up any memory as needed
- */
-void FOpenGLDynamicRHI::RHIEndDrawPrimitiveUP()
-{
-	SCOPE_CYCLE_COUNTER_DETAILED(STAT_OpenGLDrawPrimitiveUPTime);
-	VERIFY_GL_SCOPE();
-	check(PendingState.NumPrimitives != 0);
-
-	RHI_DRAW_CALL_STATS(PendingState.PrimitiveType,PendingState.NumPrimitives);
-
-
-	if(FOpenGL::SupportsFastBufferData()) 
-	{
-		DynamicVertexBuffers.Unlock();
-	}
-
-	FOpenGLContextState& ContextState = GetContextStateForCurrentContext();
-	BindPendingFramebuffer(ContextState);
-	SetPendingBlendStateForActiveRenderTargets(ContextState);
-	UpdateViewportInOpenGLContext(ContextState);
-	UpdateScissorRectInOpenGLContext(ContextState);
-	UpdateRasterizerStateInOpenGLContext(ContextState);
-	UpdateDepthStencilStateInOpenGLContext(ContextState);
-	BindPendingShaderState(ContextState);
-	CommitGraphicsResourceTables();
-	SetupTexturesForDraw(ContextState);
-	CommitNonComputeShaderConstants();
-	CachedBindElementArrayBuffer(ContextState,0);
-
-	GLenum DrawMode = GL_TRIANGLES;
-	GLsizei NumElements = 0;
-	GLint PatchSize = 0;
-	FindPrimitiveType(PendingState.PrimitiveType, ContextState.bUsingTessellation, PendingState.NumPrimitives, DrawMode, NumElements, PatchSize);
-
-	if (FOpenGL::SupportsTessellation() && DrawMode == GL_PATCHES)
-	{
-		FOpenGL::PatchParameteri(GL_PATCH_VERTICES, PatchSize);
-	}
-
-	if(FOpenGL::SupportsFastBufferData()) 
-	{
-		SetupVertexArrays(ContextState, 0, &PendingState.DynamicVertexStream, 1, PendingState.NumVertices);
-	}
-	else
-	{
-		SetupVertexArraysUP(ContextState, PendingState.UpVertexBuffer, PendingState.UpStride);
-	}
-	
-#if DEBUG_GL_SHADERS
-	VerifyProgramPipeline();
-#endif
-
-	GPUProfilingData.RegisterGPUWork(PendingState.NumPrimitives,PendingState.NumVertices);
-	{
-		CONDITIONAL_SCOPE_CYCLE_COUNTER(STAT_OpenGLShaderFirstDrawTime, PendingState.BoundShaderState->RequiresDriverInstantiation());
-		glDrawArrays(DrawMode, 0, NumElements);
-	}
-	PendingState.NumPrimitives = 0;
-
-	REPORT_GL_DRAW_ARRAYS_EVENT_FOR_FRAME_DUMP( DrawMode, 0, NumElements );
-}
-
-/**
- * Preallocate memory or get a direct command stream pointer to fill up for immediate rendering . This avoids memcpys below in DrawIndexedPrimitiveUP
- * @param PrimitiveType The type (triangles, lineloop, etc) of primitive to draw
- * @param NumPrimitives The number of primitives in the VertexData buffer
- * @param NumVertices The number of vertices to be written
- * @param VertexDataStride Size of each vertex
- * @param OutVertexData Reference to the allocated vertex memory
- * @param MinVertexIndex The lowest vertex index used by the index buffer
- * @param NumIndices Number of indices to be written
- * @param IndexDataStride Size of each index (either 2 or 4 bytes)
- * @param OutIndexData Reference to the allocated index memory
- */
-void FOpenGLDynamicRHI::RHIBeginDrawIndexedPrimitiveUP(uint32 NumPrimitives, uint32 NumVertices, uint32 VertexDataStride, void*& OutVertexData, uint32 MinVertexIndex, uint32 NumIndices, uint32 IndexDataStride, void*& OutIndexData)
-{
-	VERIFY_GL_SCOPE();
-	SCOPE_CYCLE_COUNTER_DETAILED(STAT_OpenGLDrawPrimitiveUPTime);
-	check(PendingState.NumPrimitives == 0);
-	check((sizeof(uint16) == IndexDataStride) || (sizeof(uint32) == IndexDataStride));
-
-	if(FOpenGL::SupportsFastBufferData()) 
-	{
-		OutVertexData = DynamicVertexBuffers.Lock(NumVertices * VertexDataStride);
-		OutIndexData = DynamicIndexBuffers.Lock(NumIndices * IndexDataStride);
-	}
-	else
-	{
-		uint32 BytesVertex = NumVertices * VertexDataStride;
-		if(BytesVertex > PendingState.UpVertexBufferBytes) 
-		{
-			if(PendingState.UpVertexBuffer) 
-			{
-				FMemory::Free(PendingState.UpVertexBuffer);
-			}
-			PendingState.UpVertexBuffer = FMemory::Malloc(BytesVertex);
-			PendingState.UpVertexBufferBytes = BytesVertex;
-		}
-		OutVertexData = PendingState.UpVertexBuffer;
-		PendingState.UpStride = VertexDataStride;
-		uint32 BytesIndex = NumIndices * IndexDataStride;
-		if(BytesIndex > PendingState.UpIndexBufferBytes)
-		{
-			if(PendingState.UpIndexBuffer) 
-			{
-				FMemory::Free(PendingState.UpIndexBuffer);
-			}
-			PendingState.UpIndexBuffer = FMemory::Malloc(BytesIndex);
-			PendingState.UpIndexBufferBytes = BytesIndex;
-		}
-		OutIndexData = PendingState.UpIndexBuffer;
-	}
-
-	PendingState.PrimitiveType = PrimitiveType;
-	PendingState.NumPrimitives = NumPrimitives;
-	PendingState.MinVertexIndex = MinVertexIndex;
-	PendingState.IndexDataStride = IndexDataStride;
-	PendingState.NumVertices = NumVertices;
-	if(FOpenGL::SupportsFastBufferData()) 
-	{
-		PendingState.DynamicVertexStream.VertexBuffer = DynamicVertexBuffers.GetPendingBuffer();
-		PendingState.DynamicVertexStream.Offset = DynamicVertexBuffers.GetPendingOffset();
-		PendingState.DynamicVertexStream.Stride = VertexDataStride;
-	}
-	else
-	{
-		PendingState.DynamicVertexStream.VertexBuffer = 0;
-		PendingState.DynamicVertexStream.Offset = 0;
-		PendingState.DynamicVertexStream.Stride = VertexDataStride;
-	}
-}
-
-/**
- * Draw a primitive using the vertex and index data populated since RHIBeginDrawIndexedPrimitiveUP and clean up any memory as needed
- */
-void FOpenGLDynamicRHI::RHIEndDrawIndexedPrimitiveUP()
-{
-	SCOPE_CYCLE_COUNTER_DETAILED(STAT_OpenGLDrawPrimitiveUPTime);
-	VERIFY_GL_SCOPE();
-	check(PendingState.NumPrimitives != 0);
-
-	RHI_DRAW_CALL_STATS(PendingState.PrimitiveType,PendingState.NumPrimitives);
-
-	if(FOpenGL::SupportsFastBufferData()) 
-	{
-		DynamicVertexBuffers.Unlock();
-		DynamicIndexBuffers.Unlock();
-	}
-
-	FOpenGLContextState& ContextState = GetContextStateForCurrentContext();
-	BindPendingFramebuffer(ContextState);
-	SetPendingBlendStateForActiveRenderTargets(ContextState);
-	UpdateViewportInOpenGLContext(ContextState);
-	UpdateScissorRectInOpenGLContext(ContextState);
-	UpdateRasterizerStateInOpenGLContext(ContextState);
-	UpdateDepthStencilStateInOpenGLContext(ContextState);
-	BindPendingShaderState(ContextState);
-	CommitGraphicsResourceTables();
-	SetupTexturesForDraw(ContextState);
-	CommitNonComputeShaderConstants();
-	if(FOpenGL::SupportsFastBufferData()) 
-	{
-		CachedBindElementArrayBuffer(ContextState,DynamicIndexBuffers.GetPendingBuffer()->Resource);
-		SetupVertexArrays(ContextState, 0, &PendingState.DynamicVertexStream, 1, PendingState.NumVertices);
-	}
-	else
-	{
-		CachedBindElementArrayBuffer(ContextState,0);
-		SetupVertexArraysUP(ContextState, PendingState.UpVertexBuffer, PendingState.UpStride);
-	}
-
-	GLenum DrawMode = GL_TRIANGLES;
-	GLsizei NumElements = 0;
-	GLint PatchSize = 0;
-	FindPrimitiveType(PendingState.PrimitiveType, ContextState.bUsingTessellation, PendingState.NumPrimitives, DrawMode, NumElements, PatchSize);
-	GLenum IndexType = (PendingState.IndexDataStride == sizeof(uint32)) ? GL_UNSIGNED_INT : GL_UNSIGNED_SHORT;
-
-	if (FOpenGL::SupportsTessellation() && DrawMode == GL_PATCHES )
-	{
-		FOpenGL::PatchParameteri(GL_PATCH_VERTICES, PatchSize);
-	}
-	
-#if DEBUG_GL_SHADERS
-	VerifyProgramPipeline();
-#endif
-
-	GPUProfilingData.RegisterGPUWork(PendingState.NumPrimitives,PendingState.NumVertices);
-	if(FOpenGL::SupportsFastBufferData())
-	{
-		CONDITIONAL_SCOPE_CYCLE_COUNTER(STAT_OpenGLShaderFirstDrawTime, PendingState.BoundShaderState->RequiresDriverInstantiation());
-		if ( FOpenGL::SupportsDrawIndexOffset() )
-		{
-			FOpenGL::DrawRangeElements(DrawMode, PendingState.MinVertexIndex, PendingState.MinVertexIndex + PendingState.NumVertices, NumElements, IndexType, INDEX_TO_VOID(DynamicIndexBuffers.GetPendingOffset()));
-		}
-		else
-		{
-			check(PendingState.MinVertexIndex == 0);
-			glDrawElements(DrawMode, NumElements, IndexType, INDEX_TO_VOID(DynamicIndexBuffers.GetPendingOffset()));
-		}
-	}
-	else
-	{
-		CONDITIONAL_SCOPE_CYCLE_COUNTER(STAT_OpenGLShaderFirstDrawTime, PendingState.BoundShaderState->RequiresDriverInstantiation());
-		glDrawElements(DrawMode, NumElements, IndexType, PendingState.UpIndexBuffer);
-	}
-
-	PendingState.NumPrimitives = 0;
-
-	REPORT_GL_DRAW_RANGE_ELEMENTS_EVENT_FOR_FRAME_DUMP( DrawMode, PendingState.MinVertexIndex, PendingState.MinVertexIndex + PendingState.NumVertices, NumElements, IndexType, 0 );
-}
-
-
 // Raster operations.
 static inline void ClearCurrentDepthStencilWithCurrentScissor( int8 ClearType, float Depth, uint32 Stencil )
 {
diff --git a/Engine/Source/Runtime/OpenGLDrv/Public/OpenGLDrv.h b/Engine/Source/Runtime/OpenGLDrv/Public/OpenGLDrv.h
--- a/Engine/Source/Runtime/OpenGLDrv/Public/OpenGLDrv.h
+++ b/Engine/Source/Runtime/OpenGLDrv/Public/OpenGLDrv.h
@@ -517,10 +517,6 @@ public:
 	virtual void RHIDrawIndexedIndirect(FIndexBufferRHIParamRef IndexBufferRHI, FStructuredBufferRHIParamRef ArgumentsBufferRHI, int32 DrawArgumentsIndex, uint32 NumInstances) final override;
 	virtual void RHIDrawIndexedPrimitive(FIndexBufferRHIParamRef IndexBuffer, int32 BaseVertexIndex, uint32 FirstInstance, uint32 NumVertices, uint32 StartIndex, uint32 NumPrimitives, uint32 NumInstances) final override;
 	virtual void RHIDrawIndexedPrimitiveIndirect(FIndexBufferRHIParamRef IndexBuffer, FVertexBufferRHIParamRef ArgumentBuffer, uint32 ArgumentOffset) final override;
-	virtual void RHIBeginDrawPrimitiveUP(uint32 NumPrimitives, uint32 NumVertices, uint32 VertexDataStride, void*& OutVertexData) final override;
-	virtual void RHIEndDrawPrimitiveUP() final override;
-	virtual void RHIBeginDrawIndexedPrimitiveUP(uint32 NumPrimitives, uint32 NumVertices, uint32 VertexDataStride, void*& OutVertexData, uint32 MinVertexIndex, uint32 NumIndices, uint32 IndexDataStride, void*& OutIndexData) final override;
-	virtual void RHIEndDrawIndexedPrimitiveUP() final override;
 	virtual void RHIEnableDepthBoundsTest(bool bEnable) final override;
 	virtual void RHISetDepthBounds(float MinDepth, float MaxDepth) final override;
 	virtual void RHIPushEvent(const TCHAR* Name, FColor Color) final override;
diff --git a/Engine/Source/Runtime/RHI/Public/RHICommandList.h b/Engine/Source/Runtime/RHI/Public/RHICommandList.h
--- a/Engine/Source/Runtime/RHI/Public/RHICommandList.h
+++ b/Engine/Source/Runtime/RHI/Public/RHICommandList.h
@@ -525,24 +525,6 @@ public:
 	}
 	void MaybeDispatchToRHIThreadInner();
 
-	struct FDrawUpData
-	{
-		uint32 NumPrimitives;
-		uint32 NumVertices;
-		uint32 VertexDataStride;
-		void* OutVertexData;
-		uint32 MinVertexIndex;
-		uint32 NumIndices;
-		uint32 IndexDataStride;
-		void* OutIndexData;
-
-		FDrawUpData()
-			: OutVertexData(nullptr)
-			, OutIndexData(nullptr)
-		{
-		}
-	};
-
 	FORCEINLINE const FRHIGPUMask& GetGPUMask() const { return GPUMask; }
 
 private:
@@ -622,8 +604,6 @@ public:
 		return RenderThreadContexts[int32(Slot)];
 	}
 
-	FDrawUpData DrawUPData; 
-
 	struct FCommonData
 	{
 		class FRHICommandListBase* Parent = nullptr;
@@ -1152,49 +1132,6 @@ struct FRHICommandBindClearMRTValues final : public FRHICommand<FRHICommandBindC
 	RHI_API void Execute(FRHICommandListBase& CmdList);
 };
 
-struct FRHICommandEndDrawPrimitiveUP final : public FRHICommand<FRHICommandEndDrawPrimitiveUP>
-{
-	uint32 NumPrimitives;
-	uint32 NumVertices;
-	uint32 VertexDataStride;
-	void* OutVertexData;
-
-	FORCEINLINE_DEBUGGABLE FRHICommandEndDrawPrimitiveUP(uint32 InNumPrimitives, uint32 InNumVertices, uint32 InVertexDataStride, void* InOutVertexData)
-		: NumPrimitives(InNumPrimitives)
-		, NumVertices(InNumVertices)
-		, VertexDataStride(InVertexDataStride)
-		, OutVertexData(InOutVertexData)
-	{
-	}
-	RHI_API void Execute(FRHICommandListBase& CmdList);
-};
-
-
-struct FRHICommandEndDrawIndexedPrimitiveUP final : public FRHICommand<FRHICommandEndDrawIndexedPrimitiveUP>
-{
-	uint32 NumPrimitives;
-	uint32 NumVertices;
-	uint32 VertexDataStride;
-	void* OutVertexData;
-	uint32 MinVertexIndex;
-	uint32 NumIndices;
-	uint32 IndexDataStride;
-	void* OutIndexData;
-
-	FORCEINLINE_DEBUGGABLE FRHICommandEndDrawIndexedPrimitiveUP(uint32 InNumPrimitives, uint32 InNumVertices, uint32 InVertexDataStride, void* InOutVertexData, uint32 InMinVertexIndex, uint32 InNumIndices, uint32 InIndexDataStride, void* InOutIndexData)
-		: NumPrimitives(InNumPrimitives)
-		, NumVertices(InNumVertices)
-		, VertexDataStride(InVertexDataStride)
-		, OutVertexData(InOutVertexData)
-		, MinVertexIndex(InMinVertexIndex)
-		, NumIndices(InNumIndices)
-		, IndexDataStride(InIndexDataStride)
-		, OutIndexData(InOutIndexData)
-	{
-	}
-	RHI_API void Execute(FRHICommandListBase& CmdList);
-};
-
 template<ECmdList CmdListType>
 struct FRHICommandSetComputeShader final : public FRHICommand<FRHICommandSetComputeShader<CmdListType>>
 {
@@ -2454,86 +2391,6 @@ public:
 		ALLOC_COMMAND(FRHICommandBindClearMRTValues)(bClearColor, bClearDepth, bClearStencil);
 	}	
 
-	UE_DEPRECATED(4.21, "This function is deprecated and will be removed in future releases.")
-	FORCEINLINE_DEBUGGABLE void BeginDrawPrimitiveUP(uint32 PrimitiveType, uint32 NumPrimitives, uint32 NumVertices, uint32 VertexDataStride, void*& OutVertexData)
-	{
-		//check(IsOutsideRenderPass());
-		if (Bypass())
-		{
-			GetContext().RHIBeginDrawPrimitiveUP(NumPrimitives, NumVertices, VertexDataStride, OutVertexData);
-			return;
-		}
-		checkf(!DrawUPData.OutVertexData && NumVertices * VertexDataStride > 0,
-			TEXT("Data: 0x%x, NumVerts:%i, Stride:%i"), (void*)DrawUPData.OutVertexData, NumVertices, VertexDataStride);
-
-		OutVertexData = Alloc(NumVertices * VertexDataStride, 16);
-
-		DrawUPData.NumPrimitives = NumPrimitives;
-		DrawUPData.NumVertices = NumVertices;
-		DrawUPData.VertexDataStride = VertexDataStride;
-		DrawUPData.OutVertexData = OutVertexData;
-	}
-
-	UE_DEPRECATED(4.21, "This function is deprecated and will be removed in future releases.")
-	FORCEINLINE_DEBUGGABLE void EndDrawPrimitiveUP()
-	{
-		//check(IsOutsideRenderPass());
-		if (Bypass())
-		{
-			GetContext().RHIEndDrawPrimitiveUP();
-			return;
-		}
-		check(DrawUPData.OutVertexData && DrawUPData.NumVertices);
-		ALLOC_COMMAND(FRHICommandEndDrawPrimitiveUP)(DrawUPData.NumPrimitives, DrawUPData.NumVertices, DrawUPData.VertexDataStride, DrawUPData.OutVertexData);
-
-		DrawUPData.OutVertexData = nullptr;
-		DrawUPData.NumVertices = 0;
-	}
-
-	UE_DEPRECATED(4.21, "This function is deprecated and will be removed in future releases.")
-	FORCEINLINE_DEBUGGABLE void BeginDrawIndexedPrimitiveUP(uint32 PrimitiveType, uint32 NumPrimitives, uint32 NumVertices, uint32 VertexDataStride, void*& OutVertexData, uint32 MinVertexIndex, uint32 NumIndices, uint32 IndexDataStride, void*& OutIndexData)
-	{
-		//check(IsOutsideRenderPass());
-		if (Bypass())
-		{
-			GetContext().RHIBeginDrawIndexedPrimitiveUP(NumPrimitives, NumVertices, VertexDataStride, OutVertexData, MinVertexIndex, NumIndices, IndexDataStride, OutIndexData);
-			return;
-		}
-		check(!DrawUPData.OutVertexData && !DrawUPData.OutIndexData && NumVertices * VertexDataStride > 0 && NumIndices * IndexDataStride > 0);
-
-		OutVertexData = Alloc(NumVertices * VertexDataStride, 16);
-		OutIndexData = Alloc(NumIndices * IndexDataStride, 16);
-
-		DrawUPData.NumPrimitives = NumPrimitives;
-		DrawUPData.NumVertices = NumVertices;
-		DrawUPData.VertexDataStride = VertexDataStride;
-		DrawUPData.OutVertexData = OutVertexData;
-		DrawUPData.MinVertexIndex = MinVertexIndex;
-		DrawUPData.NumIndices = NumIndices;
-		DrawUPData.IndexDataStride = IndexDataStride;
-		DrawUPData.OutIndexData = OutIndexData;
-
-	}
-
-	UE_DEPRECATED(4.21, "This function is deprecated and will be removed in future releases.")
-	FORCEINLINE_DEBUGGABLE void EndDrawIndexedPrimitiveUP()
-	{
-		//check(IsOutsideRenderPass());
-		if (Bypass())
-		{
-			GetContext().RHIEndDrawIndexedPrimitiveUP();
-			return;
-		}
-		check(DrawUPData.OutVertexData && DrawUPData.OutIndexData && DrawUPData.NumIndices && DrawUPData.NumVertices);
-
-		ALLOC_COMMAND(FRHICommandEndDrawIndexedPrimitiveUP)(DrawUPData.NumPrimitives, DrawUPData.NumVertices, DrawUPData.VertexDataStride, DrawUPData.OutVertexData, DrawUPData.MinVertexIndex, DrawUPData.NumIndices, DrawUPData.IndexDataStride, DrawUPData.OutIndexData);
-
-		DrawUPData.OutVertexData = nullptr;
-		DrawUPData.OutIndexData = nullptr;
-		DrawUPData.NumIndices = 0;
-		DrawUPData.NumVertices = 0;
-	}
-
 	FORCEINLINE_DEBUGGABLE void SetComputeShader(FComputeShaderRHIParamRef ComputeShader)
 	{
 		ComputeShader->UpdateStats();
diff --git a/Engine/Source/Runtime/RHI/Public/RHICommandList.inl b/Engine/Source/Runtime/RHI/Public/RHICommandList.inl
--- a/Engine/Source/Runtime/RHI/Public/RHICommandList.inl
+++ b/Engine/Source/Runtime/RHI/Public/RHICommandList.inl
@@ -17,7 +17,6 @@ FORCEINLINE_DEBUGGABLE void FRHICommandListBase::Flush()
 {
 	if (HasCommands())
 	{
-		check(!DrawUPData.OutVertexData && !DrawUPData.OutIndexData);
 		check(!IsImmediate());
 		GRHICommandList.ExecuteList(*this);
 	}
@@ -67,7 +66,6 @@ namespace PipelineStateCache
 
 FORCEINLINE_DEBUGGABLE void FRHICommandListImmediate::ImmediateFlush(EImmediateFlushType::Type FlushType)
 {
-	check(!DrawUPData.OutVertexData && !DrawUPData.OutIndexData);
 	switch (FlushType)
 	{
 	case EImmediateFlushType::WaitForOutstandingTasksOnly:
diff --git a/Engine/Source/Runtime/RHI/Public/RHICommandListCommandExecutes.inl b/Engine/Source/Runtime/RHI/Public/RHICommandListCommandExecutes.inl
--- a/Engine/Source/Runtime/RHI/Public/RHICommandListCommandExecutes.inl
+++ b/Engine/Source/Runtime/RHI/Public/RHICommandListCommandExecutes.inl
@@ -42,9 +42,7 @@ struct FRHICommandDrawIndexedPrimitiveIndirect;
 struct FRHICommandDrawPrimitive;
 struct FRHICommandDrawPrimitiveIndirect;
 struct FRHICommandSetDepthBounds;
-struct FRHICommandEndDrawIndexedPrimitiveUP;
 struct FRHICommandEndDrawingViewport;
-struct FRHICommandEndDrawPrimitiveUP;
 struct FRHICommandEndFrame;
 struct FRHICommandEndOcclusionQueryBatch;
 struct FRHICommandEndRenderQuery;
@@ -304,34 +302,6 @@ void FRHICommandBindClearMRTValues::Execute(FRHICommandListBase& CmdList)
 		);
 }
 
-void FRHICommandEndDrawPrimitiveUP::Execute(FRHICommandListBase& CmdList)
-{
-	RHISTAT(EndDrawPrimitiveUP);
-	void* Buffer = NULL;
-	INTERNAL_DECORATOR(RHIBeginDrawPrimitiveUP)(NumPrimitives, NumVertices, VertexDataStride, Buffer);
-	FMemory::Memcpy(Buffer, OutVertexData, NumVertices * VertexDataStride);
-	INTERNAL_DECORATOR(RHIEndDrawPrimitiveUP)();
-}
-
-void FRHICommandEndDrawIndexedPrimitiveUP::Execute(FRHICommandListBase& CmdList)
-{
-	RHISTAT(EndDrawIndexedPrimitiveUP);
-	void* VertexBuffer = nullptr;
-	void* IndexBuffer = nullptr;
-	INTERNAL_DECORATOR(RHIBeginDrawIndexedPrimitiveUP)(
-		NumPrimitives,
-		NumVertices,
-		VertexDataStride,
-		VertexBuffer,
-		MinVertexIndex,
-		NumIndices,
-		IndexDataStride,
-		IndexBuffer);
-	FMemory::Memcpy(VertexBuffer, OutVertexData, NumVertices * VertexDataStride);
-	FMemory::Memcpy(IndexBuffer, OutIndexData, NumIndices * IndexDataStride);
-	INTERNAL_DECORATOR(RHIEndDrawIndexedPrimitiveUP)();
-}
-
 template<ECmdList CmdListType>
 void FRHICommandSetComputeShader<CmdListType>::Execute(FRHICommandListBase& CmdList)
 {
diff --git a/Engine/Source/Runtime/RHI/Public/RHIContext.h b/Engine/Source/Runtime/RHI/Public/RHIContext.h
--- a/Engine/Source/Runtime/RHI/Public/RHIContext.h
+++ b/Engine/Source/Runtime/RHI/Public/RHIContext.h
@@ -55,7 +55,7 @@ public:
 
 	virtual void RHIDispatchIndirectComputeShader(FVertexBufferRHIParamRef ArgumentBuffer, uint32 ArgumentOffset) = 0;
 
-	virtual void RHISetAsyncComputeBudget(EAsyncComputeBudget Budget) = 0;
+	virtual void RHISetAsyncComputeBudget(EAsyncComputeBudget Budget) {}
 
 	/**
 	* Explicitly transition a UAV from readable -> writable by the GPU or vice versa.
@@ -117,7 +117,7 @@ public:
 	 * Some RHI implementations (OpenGL) cache render state internally
 	 * Signal to RHI that cached state is no longer valid
 	 */
-	virtual void RHIInvalidateCachedState() {};
+	virtual void RHIInvalidateCachedState() {}
 
 	/**
 	 * Performs a copy of the data in 'SourceBuffer' to 'DestinationStagingBuffer.' This will occur inline on the GPU timeline. This is a mechanism to perform nonblocking readback of a buffer at a point in time.
@@ -139,7 +139,6 @@ public:
 	{
 		check(false);
 	}
-
 };
 
 struct FAccelerationStructureUpdateParams
@@ -176,20 +175,10 @@ public:
 		}
 	}
 
-	/**
-	*Sets the current compute shader.  Mostly for compliance with platforms
-	*that require shader setting before resource binding.
-	*/
-	virtual void RHISetComputeShader(FComputeShaderRHIParamRef ComputeShader) = 0;
-
 	virtual void RHIDispatchComputeShader(uint32 ThreadGroupCountX, uint32 ThreadGroupCountY, uint32 ThreadGroupCountZ) = 0;
 
 	virtual void RHIDispatchIndirectComputeShader(FVertexBufferRHIParamRef ArgumentBuffer, uint32 ArgumentOffset) = 0;
 
-	virtual void RHISetAsyncComputeBudget(EAsyncComputeBudget Budget) override
-	{
-	}
-
 	virtual void RHIAutomaticCacheFlushAfterComputeShader(bool bEnable) = 0;
 
 	virtual void RHIFlushComputeShaderCache() = 0;
@@ -259,8 +248,6 @@ public:
 
 	virtual void RHIEndRenderQuery(FRenderQueryRHIParamRef RenderQuery) = 0;
 
-	virtual void RHISubmitCommandsHint() = 0;
-
 	// Used for OpenGL to check and see if any occlusion queries can be read back on the RHI thread. If they aren't ready when we need them, then we end up stalling.
 	virtual void RHIPollOcclusionQueries()
 	{
@@ -268,7 +255,7 @@ public:
 	}
 
 	// Not all RHIs need this (Mobile specific)
-	virtual void RHIDiscardRenderTargets(bool Depth, bool Stencil, uint32 ColorBitMask) {};
+	virtual void RHIDiscardRenderTargets(bool Depth, bool Stencil, uint32 ColorBitMask) {}
 
 	// This method is queued with an RHIThread, otherwise it will flush after it is queued; without an RHI thread there is no benefit to queuing this frame advance commands
 	virtual void RHIBeginDrawingViewport(FViewportRHIParamRef Viewport, FTextureRHIParamRef RenderTargetRHI) = 0;
@@ -483,38 +470,6 @@ public:
 
 	virtual void RHIDrawIndexedPrimitiveIndirect(FIndexBufferRHIParamRef IndexBuffer, FVertexBufferRHIParamRef ArgumentBuffer, uint32 ArgumentOffset) = 0;
 
-	/**
-	* Preallocate memory or get a direct command stream pointer to fill up for immediate rendering . This avoids memcpys below in DrawPrimitiveUP
-	* @param NumPrimitives The number of primitives in the VertexData buffer
-	* @param NumVertices The number of vertices to be written
-	* @param VertexDataStride Size of each vertex
-	* @param OutVertexData Reference to the allocated vertex memory
-	*/
-	virtual void RHIBeginDrawPrimitiveUP(uint32 NumPrimitives, uint32 NumVertices, uint32 VertexDataStride, void*& OutVertexData) = 0;
-
-	/**
-	* Draw a primitive using the vertex data populated since RHIBeginDrawPrimitiveUP and clean up any memory as needed
-	*/
-	virtual void RHIEndDrawPrimitiveUP() = 0;
-
-	/**
-	* Preallocate memory or get a direct command stream pointer to fill up for immediate rendering . This avoids memcpys below in DrawIndexedPrimitiveUP
-	* @param NumPrimitives The number of primitives in the VertexData buffer
-	* @param NumVertices The number of vertices to be written
-	* @param VertexDataStride Size of each vertex
-	* @param OutVertexData Reference to the allocated vertex memory
-	* @param MinVertexIndex The lowest vertex index used by the index buffer
-	* @param NumIndices Number of indices to be written
-	* @param IndexDataStride Size of each index (either 2 or 4 bytes)
-	* @param OutIndexData Reference to the allocated index memory
-	*/
-	virtual void RHIBeginDrawIndexedPrimitiveUP(uint32 NumPrimitives, uint32 NumVertices, uint32 VertexDataStride, void*& OutVertexData, uint32 MinVertexIndex, uint32 NumIndices, uint32 IndexDataStride, void*& OutIndexData) = 0;
-
-	/**
-	* Draw a primitive using the vertex and index data populated since RHIBeginDrawIndexedPrimitiveUP and clean up any memory as needed
-	*/
-	virtual void RHIEndDrawIndexedPrimitiveUP() = 0;
-
 	/**
 	* Sets Depth Bounds range with the given min/max depth.
 	* @param MinDepth	The minimum depth for depth bounds test
@@ -523,10 +478,6 @@ public:
 	*/
 	virtual void RHISetDepthBounds(float MinDepth, float MaxDepth) = 0;
 
-	virtual void RHIPushEvent(const TCHAR* Name, FColor Color) = 0;
-
-	virtual void RHIPopEvent() = 0;
-
 	virtual void RHIUpdateTextureReference(FTextureReferenceRHIParamRef TextureRef, FTextureRHIParamRef NewTexture) = 0;
 
 	virtual void RHIBeginRenderPass(const FRHIRenderPassInfo& InInfo, const TCHAR* InName)
diff --git a/Engine/Source/Runtime/RHI/Public/RHIUtilities.h b/Engine/Source/Runtime/RHI/Public/RHIUtilities.h
--- a/Engine/Source/Runtime/RHI/Public/RHIUtilities.h
+++ b/Engine/Source/Runtime/RHI/Public/RHIUtilities.h
@@ -848,76 +848,6 @@ inline uint32 GetVertexCountForPrimitiveCount(uint32 NumPrimitives, uint32 Primi
 	return VertexCount;
 }
 
-/**
- * Draw a primitive using the vertices passed in.
- * @param PrimitiveType The type (triangles, lineloop, etc) of primitive to draw
- * @param NumPrimitives The number of primitives in the VertexData buffer
- * @param VertexData A reference to memory preallocate in RHIBeginDrawPrimitiveUP
- * @param VertexDataStride Size of each vertex
- */
-UE_DEPRECATED(4.21, "This function is deprecated and will be removed in future releases.")
-inline void DrawPrimitiveUP(FRHICommandList& RHICmdList, uint32 PrimitiveType, uint32 NumPrimitives, const void* VertexData, uint32 VertexDataStride)
-{
-	check(NumPrimitives > 0);
-	const uint32 VertexCount = GetVertexCountForPrimitiveCount( NumPrimitives, PrimitiveType );
-
-	FRHIResourceCreateInfo CreateInfo;
-	FVertexBufferRHIRef VertexBufferRHI = RHICreateVertexBuffer(VertexDataStride * VertexCount, BUF_Volatile, CreateInfo);
-	void* VoidPtr = RHILockVertexBuffer(VertexBufferRHI, 0, VertexDataStride * VertexCount, RLM_WriteOnly);
-	FPlatformMemory::Memcpy(VoidPtr, VertexData, VertexDataStride * VertexCount);
-	RHIUnlockVertexBuffer(VertexBufferRHI);
-
-	RHICmdList.SetStreamSource(0, VertexBufferRHI, 0);
-	PRAGMA_DISABLE_DEPRECATION_WARNINGS
-	RHICmdList.DrawPrimitive(0, NumPrimitives, 1);
-	PRAGMA_ENABLE_DEPRECATION_WARNINGS
-
-	VertexBufferRHI.SafeRelease();
-}
-
-/**
- * Draw a primitive using the vertices passed in as described the passed in indices. 
- * @param PrimitiveType The type (triangles, lineloop, etc) of primitive to draw
- * @param MinVertexIndex The lowest vertex index used by the index buffer
- * @param NumVertices The number of vertices in the vertex buffer
- * @param NumPrimitives THe number of primitives described by the index buffer
- * @param IndexData The memory preallocated in RHIBeginDrawIndexedPrimitiveUP
- * @param IndexDataStride The size of one index
- * @param VertexData The memory preallocate in RHIBeginDrawIndexedPrimitiveUP
- * @param VertexDataStride The size of one vertex
- */
-UE_DEPRECATED(4.21, "This function is deprecated and will be removed in future releases.")
-inline void DrawIndexedPrimitiveUP(
-	FRHICommandList& RHICmdList,
-	uint32 PrimitiveType,
-	uint32 MinVertexIndex,
-	uint32 NumVertices,
-	uint32 NumPrimitives,
-	const void* IndexData,
-	uint32 IndexDataStride,
-	const void* VertexData,
-	uint32 VertexDataStride )
-{
-	const uint32 NumIndices = GetVertexCountForPrimitiveCount( NumPrimitives, PrimitiveType );
-
-	FRHIResourceCreateInfo CreateInfo;
-	FVertexBufferRHIRef VertexBufferRHI = RHICreateVertexBuffer(VertexDataStride * NumVertices, BUF_Volatile, CreateInfo);
-	void* VoidPtr = RHILockVertexBuffer(VertexBufferRHI, 0, VertexDataStride * NumVertices, RLM_WriteOnly);
-	FPlatformMemory::Memcpy(VoidPtr, VertexData, VertexDataStride * NumVertices);
-	RHIUnlockVertexBuffer(VertexBufferRHI);
-
-	FIndexBufferRHIRef IndexBufferRHI = RHICreateIndexBuffer(IndexDataStride, IndexDataStride * NumIndices, BUF_Volatile, CreateInfo);
-	void* VoidPtr2 = RHILockIndexBuffer(IndexBufferRHI, 0, IndexDataStride * NumIndices, RLM_WriteOnly);
-	FPlatformMemory::Memcpy(VoidPtr2, IndexData, IndexDataStride * NumIndices);
-	RHIUnlockIndexBuffer(IndexBufferRHI);
-
-	RHICmdList.SetStreamSource(0, VertexBufferRHI, 0);
-	RHICmdList.DrawIndexedPrimitive(IndexBufferRHI, MinVertexIndex, 0, NumVertices, 0, NumPrimitives, 1);
-
-	IndexBufferRHI.SafeRelease();
-	VertexBufferRHI.SafeRelease();
-}
-
 inline uint32 ComputeAnisotropyRT(int32 InitializerMaxAnisotropy)
 {
 	static const auto CVar = IConsoleManager::Get().FindTConsoleVariableDataInt(TEXT("r.MaxAnisotropy"));
diff --git a/Engine/Source/Runtime/VulkanRHI/Private/VulkanCommands.cpp b/Engine/Source/Runtime/VulkanRHI/Private/VulkanCommands.cpp
--- a/Engine/Source/Runtime/VulkanRHI/Private/VulkanCommands.cpp
+++ b/Engine/Source/Runtime/VulkanRHI/Private/VulkanCommands.cpp
@@ -868,86 +868,6 @@ void FVulkanCommandListContext::RHIDrawIndexedPrimitiveIndirect(FIndexBufferRHIP
 	}
 }
 
-void FVulkanCommandListContext::RHIBeginDrawPrimitiveUP(uint32 NumPrimitives, uint32 NumVertices, uint32 VertexDataStride, void*& OutVertexData)
-{
-#if VULKAN_ENABLE_AGGRESSIVE_STATS
-	SCOPE_CYCLE_COUNTER(STAT_VulkanUPPrepTime);
-#endif
-
-//	checkSlow(GPendingDrawPrimitiveUPVertexData == nullptr);
-
-	TempFrameAllocationBuffer.Alloc(VertexDataStride * NumVertices, VertexDataStride, UserPrimitive.VertexAllocInfo);
-	OutVertexData = UserPrimitive.VertexAllocInfo.Data;
-
-	UserPrimitive.NumPrimitives = NumPrimitives;
-	UserPrimitive.NumVertices = NumVertices;
-	UserPrimitive.VertexDataStride = VertexDataStride;
-}
-
-
-void FVulkanCommandListContext::RHIEndDrawPrimitiveUP()
-{
-#if VULKAN_ENABLE_AGGRESSIVE_STATS
-	SCOPE_CYCLE_COUNTER(STAT_VulkanDrawCallTime);
-#endif
-	RHI_DRAW_CALL_STATS(PendingGfxState->PrimitiveType, UserPrimitive.NumPrimitives);
-	PendingGfxState->SetStreamSource(0, UserPrimitive.VertexAllocInfo.GetHandle(), UserPrimitive.VertexAllocInfo.GetBindOffset());
-	FVulkanCmdBuffer* CmdBuffer = CommandBufferManager->GetActiveCmdBuffer();
-	PendingGfxState->PrepareForDraw(CmdBuffer);
-	VkCommandBuffer Cmd = CmdBuffer->GetHandle();
-	VulkanRHI::vkCmdDraw(CmdBuffer->GetHandle(), UserPrimitive.NumVertices, 1, UserPrimitive.MinVertexIndex, 0);
-
-	if (FVulkanPlatform::RegisterGPUWork() && IsImmediate())
-	{
-		GpuProfiler.RegisterGPUWork(UserPrimitive.NumPrimitives, UserPrimitive.NumVertices);
-	}
-}
-
-void FVulkanCommandListContext::RHIBeginDrawIndexedPrimitiveUP(uint32 NumPrimitives, uint32 NumVertices, uint32 VertexDataStride,
-	void*& OutVertexData, uint32 MinVertexIndex, uint32 NumIndices, uint32 IndexDataStride, void*& OutIndexData)
-{
-	LLM_SCOPE_VULKAN(ELLMTagVulkan::VulkanMisc);
-#if VULKAN_ENABLE_AGGRESSIVE_STATS
-	SCOPE_CYCLE_COUNTER(STAT_VulkanUPPrepTime);
-#endif
-
-	TempFrameAllocationBuffer.Alloc(VertexDataStride * NumVertices, IndexDataStride, UserPrimitive.VertexAllocInfo);
-	OutVertexData = UserPrimitive.VertexAllocInfo.Data;
-
-	check(IndexDataStride == 2 || IndexDataStride == 4);
-	UserPrimitive.IndexType = (IndexDataStride == 2) ? VK_INDEX_TYPE_UINT16 : VK_INDEX_TYPE_UINT32;
-	TempFrameAllocationBuffer.Alloc(IndexDataStride * NumIndices, IndexDataStride, UserPrimitive.IndexAllocInfo);
-	OutIndexData = UserPrimitive.IndexAllocInfo.Data;
-
-	UserPrimitive.NumPrimitives = NumPrimitives;
-	UserPrimitive.MinVertexIndex = MinVertexIndex;
-	UserPrimitive.IndexDataStride = IndexDataStride;
-
-	UserPrimitive.NumVertices = NumVertices;
-	UserPrimitive.VertexDataStride = VertexDataStride;
-}
-
-void FVulkanCommandListContext::RHIEndDrawIndexedPrimitiveUP()
-{
-	LLM_SCOPE_VULKAN(ELLMTagVulkan::VulkanMisc);
-#if VULKAN_ENABLE_AGGRESSIVE_STATS
-	SCOPE_CYCLE_COUNTER(STAT_VulkanDrawCallTime);
-#endif
-	RHI_DRAW_CALL_STATS(PendingGfxState->PrimitiveType, UserPrimitive.NumPrimitives);
-	PendingGfxState->SetStreamSource(0, UserPrimitive.VertexAllocInfo.GetHandle(), UserPrimitive.VertexAllocInfo.GetBindOffset());
-	FVulkanCmdBuffer* CmdBuffer = CommandBufferManager->GetActiveCmdBuffer();
-	PendingGfxState->PrepareForDraw(CmdBuffer);
-	VkCommandBuffer Cmd = CmdBuffer->GetHandle();
-	uint32 NumIndices = GetVertexCountForPrimitiveCount(UserPrimitive.NumPrimitives, PendingGfxState->PrimitiveType);
-	VulkanRHI::vkCmdBindIndexBuffer(Cmd, UserPrimitive.IndexAllocInfo.GetHandle(), UserPrimitive.IndexAllocInfo.GetBindOffset(), UserPrimitive.IndexType);
-	VulkanRHI::vkCmdDrawIndexed(Cmd, NumIndices, 1, UserPrimitive.MinVertexIndex, 0, 0);
-
-	if (FVulkanPlatform::RegisterGPUWork() && IsImmediate())
-	{
-		GpuProfiler.RegisterGPUWork(UserPrimitive.NumPrimitives, UserPrimitive.NumVertices);
-	}
-}
-
 void FVulkanCommandListContext::RHIClearMRT(bool bClearColor, int32 NumClearColors, const FLinearColor* ClearColorArray, bool bClearDepth, float Depth, bool bClearStencil, uint32 Stencil)
 {
 	if (!(bClearColor || bClearDepth || bClearStencil))
diff --git a/Engine/Source/Runtime/VulkanRHI/Public/VulkanContext.h b/Engine/Source/Runtime/VulkanRHI/Public/VulkanContext.h
--- a/Engine/Source/Runtime/VulkanRHI/Public/VulkanContext.h
+++ b/Engine/Source/Runtime/VulkanRHI/Public/VulkanContext.h
@@ -213,10 +213,6 @@ public:
 	virtual void RHIDrawIndexedIndirect(FIndexBufferRHIParamRef IndexBufferRHI, FStructuredBufferRHIParamRef ArgumentsBufferRHI, int32 DrawArgumentsIndex, uint32 NumInstances) final override;
 	virtual void RHIDrawIndexedPrimitive(FIndexBufferRHIParamRef IndexBuffer, int32 BaseVertexIndex, uint32 FirstInstance, uint32 NumVertices, uint32 StartIndex, uint32 NumPrimitives, uint32 NumInstances) final override;
 	virtual void RHIDrawIndexedPrimitiveIndirect(FIndexBufferRHIParamRef IndexBuffer, FVertexBufferRHIParamRef ArgumentBuffer, uint32 ArgumentOffset) final override;
-	virtual void RHIBeginDrawPrimitiveUP(uint32 NumPrimitives, uint32 NumVertices, uint32 VertexDataStride, void*& OutVertexData) final override;
-	virtual void RHIEndDrawPrimitiveUP() final override;
-	virtual void RHIBeginDrawIndexedPrimitiveUP(uint32 NumPrimitives, uint32 NumVertices, uint32 VertexDataStride, void*& OutVertexData, uint32 MinVertexIndex, uint32 NumIndices, uint32 IndexDataStride, void*& OutIndexData) final override;
-	virtual void RHIEndDrawIndexedPrimitiveUP() final override;
 	virtual void RHISetDepthBounds(float MinDepth, float MaxDepth) final override;
 	virtual void RHIPushEvent(const TCHAR* Name, FColor Color) final override;
 	virtual void RHIPopEvent() final override;
@@ -386,20 +382,6 @@ protected:
 
 	void SetShaderUniformBuffer(ShaderStage::EStage Stage, const FVulkanUniformBuffer* UniformBuffer, int32 ParameterIndex, const FVulkanShader* Shader);
 
-	struct
-	{
-		VulkanRHI::FTempFrameAllocationBuffer::FTempAllocInfo VertexAllocInfo;
-		uint32 NumVertices = 0;
-		uint32 VertexDataStride = 0;
-
-		VulkanRHI::FTempFrameAllocationBuffer::FTempAllocInfo IndexAllocInfo;
-		VkIndexType IndexType = VK_INDEX_TYPE_MAX_ENUM;
-		uint32 NumPrimitives = 0;
-		uint32 MinVertexIndex = 0;
-		uint32 IndexDataStride = 0;
-	} UserPrimitive;
-
-
 	VulkanRHI::FTempFrameAllocationBuffer TempFrameAllocationBuffer;
 
 	TArray<FString> EventStack;
diff --git a/Engine/Source/Runtime/Windows/D3D11RHI/Private/D3D11Commands.cpp b/Engine/Source/Runtime/Windows/D3D11RHI/Private/D3D11Commands.cpp
--- a/Engine/Source/Runtime/Windows/D3D11RHI/Private/D3D11Commands.cpp
+++ b/Engine/Source/Runtime/Windows/D3D11RHI/Private/D3D11Commands.cpp
@@ -1707,121 +1707,6 @@ void FD3D11DynamicRHI::RHIDrawIndexedPrimitiveIndirect(FIndexBufferRHIParamRef I
 	Direct3DDeviceIMContext->DrawIndexedInstancedIndirect(ArgumentBuffer->Resource,ArgumentOffset);
 }
 
-/**
- * Preallocate memory or get a direct command stream pointer to fill up for immediate rendering . This avoids memcpys below in DrawPrimitiveUP
- * @param PrimitiveType The type (triangles, lineloop, etc) of primitive to draw
- * @param NumPrimitives The number of primitives in the VertexData buffer
- * @param NumVertices The number of vertices to be written
- * @param VertexDataStride Size of each vertex 
- * @param OutVertexData Reference to the allocated vertex memory
- */
-void FD3D11DynamicRHI::RHIBeginDrawPrimitiveUP(uint32 NumPrimitives, uint32 NumVertices, uint32 VertexDataStride, void*& OutVertexData)
-{
-	checkSlow( PendingNumVertices == 0 );
-
-	// Remember the parameters for the draw call.
-	PendingNumPrimitives = NumPrimitives;
-	PendingNumVertices = NumVertices;
-	PendingVertexDataStride = VertexDataStride;
-
-	// Map the dynamic buffer.
-	OutVertexData = DynamicVB->Lock(NumVertices * VertexDataStride);
-}
-
-/**
- * Draw a primitive using the vertex data populated since RHIBeginDrawPrimitiveUP and clean up any memory as needed
- */
-void FD3D11DynamicRHI::RHIEndDrawPrimitiveUP()
-{
-	RHI_DRAW_CALL_STATS(PrimitiveType,PendingNumPrimitives);
-
-	checkSlow(!bUsingTessellation || PrimitiveType == PT_TriangleList);
-
-	GPUProfilingData.RegisterGPUWork(PendingNumPrimitives,PendingNumVertices);
-
-	// Unmap the dynamic vertex buffer.
-	ID3D11Buffer* D3DBuffer = DynamicVB->Unlock();
-	uint32 VBOffset = 0;
-
-	// Issue the draw call.
-	CommitGraphicsResourceTables();
-	CommitNonComputeShaderConstants();
-	TrackResourceBoundAsVB(nullptr, 0);
-	StateCache.SetStreamSource(D3DBuffer, 0, PendingVertexDataStride, VBOffset);
-	StateCache.SetPrimitiveTopology(GetD3D11PrimitiveType(PrimitiveType,bUsingTessellation));
-	Direct3DDeviceIMContext->Draw(PendingNumVertices,0);
-
-	// Clear these parameters.
-	PendingNumPrimitives = 0;
-	PendingNumVertices = 0;
-	PendingVertexDataStride = 0;
-}
-
-/**
- * Preallocate memory or get a direct command stream pointer to fill up for immediate rendering . This avoids memcpys below in DrawIndexedPrimitiveUP
- * @param PrimitiveType The type (triangles, lineloop, etc) of primitive to draw
- * @param NumPrimitives The number of primitives in the VertexData buffer
- * @param NumVertices The number of vertices to be written
- * @param VertexDataStride Size of each vertex
- * @param OutVertexData Reference to the allocated vertex memory
- * @param MinVertexIndex The lowest vertex index used by the index buffer
- * @param NumIndices Number of indices to be written
- * @param IndexDataStride Size of each index (either 2 or 4 bytes)
- * @param OutIndexData Reference to the allocated index memory
- */
-void FD3D11DynamicRHI::RHIBeginDrawIndexedPrimitiveUP(uint32 NumPrimitives, uint32 NumVertices, uint32 VertexDataStride, void*& OutVertexData, uint32 MinVertexIndex, uint32 NumIndices, uint32 IndexDataStride, void*& OutIndexData)
-{
-	checkSlow((sizeof(uint16) == IndexDataStride) || (sizeof(uint32) == IndexDataStride));
-
-	// Store off information needed for the draw call.
-	PendingNumPrimitives = NumPrimitives;
-	PendingMinVertexIndex = MinVertexIndex;
-	PendingIndexDataStride = IndexDataStride;
-	PendingNumVertices = NumVertices;
-	PendingNumIndices = NumIndices;
-	PendingVertexDataStride = VertexDataStride;
-
-	// Map dynamic vertex and index buffers.
-	OutVertexData = DynamicVB->Lock(NumVertices * VertexDataStride);
-	OutIndexData = DynamicIB->Lock(NumIndices * IndexDataStride);
-}
-
-/**
- * Draw a primitive using the vertex and index data populated since RHIBeginDrawIndexedPrimitiveUP and clean up any memory as needed
- */
-void FD3D11DynamicRHI::RHIEndDrawIndexedPrimitiveUP()
-{
-	// tessellation only supports trilists
-	checkSlow(!bUsingTessellation || PrimitiveType == PT_TriangleList);
-
-	RHI_DRAW_CALL_STATS(PrimitiveType,PendingNumPrimitives);
-
-	GPUProfilingData.RegisterGPUWork(PendingNumPrimitives,PendingNumVertices);
-
-	// Unmap the dynamic buffers.
-	ID3D11Buffer* VertexBuffer = DynamicVB->Unlock();
-	ID3D11Buffer* IndexBuffer = DynamicIB->Unlock();
-	uint32 VBOffset = 0;
-
-	// Issue the draw call.
-	CommitGraphicsResourceTables();
-	CommitNonComputeShaderConstants();
-	TrackResourceBoundAsVB(nullptr, 0);
-	TrackResourceBoundAsIB(nullptr);
-	StateCache.SetStreamSource(VertexBuffer, 0, PendingVertexDataStride, VBOffset);
-	StateCache.SetIndexBuffer(IndexBuffer, PendingIndexDataStride == sizeof(uint16) ? DXGI_FORMAT_R16_UINT : DXGI_FORMAT_R32_UINT, 0);
-	StateCache.SetPrimitiveTopology(GetD3D11PrimitiveType(PrimitiveType, bUsingTessellation));
-	Direct3DDeviceIMContext->DrawIndexed(PendingNumIndices,0,PendingMinVertexIndex);
-
-	// Clear these parameters.
-	PendingNumPrimitives = 0;
-	PendingMinVertexIndex = 0;
-	PendingIndexDataStride = 0;
-	PendingNumVertices = 0;
-	PendingNumIndices = 0;
-	PendingVertexDataStride = 0;
-}
-
 // Raster operations.
 void FD3D11DynamicRHI::RHIClearMRT(bool bClearColor, int32 NumClearColors, const FLinearColor* ClearColorArray, bool bClearDepth, float Depth, bool bClearStencil, uint32 Stencil)
 {
diff --git a/Engine/Source/Runtime/Windows/D3D11RHI/Private/D3D11Device.cpp b/Engine/Source/Runtime/Windows/D3D11RHI/Private/D3D11Device.cpp
--- a/Engine/Source/Runtime/Windows/D3D11RHI/Private/D3D11Device.cpp
+++ b/Engine/Source/Runtime/Windows/D3D11RHI/Private/D3D11Device.cpp
@@ -67,12 +67,6 @@ FD3D11DynamicRHI::FD3D11DynamicRHI(IDXGIFactory1* InDXGIFactory1,D3D_FEATURE_LEV
 	CurrentDSVAccessType(FExclusiveDepthStencil::DepthWrite_StencilWrite),
 	bDiscardSharedConstants(false),
 	bUsingTessellation(false),
-	PendingNumVertices(0),
-	PendingVertexDataStride(0),
-	PendingNumPrimitives(0),
-	PendingMinVertexIndex(0),
-	PendingNumIndices(0),
-	PendingIndexDataStride(0),
 	GPUProfilingData(this),
 	ChosenAdapter(InChosenAdapter),
 	ChosenDescription(InChosenDescription)
diff --git a/Engine/Source/Runtime/Windows/D3D11RHI/Private/D3D11RHIPrivate.h b/Engine/Source/Runtime/Windows/D3D11RHI/Private/D3D11RHIPrivate.h
--- a/Engine/Source/Runtime/Windows/D3D11RHI/Private/D3D11RHIPrivate.h
+++ b/Engine/Source/Runtime/Windows/D3D11RHI/Private/D3D11RHIPrivate.h
@@ -573,10 +573,6 @@ public:
 	virtual void RHIDrawIndexedIndirect(FIndexBufferRHIParamRef IndexBufferRHI, FStructuredBufferRHIParamRef ArgumentsBufferRHI, int32 DrawArgumentsIndex, uint32 NumInstances) final override;
 	virtual void RHIDrawIndexedPrimitive(FIndexBufferRHIParamRef IndexBuffer, int32 BaseVertexIndex, uint32 FirstInstance, uint32 NumVertices, uint32 StartIndex, uint32 NumPrimitives, uint32 NumInstances) final override;
 	virtual void RHIDrawIndexedPrimitiveIndirect(FIndexBufferRHIParamRef IndexBuffer, FVertexBufferRHIParamRef ArgumentBuffer, uint32 ArgumentOffset) final override;
-	virtual void RHIBeginDrawPrimitiveUP(uint32 NumPrimitives, uint32 NumVertices, uint32 VertexDataStride, void*& OutVertexData) final override;
-	virtual void RHIEndDrawPrimitiveUP() final override;
-	virtual void RHIBeginDrawIndexedPrimitiveUP(uint32 NumPrimitives, uint32 NumVertices, uint32 VertexDataStride, void*& OutVertexData, uint32 MinVertexIndex, uint32 NumIndices, uint32 IndexDataStride, void*& OutIndexData) final override;
-	virtual void RHIEndDrawIndexedPrimitiveUP() final override;
 	virtual void RHIEnableDepthBoundsTest(bool bEnable) final override
 	{
 		if (GSupportsDepthBoundsTest && StateCache.bDepthBoundsEnabled != bEnable)
@@ -874,14 +870,6 @@ protected:
 	TRefCountPtr<FD3D11DynamicBuffer> DynamicVB;
 	TRefCountPtr<FD3D11DynamicBuffer> DynamicIB;
 
-	// State for begin/end draw primitive UP interface.
-	uint32 PendingNumVertices;
-	uint32 PendingVertexDataStride;
-	uint32 PendingNumPrimitives;
-	uint32 PendingMinVertexIndex;
-	uint32 PendingNumIndices;
-	uint32 PendingIndexDataStride;
-
 	/** A list of all D3D constant buffers RHIs that have been created. */
 	TArray<TRefCountPtr<FD3D11ConstantBuffer> > VSConstantBuffers;
 	TArray<TRefCountPtr<FD3D11ConstantBuffer> > HSConstantBuffers;
