Implemented RHICopySubTextureRegion interface function for OpenGL backend.

diff --git a/Engine/Source/Runtime/OpenGLDrv/Private/Android/AndroidOpenGL.h b/Engine/Source/Runtime/OpenGLDrv/Private/Android/AndroidOpenGL.h
--- a/Engine/Source/Runtime/OpenGLDrv/Private/Android/AndroidOpenGL.h
+++ b/Engine/Source/Runtime/OpenGLDrv/Private/Android/AndroidOpenGL.h
@@ -416,6 +416,11 @@ struct FAndroidOpenGL : public FOpenGLES2
 		glTexSubImage3D(Target, Level, XOffset, YOffset, ZOffset, Width, Height, Depth, Format, Type, PixelData);
 	}
 
+	static FORCEINLINE void	CopyTexSubImage2D(GLenum Target, GLint Level, GLint XOffset, GLint YOffset, GLint X, GLint Y, GLsizei Width, GLsizei Height)
+	{
+		glCopyTexSubImage2D(Target, Level, XOffset, YOffset, X, Y, Width, Height);
+	}
+
 	static FORCEINLINE void	CopyTexSubImage3D(GLenum Target, GLint Level, GLint XOffset, GLint YOffset, GLint ZOffset, GLint X, GLint Y, GLsizei Width, GLsizei Height)
 	{
 		glCopyTexSubImage3D(Target, Level, XOffset, YOffset, ZOffset, X, Y, Width, Height);
diff --git a/Engine/Source/Runtime/OpenGLDrv/Private/Lumin/LuminOpenGL.h b/Engine/Source/Runtime/OpenGLDrv/Private/Lumin/LuminOpenGL.h
--- a/Engine/Source/Runtime/OpenGLDrv/Private/Lumin/LuminOpenGL.h
+++ b/Engine/Source/Runtime/OpenGLDrv/Private/Lumin/LuminOpenGL.h
@@ -323,6 +323,11 @@ struct FLuminOpenGL : public FOpenGLES2
 		glTexSubImage3D(Target, Level, XOffset, YOffset, ZOffset, Width, Height, Depth, Format, Type, PixelData);
 	}
 
+	static FORCEINLINE void	CopyTexSubImage2D(GLenum Target, GLint Level, GLint XOffset, GLint YOffset, GLint X, GLint Y, GLsizei Width, GLsizei Height)
+	{
+		glCopyTexSubImage2D(Target, Level, XOffset, YOffset, X, Y, Width, Height);
+	}
+
 	static FORCEINLINE void	CopyTexSubImage3D(GLenum Target, GLint Level, GLint XOffset, GLint YOffset, GLint ZOffset, GLint X, GLint Y, GLsizei Width, GLsizei Height)
 	{
 		glCopyTexSubImage3D(Target, Level, XOffset, YOffset, ZOffset, X, Y, Width, Height);
diff --git a/Engine/Source/Runtime/OpenGLDrv/Private/OpenGLRenderTarget.cpp b/Engine/Source/Runtime/OpenGLDrv/Private/OpenGLRenderTarget.cpp
--- a/Engine/Source/Runtime/OpenGLDrv/Private/OpenGLRenderTarget.cpp
+++ b/Engine/Source/Runtime/OpenGLDrv/Private/OpenGLRenderTarget.cpp
@@ -899,7 +899,8 @@ void FOpenGLDynamicRHI::RHIReadSurfaceData(FTextureRHIParamRef TextureRHI, FIntR
 	VERIFY_GL_SCOPE();
 
 	// Verify requirements, but don't crash
-	if (!ensure(FOpenGL::SupportsFloatReadSurface()) || !ensure(TextureRHI) || !ensure(TextureRHI->GetFormat() == PF_A32B32G32R32F))
+	// Ignore texture format here, GL will convert it for us in glReadPixels
+	if (!ensure(FOpenGL::SupportsFloatReadSurface()) || !ensure(TextureRHI))
 	{
 		return;
 	}
diff --git a/Engine/Source/Runtime/OpenGLDrv/Private/OpenGLTexture.cpp b/Engine/Source/Runtime/OpenGLDrv/Private/OpenGLTexture.cpp
--- a/Engine/Source/Runtime/OpenGLDrv/Private/OpenGLTexture.cpp
+++ b/Engine/Source/Runtime/OpenGLDrv/Private/OpenGLTexture.cpp
@@ -2554,6 +2554,41 @@ void FOpenGLDynamicRHI::RHIUpdateTextureReference(FTextureReferenceRHIParamRef T
 	}
 }
 
+void FOpenGLDynamicRHI::RHICopySubTextureRegion(FTexture2DRHIParamRef SourceTextureRHI, FTexture2DRHIParamRef DestinationTextureRHI, FBox2D SourceBox, FBox2D DestinationBox)
+{
+	VERIFY_GL_SCOPE();
+	FOpenGLTexture2D* SourceTexture = ResourceCast(SourceTextureRHI);
+	FOpenGLTexture2D* DestinationTexture = ResourceCast(DestinationTextureRHI);
+
+	check(SourceTexture->Target == DestinationTexture->Target);
+
+	// Use a texture stage that's not likely to be used for draws, to avoid waiting
+	FOpenGLContextState& ContextState = GetContextStateForCurrentContext();
+	CachedSetupTextureStage(ContextState, FOpenGL::GetMaxCombinedTextureImageUnits() - 1, DestinationTexture->Target, DestinationTexture->Resource, 0, DestinationTexture->GetNumMips());
+	CachedBindPixelUnpackBuffer(ContextState, 0);
+
+	// Convert sub texture regions to GL types
+	GLint XOffset = DestinationBox.Min.X;
+	GLint YOffset = DestinationBox.Min.Y;
+	GLint X = SourceBox.Min.X;
+	GLint Y = SourceBox.Min.Y;
+	GLsizei Width = DestinationBox.Max.X - DestinationBox.Min.X;
+	GLsizei Height = DestinationBox.Max.Y - DestinationBox.Min.Y;
+
+	// Bind source texture to an FBO to read from
+	FOpenGLTextureBase* RenderTarget[] = { SourceTexture };
+	uint32 MipLevel = 0;
+	GLuint SourceFBO = GetOpenGLFramebuffer(1, RenderTarget, NULL, &MipLevel, NULL);
+	check(SourceFBO != 0);
+
+	glBindFramebuffer(GL_FRAMEBUFFER, SourceFBO);
+
+	FOpenGL::ReadBuffer(GL_COLOR_ATTACHMENT0);
+	FOpenGL::CopyTexSubImage2D(DestinationTexture->Target, 0, XOffset, YOffset, X, Y, Width, Height);
+
+	ContextState.Framebuffer = (GLuint)-1;
+}
+
 FTexture2DRHIRef FOpenGLDynamicRHI::RHICreateTexture2DFromResource(EPixelFormat Format, uint32 SizeX, uint32 SizeY, uint32 NumMips, uint32 NumSamples, uint32 NumSamplesTileMem, const FClearValueBinding& ClearValueBinding, GLuint Resource, uint32 TexCreateFlags)
 {
 	FOpenGLTexture2D* Texture2D = new FOpenGLTexture2D(
diff --git a/Engine/Source/Runtime/OpenGLDrv/Public/OpenGL.h b/Engine/Source/Runtime/OpenGLDrv/Public/OpenGL.h
--- a/Engine/Source/Runtime/OpenGLDrv/Public/OpenGL.h
+++ b/Engine/Source/Runtime/OpenGLDrv/Public/OpenGL.h
@@ -295,6 +295,7 @@ public:
 	static FORCEINLINE void TexImage2DMultisample(GLenum Target, GLsizei Samples, GLint InternalFormat, GLsizei Width, GLsizei Height, GLboolean FixedSampleLocations) UGL_REQUIRED_VOID
 	static FORCEINLINE void TexBuffer(GLenum Target, GLenum InternalFormat, GLuint Buffer) UGL_REQUIRED_VOID
 	static FORCEINLINE void TexSubImage3D(GLenum Target, GLint Level, GLint XOffset, GLint YOffset, GLint ZOffset, GLsizei Width, GLsizei Height, GLsizei Depth, GLenum Format, GLenum Type, const GLvoid* PixelData) UGL_REQUIRED_VOID
+	static FORCEINLINE void	CopyTexSubImage2D(GLenum Target, GLint Level, GLint XOffset, GLint YOffset, GLint X, GLint Y, GLsizei Width, GLsizei Height) UGL_REQUIRED_VOID
 	static FORCEINLINE void	CopyTexSubImage3D(GLenum Target, GLint Level, GLint XOffset, GLint YOffset, GLint ZOffset, GLint X, GLint Y, GLsizei Width, GLsizei Height) UGL_REQUIRED_VOID
 	static FORCEINLINE void GetCompressedTexImage(GLenum Target, GLint Level, GLvoid* OutImageData) UGL_REQUIRED_VOID
 	static FORCEINLINE void GetTexImage(GLenum Target, GLint Level, GLenum Format, GLenum Type, GLvoid* OutPixelData) UGL_REQUIRED_VOID
diff --git a/Engine/Source/Runtime/OpenGLDrv/Public/OpenGL3.h b/Engine/Source/Runtime/OpenGLDrv/Public/OpenGL3.h
--- a/Engine/Source/Runtime/OpenGLDrv/Public/OpenGL3.h
+++ b/Engine/Source/Runtime/OpenGLDrv/Public/OpenGL3.h
@@ -411,6 +411,11 @@ struct FOpenGL3 : public FOpenGLBase
 		glTexSubImage3D(Target, Level, XOffset, YOffset, ZOffset, Width, Height, Depth, Format, Type, PixelData);
 	}
 
+	static FORCEINLINE void	CopyTexSubImage2D(GLenum Target, GLint Level, GLint XOffset, GLint YOffset, GLint X, GLint Y, GLsizei Width, GLsizei Height)
+	{
+		glCopyTexSubImage2D(Target, Level, XOffset, YOffset, X, Y, Width, Height);
+	}
+
 	static FORCEINLINE void	CopyTexSubImage3D(GLenum Target, GLint Level, GLint XOffset, GLint YOffset, GLint ZOffset, GLint X, GLint Y, GLsizei Width, GLsizei Height)
 	{
 		glCopyTexSubImage3D(Target, Level, XOffset, YOffset, ZOffset, X, Y, Width, Height);
diff --git a/Engine/Source/Runtime/OpenGLDrv/Public/OpenGLDrv.h b/Engine/Source/Runtime/OpenGLDrv/Public/OpenGLDrv.h
--- a/Engine/Source/Runtime/OpenGLDrv/Public/OpenGLDrv.h
+++ b/Engine/Source/Runtime/OpenGLDrv/Public/OpenGLDrv.h
@@ -524,6 +524,8 @@ public:
 	virtual void RHIInvalidateCachedState() final override;
 	virtual void RHIDiscardRenderTargets(bool Depth,bool Stencil,uint32 ColorBitMask) final override;
 
+	virtual void RHICopySubTextureRegion(FTexture2DRHIParamRef SourceTexture, FTexture2DRHIParamRef DestinationTexture, FBox2D SourceBox, FBox2D DestinationBox) final override;
+
 	// FOpenGLDynamicRHI interface.
 	virtual FTexture2DRHIRef RHICreateTexture2DFromResource(EPixelFormat Format, uint32 SizeX, uint32 SizeY, uint32 NumMips, uint32 NumSamples, uint32 NumSamplesTileMem, const FClearValueBinding& ClearValueBinding, GLuint Resource, uint32 Flags);
 	virtual FTexture2DRHIRef RHICreateTexture2DArrayFromResource(EPixelFormat Format, uint32 SizeX, uint32 SizeY, uint32 ArraySize, uint32 NumMips, uint32 NumSamples, uint32 NumSamplesTileMem, const FClearValueBinding& ClearValueBinding, GLuint Resource, uint32 Flags);
diff --git a/Engine/Source/Runtime/OpenGLDrv/Public/OpenGLESDeferred.h b/Engine/Source/Runtime/OpenGLDrv/Public/OpenGLESDeferred.h
--- a/Engine/Source/Runtime/OpenGLDrv/Public/OpenGLESDeferred.h
+++ b/Engine/Source/Runtime/OpenGLDrv/Public/OpenGLESDeferred.h
@@ -616,6 +616,11 @@ struct FOpenGLESDeferred : public FOpenGLBase
 		glTexSubImage3D(Target, Level, XOffset, YOffset, ZOffset, Width, Height, Depth, Format, Type, PixelData);
 	}
 
+	static FORCEINLINE void	CopyTexSubImage2D(GLenum Target, GLint Level, GLint XOffset, GLint YOffset, GLint X, GLint Y, GLsizei Width, GLsizei Height)
+	{
+		glCopyTexSubImage2D(Target, Level, XOffset, YOffset, X, Y, Width, Height);
+	}
+
 	static FORCEINLINE void	CopyTexSubImage3D(GLenum Target, GLint Level, GLint XOffset, GLint YOffset, GLint ZOffset, GLint X, GLint Y, GLsizei Width, GLsizei Height)
 	{
 		glCopyTexSubImage3D(Target, Level, XOffset, YOffset, ZOffset, X, Y, Width, Height);
