DR - HTML5 compile fix

diff --git a/Engine/Source/Runtime/OpenGLDrv/Private/HTML5/HTML5OpenGL.h b/Engine/Source/Runtime/OpenGLDrv/Private/HTML5/HTML5OpenGL.h
--- a/Engine/Source/Runtime/OpenGLDrv/Private/HTML5/HTML5OpenGL.h
+++ b/Engine/Source/Runtime/OpenGLDrv/Private/HTML5/HTML5OpenGL.h
@@ -51,7 +51,17 @@ void glGetObjectLabelEXT (GLenum type, GLuint object, GLsizei bufSize, GLsizei *
 #endif
 
 
+#ifndef GL_TEXTURE_1D
+#define GL_TEXTURE_1D			0x0DE0
+#endif
+
+#ifndef GL_TEXTURE_2D_ARRAY
+#define GL_TEXTURE_2D_ARRAY		0x8C1A
+#endif
 
+#ifndef GL_TEXTURE_RECTANGLE
+#define GL_TEXTURE_RECTANGLE	0x84F5
+#endif
 
 /** Unreal tokens that maps to different OpenGL tokens by platform. */
 #ifndef GL_DRAW_FRAMEBUFFER
@@ -92,6 +102,11 @@ struct FHTML5OpenGL : public FOpenGLES2
 		return true;
 	}
 
+	static FORCEINLINE void	CopyTexSubImage1D(GLenum Target, GLint Level, GLint XOffset, GLint X, GLint Y, GLsizei Width)
+	{
+		check(0);
+	}
+
 	static FORCEINLINE bool SupportsMapBuffer()							{ return false; }
 	static FORCEINLINE bool SupportsCombinedDepthStencilAttachment()	{ return bCombinedDepthStencilAttachment; }
 	static FORCEINLINE bool SupportsMultipleRenderTargets()				{ return bIsWebGL2; }
