DR - More HTML5 & Lumin fixes
#rnx

diff --git a/Engine/Source/Runtime/OpenGLDrv/Private/HTML5/HTML5OpenGL.h b/Engine/Source/Runtime/OpenGLDrv/Private/HTML5/HTML5OpenGL.h
--- a/Engine/Source/Runtime/OpenGLDrv/Private/HTML5/HTML5OpenGL.h
+++ b/Engine/Source/Runtime/OpenGLDrv/Private/HTML5/HTML5OpenGL.h
@@ -55,6 +55,10 @@ void glGetObjectLabelEXT (GLenum type, GLuint object, GLsizei bufSize, GLsizei *
 #define GL_TEXTURE_1D			0x0DE0
 #endif
 
+#ifndef GL_TEXTURE_1D_ARRAY
+#define GL_TEXTURE_1D_ARRAY		0x8C18
+#endif
+
 #ifndef GL_TEXTURE_2D_ARRAY
 #define GL_TEXTURE_2D_ARRAY		0x8C1A
 #endif
diff --git a/Engine/Source/Runtime/OpenGLDrv/Private/Lumin/LuminOpenGL.h b/Engine/Source/Runtime/OpenGLDrv/Private/Lumin/LuminOpenGL.h
--- a/Engine/Source/Runtime/OpenGLDrv/Private/Lumin/LuminOpenGL.h
+++ b/Engine/Source/Runtime/OpenGLDrv/Private/Lumin/LuminOpenGL.h
@@ -35,6 +35,24 @@ typedef khronos_uint64_t GLuint64;
 #define GL_HALF_FLOAT	GL_HALF_FLOAT_OES
 #endif
 
+
+#ifndef GL_TEXTURE_1D
+#define GL_TEXTURE_1D			0x0DE0
+#endif
+
+#ifndef GL_TEXTURE_1D_ARRAY
+#define GL_TEXTURE_1D_ARRAY		0x8C18
+#endif
+
+#ifndef GL_TEXTURE_2D_ARRAY
+#define GL_TEXTURE_2D_ARRAY		0x8C1A
+#endif
+
+#ifndef GL_TEXTURE_RECTANGLE
+#define GL_TEXTURE_RECTANGLE	0x84F5
+#endif
+
+
 #define GL_COMPRESSED_RGB8_ETC2           0x9274
 #define GL_COMPRESSED_SRGB8_ETC2          0x9275
 #define GL_COMPRESSED_RGBA8_ETC2_EAC      0x9278
