DR - ios compile fix
#rnx

diff --git a/Engine/Source/Runtime/OpenGLDrv/Private/IOS/IOSOpenGL.h b/Engine/Source/Runtime/OpenGLDrv/Private/IOS/IOSOpenGL.h
--- a/Engine/Source/Runtime/OpenGLDrv/Private/IOS/IOSOpenGL.h
+++ b/Engine/Source/Runtime/OpenGLDrv/Private/IOS/IOSOpenGL.h
@@ -23,6 +23,21 @@ typedef GLsync UGLsync;
 #undef GL_TEXTURE_MAX_LEVEL
 #define GL_TEXTURE_MAX_LEVEL	GL_TEXTURE_MAX_LEVEL_APPLE
 
+#ifndef GL_TEXTURE_1D
+#define GL_TEXTURE_1D			0x0DE0
+#endif
+
+#ifndef GL_TEXTURE_1D_ARRAY
+#define GL_TEXTURE_1D_ARRAY		0x8C18
+#endif
+
+#ifndef GL_TEXTURE_2D_ARRAY
+#define GL_TEXTURE_2D_ARRAY		0x8C1A
+#endif
+
+#ifndef GL_TEXTURE_RECTANGLE
+#define GL_TEXTURE_RECTANGLE	0x84F5
+#endif
 
 struct FIOSOpenGL : public FOpenGLES2
 {
@@ -124,6 +139,10 @@ struct FIOSOpenGL : public FOpenGLES2
 		return SP_OPENGL_ES2_IOS;
 	}
 
+	static FORCEINLINE void	CopyTexSubImage1D(GLenum Target, GLint Level, GLint XOffset, GLint X, GLint Y, GLsizei Width)
+	{
+	}
+
 	static FORCEINLINE bool SupportsFramebufferSRGBEnable()				{ return false; }
 
 	static FORCEINLINE GLenum GetShadowDepthFormat()					{ return GL_DEPTH_COMPONENT16; }
