DR - Lumin compile fix
#rnx

diff --git a/Engine/Source/Runtime/OpenGLDrv/Private/Lumin/LuminGL4.h b/Engine/Source/Runtime/OpenGLDrv/Private/Lumin/LuminGL4.h
--- a/Engine/Source/Runtime/OpenGLDrv/Private/Lumin/LuminGL4.h
+++ b/Engine/Source/Runtime/OpenGLDrv/Private/Lumin/LuminGL4.h
@@ -524,6 +524,10 @@ struct FLuminOpenGL4 : public FOpenGL4
 		glCopyImageSubData(SrcName, SrcTarget, SrcLevel, SrcX, SrcY, SrcZ, DstName, DstTarget, DstLevel, DstX, DstY, DstZ, Width, Height, Depth);
 	}
 
+	static FORCEINLINE void	CopyTexSubImage1D(GLenum Target, GLint Level, GLint XOffset, GLint X, GLint Y, GLsizei Width)
+	{
+	}
+
 	static FORCEINLINE bool SupportsBindlessTexture()
 	{
 		return bSupportsBindlessTexture;
diff --git a/Engine/Source/Runtime/OpenGLDrv/Private/Lumin/LuminOpenGL.h b/Engine/Source/Runtime/OpenGLDrv/Private/Lumin/LuminOpenGL.h
--- a/Engine/Source/Runtime/OpenGLDrv/Private/Lumin/LuminOpenGL.h
+++ b/Engine/Source/Runtime/OpenGLDrv/Private/Lumin/LuminOpenGL.h
@@ -415,6 +415,10 @@ struct FLuminOpenGL : public FOpenGLES2
 		glFramebufferTextureLayer(Target, Attachment, Texture, Level, Layer);
 	}
 
+	static FORCEINLINE void	CopyTexSubImage1D(GLenum Target, GLint Level, GLint XOffset, GLint X, GLint Y, GLsizei Width)
+	{
+	}
+
 	// Android ES2 shaders have code that allows compile selection of
 	// 32 bpp HDR encoding mode via 'intrinsic_GetHDR32bppEncodeModeES2()'.
 	static FORCEINLINE bool SupportsHDR32bppEncodeModeIntrinsic()		{ return true; }
