DR - Android compile fix
#rnx

diff --git a/Engine/Source/Runtime/OpenGLDrv/Private/Android/AndroidOpenGL.h b/Engine/Source/Runtime/OpenGLDrv/Private/Android/AndroidOpenGL.h
--- a/Engine/Source/Runtime/OpenGLDrv/Private/Android/AndroidOpenGL.h
+++ b/Engine/Source/Runtime/OpenGLDrv/Private/Android/AndroidOpenGL.h
@@ -43,6 +43,23 @@ typedef khronos_uint64_t GLuint64;
 #define GL_HALF_FLOAT	GL_HALF_FLOAT_OES
 #endif
 
+#ifndef GL_TEXTURE_1D
+#define GL_TEXTURE_1D			0x0DE0
+#endif
+
+#ifndef GL_TEXTURE_1D_ARRAY
+#define GL_TEXTURE_1D_ARRAY		0x8C18
+#endif
+
+#ifndef GL_TEXTURE_2D_ARRAY
+#define GL_TEXTURE_2D_ARRAY		0x8C1A
+#endif
+
+#ifndef GL_TEXTURE_RECTANGLE
+#define GL_TEXTURE_RECTANGLE	0x84F5
+#endif
+
+
 #define GL_COMPRESSED_RGB8_ETC2           0x9274
 #define GL_COMPRESSED_SRGB8_ETC2          0x9275
 #define GL_COMPRESSED_RGBA8_ETC2_EAC      0x9278
@@ -416,6 +433,10 @@ struct FAndroidOpenGL : public FOpenGLES2
 		glTexSubImage3D(Target, Level, XOffset, YOffset, ZOffset, Width, Height, Depth, Format, Type, PixelData);
 	}
 
+	static FORCEINLINE void	CopyTexSubImage1D(GLenum Target, GLint Level, GLint XOffset, GLint X, GLint Y, GLsizei Width)
+	{
+	}
+
 	static FORCEINLINE void	CopyTexSubImage2D(GLenum Target, GLint Level, GLint XOffset, GLint YOffset, GLint X, GLint Y, GLsizei Width, GLsizei Height)
 	{
 		glCopyTexSubImage2D(Target, Level, XOffset, YOffset, X, Y, Width, Height);
