DR - Deprecate FVertexShaderRHIParamRef, FHullShaderRHIParamRef, FDomainShaderRHIParamRef, FPixelShaderRHIParamRef, FGeometryShaderRHIParamRef and FComputeShaderRHIParamRef

diff --git a/Engine/Source/Runtime/EmptyRHI/Private/EmptyCommands.cpp b/Engine/Source/Runtime/EmptyRHI/Private/EmptyCommands.cpp
--- a/Engine/Source/Runtime/EmptyRHI/Private/EmptyCommands.cpp
+++ b/Engine/Source/Runtime/EmptyRHI/Private/EmptyCommands.cpp
@@ -23,7 +23,7 @@ void FEmptyDynamicRHI::RHISetRasterizerState(FRHIRasterizerState* NewStateRHI)
 
 }
 
-void FEmptyDynamicRHI::RHISetComputeShader(FComputeShaderRHIParamRef ComputeShaderRHI)
+void FEmptyDynamicRHI::RHISetComputeShader(FRHIComputeShader* ComputeShaderRHI)
 {
 	FEmptyComputeShader* ComputeShader = ResourceCast(ComputeShaderRHI);
 
@@ -62,177 +62,177 @@ void FEmptyDynamicRHI::RHISetBoundShaderState(FRHIBoundShaderState* BoundShaderS
 }
 
 
-void FEmptyDynamicRHI::RHISetUAVParameter(FComputeShaderRHIParamRef ComputeShaderRHI, uint32 UAVIndex, FRHIUnorderedAccessView* UAVRHI)
+void FEmptyDynamicRHI::RHISetUAVParameter(FRHIComputeShader* ComputeShaderRHI, uint32 UAVIndex, FRHIUnorderedAccessView* UAVRHI)
 {
 	FEmptyUnorderedAccessView* UAV = ResourceCast(UAVRHI);
 
 }
 
-void FEmptyDynamicRHI::RHISetUAVParameter(FComputeShaderRHIParamRef ComputeShaderRHI,uint32 UAVIndex, FRHIUnorderedAccessView* UAVRHI, uint32 InitialCount)
+void FEmptyDynamicRHI::RHISetUAVParameter(FRHIComputeShader* ComputeShaderRHI,uint32 UAVIndex, FRHIUnorderedAccessView* UAVRHI, uint32 InitialCount)
 {
 	FEmptyUnorderedAccessView* UAV = ResourceCast(UAVRHI);
 
 }
 
 
-void FEmptyDynamicRHI::RHISetShaderTexture(FVertexShaderRHIParamRef VertexShaderRHI, uint32 TextureIndex, FTextureRHIParamRef NewTextureRHI)
+void FEmptyDynamicRHI::RHISetShaderTexture(FRHIVertexShader* VertexShaderRHI, uint32 TextureIndex, FTextureRHIParamRef NewTextureRHI)
 {
 
 }
 
-void FEmptyDynamicRHI::RHISetShaderTexture(FHullShaderRHIParamRef HullShader, uint32 TextureIndex, FTextureRHIParamRef NewTextureRHI)
+void FEmptyDynamicRHI::RHISetShaderTexture(FRHIHullShader* HullShader, uint32 TextureIndex, FTextureRHIParamRef NewTextureRHI)
 {
 
 }
 
-void FEmptyDynamicRHI::RHISetShaderTexture(FDomainShaderRHIParamRef DomainShader, uint32 TextureIndex, FTextureRHIParamRef NewTextureRHI)
+void FEmptyDynamicRHI::RHISetShaderTexture(FRHIDomainShader* DomainShader, uint32 TextureIndex, FTextureRHIParamRef NewTextureRHI)
 {
 
 }
 
-void FEmptyDynamicRHI::RHISetShaderTexture(FGeometryShaderRHIParamRef GeometryShader, uint32 TextureIndex, FTextureRHIParamRef NewTextureRHI)
+void FEmptyDynamicRHI::RHISetShaderTexture(FRHIGeometryShader* GeometryShader, uint32 TextureIndex, FTextureRHIParamRef NewTextureRHI)
 {
 
 }
 
-void FEmptyDynamicRHI::RHISetShaderTexture(FPixelShaderRHIParamRef PixelShader, uint32 TextureIndex, FTextureRHIParamRef NewTextureRHI)
+void FEmptyDynamicRHI::RHISetShaderTexture(FRHIPixelShader* PixelShader, uint32 TextureIndex, FTextureRHIParamRef NewTextureRHI)
 {
 
 }
 
-void FEmptyDynamicRHI::RHISetShaderTexture(FComputeShaderRHIParamRef ComputeShader, uint32 TextureIndex, FTextureRHIParamRef NewTextureRHI)
+void FEmptyDynamicRHI::RHISetShaderTexture(FRHIComputeShader* ComputeShader, uint32 TextureIndex, FTextureRHIParamRef NewTextureRHI)
 {
 
 }
 
 
-void FEmptyDynamicRHI::RHISetShaderResourceViewParameter(FVertexShaderRHIParamRef VertexShaderRHI, uint32 TextureIndex, FRHIShaderResourceView* SRVRHI)
+void FEmptyDynamicRHI::RHISetShaderResourceViewParameter(FRHIVertexShader* VertexShaderRHI, uint32 TextureIndex, FRHIShaderResourceView* SRVRHI)
 {
 
 }
 
-void FEmptyDynamicRHI::RHISetShaderResourceViewParameter(FHullShaderRHIParamRef HullShaderRHI,uint32 TextureIndex, FRHIShaderResourceView* SRVRHI)
+void FEmptyDynamicRHI::RHISetShaderResourceViewParameter(FRHIHullShader* HullShaderRHI,uint32 TextureIndex, FRHIShaderResourceView* SRVRHI)
 {
 
 }
 
-void FEmptyDynamicRHI::RHISetShaderResourceViewParameter(FDomainShaderRHIParamRef DomainShaderRHI,uint32 TextureIndex, FRHIShaderResourceView* SRVRHI)
+void FEmptyDynamicRHI::RHISetShaderResourceViewParameter(FRHIDomainShader* DomainShaderRHI,uint32 TextureIndex, FRHIShaderResourceView* SRVRHI)
 {
 
 }
 
-void FEmptyDynamicRHI::RHISetShaderResourceViewParameter(FGeometryShaderRHIParamRef GeometryShaderRHI,uint32 TextureIndex, FRHIShaderResourceView* SRVRHI)
+void FEmptyDynamicRHI::RHISetShaderResourceViewParameter(FRHIGeometryShader* GeometryShaderRHI,uint32 TextureIndex, FRHIShaderResourceView* SRVRHI)
 {
 
 }
 
-void FEmptyDynamicRHI::RHISetShaderResourceViewParameter(FPixelShaderRHIParamRef PixelShaderRHI,uint32 TextureIndex, FRHIShaderResourceView* SRVRHI)
+void FEmptyDynamicRHI::RHISetShaderResourceViewParameter(FRHIPixelShader* PixelShaderRHI,uint32 TextureIndex, FRHIShaderResourceView* SRVRHI)
 {
 
 }
 
-void FEmptyDynamicRHI::RHISetShaderResourceViewParameter(FComputeShaderRHIParamRef ComputeShaderRHI,uint32 TextureIndex, FRHIShaderResourceView* SRVRHI)
+void FEmptyDynamicRHI::RHISetShaderResourceViewParameter(FRHIComputeShader* ComputeShaderRHI,uint32 TextureIndex, FRHIShaderResourceView* SRVRHI)
 {
 
 }
 
 
-void FEmptyDynamicRHI::RHISetShaderSampler(FVertexShaderRHIParamRef VertexShaderRHI, uint32 SamplerIndex, FRHISamplerState* NewStateRHI)
+void FEmptyDynamicRHI::RHISetShaderSampler(FRHIVertexShader* VertexShaderRHI, uint32 SamplerIndex, FRHISamplerState* NewStateRHI)
 {
 	FEmptySamplerState* NewState = ResourceCast(NewStateRHI);
 	FEmptyVertexShader* VertexShader = ResourceCast(VertexShaderRHI);
 
 }
 
-void FEmptyDynamicRHI::RHISetShaderSampler(FHullShaderRHIParamRef HullShader, uint32 SamplerIndex, FRHISamplerState* NewStateRHI)
+void FEmptyDynamicRHI::RHISetShaderSampler(FRHIHullShader* HullShader, uint32 SamplerIndex, FRHISamplerState* NewStateRHI)
 {
 	FEmptySamplerState* NewState = ResourceCast(NewStateRHI);
 
 }
 
-void FEmptyDynamicRHI::RHISetShaderSampler(FDomainShaderRHIParamRef DomainShader, uint32 SamplerIndex, FRHISamplerState* NewStateRHI)
+void FEmptyDynamicRHI::RHISetShaderSampler(FRHIDomainShader* DomainShader, uint32 SamplerIndex, FRHISamplerState* NewStateRHI)
 {
 	FEmptySamplerState* NewState = ResourceCast(NewStateRHI);
 
 }
 
-void FEmptyDynamicRHI::RHISetShaderSampler(FGeometryShaderRHIParamRef GeometryShader, uint32 SamplerIndex, FRHISamplerState* NewStateRHI)
+void FEmptyDynamicRHI::RHISetShaderSampler(FRHIGeometryShader* GeometryShader, uint32 SamplerIndex, FRHISamplerState* NewStateRHI)
 {
 	FEmptySamplerState* NewState = ResourceCast(NewStateRHI);
 
 }
 
-void FEmptyDynamicRHI::RHISetShaderSampler(FPixelShaderRHIParamRef PixelShader, uint32 SamplerIndex, FRHISamplerState* NewStateRHI)
+void FEmptyDynamicRHI::RHISetShaderSampler(FRHIPixelShader* PixelShader, uint32 SamplerIndex, FRHISamplerState* NewStateRHI)
 {
 	FEmptySamplerState* NewState = ResourceCast(NewStateRHI);
 
 }
 
-void FEmptyDynamicRHI::RHISetShaderSampler(FComputeShaderRHIParamRef ComputeShader, uint32 SamplerIndex, FRHISamplerState* NewStateRHI)
+void FEmptyDynamicRHI::RHISetShaderSampler(FRHIComputeShader* ComputeShader, uint32 SamplerIndex, FRHISamplerState* NewStateRHI)
 {
 	FEmptySamplerState* NewState = ResourceCast(NewStateRHI);
 
 }
 
 
-void FEmptyDynamicRHI::RHISetShaderParameter(FVertexShaderRHIParamRef VertexShaderRHI, uint32 BufferIndex, uint32 BaseIndex, uint32 NumBytes, const void* NewValue)
+void FEmptyDynamicRHI::RHISetShaderParameter(FRHIVertexShader* VertexShaderRHI, uint32 BufferIndex, uint32 BaseIndex, uint32 NumBytes, const void* NewValue)
 {
 
 }
 
-void FEmptyDynamicRHI::RHISetShaderParameter(FHullShaderRHIParamRef HullShaderRHI, uint32 BufferIndex, uint32 BaseIndex, uint32 NumBytes, const void* NewValue)
+void FEmptyDynamicRHI::RHISetShaderParameter(FRHIHullShader* HullShaderRHI, uint32 BufferIndex, uint32 BaseIndex, uint32 NumBytes, const void* NewValue)
 {
 
 }
 
-void FEmptyDynamicRHI::RHISetShaderParameter(FPixelShaderRHIParamRef PixelShaderRHI, uint32 BufferIndex, uint32 BaseIndex, uint32 NumBytes, const void* NewValue)
+void FEmptyDynamicRHI::RHISetShaderParameter(FRHIPixelShader* PixelShaderRHI, uint32 BufferIndex, uint32 BaseIndex, uint32 NumBytes, const void* NewValue)
 {
 
 }
 
-void FEmptyDynamicRHI::RHISetShaderParameter(FDomainShaderRHIParamRef DomainShaderRHI, uint32 BufferIndex, uint32 BaseIndex, uint32 NumBytes, const void* NewValue)
+void FEmptyDynamicRHI::RHISetShaderParameter(FRHIDomainShader* DomainShaderRHI, uint32 BufferIndex, uint32 BaseIndex, uint32 NumBytes, const void* NewValue)
 {
 
 }
 
-void FEmptyDynamicRHI::RHISetShaderParameter(FGeometryShaderRHIParamRef GeometryShaderRHI, uint32 BufferIndex, uint32 BaseIndex, uint32 NumBytes, const void* NewValue)
+void FEmptyDynamicRHI::RHISetShaderParameter(FRHIGeometryShader* GeometryShaderRHI, uint32 BufferIndex, uint32 BaseIndex, uint32 NumBytes, const void* NewValue)
 {
 
 }
 
-void FEmptyDynamicRHI::RHISetShaderParameter(FComputeShaderRHIParamRef ComputeShaderRHI,uint32 BufferIndex, uint32 BaseIndex, uint32 NumBytes, const void* NewValue)
+void FEmptyDynamicRHI::RHISetShaderParameter(FRHIComputeShader* ComputeShaderRHI,uint32 BufferIndex, uint32 BaseIndex, uint32 NumBytes, const void* NewValue)
 {
 
 }
 
-void FEmptyDynamicRHI::RHISetShaderUniformBuffer(FVertexShaderRHIParamRef VertexShaderRHI, uint32 BufferIndex, FUniformBufferRHIParamRef BufferRHI)
+void FEmptyDynamicRHI::RHISetShaderUniformBuffer(FRHIVertexShader* VertexShaderRHI, uint32 BufferIndex, FUniformBufferRHIParamRef BufferRHI)
 {
 	FEmptyUniformBuffer* Buffer = ResourceCast(BufferRHI);
 	FEmptyVertexShader* VertexShader = ResourceCast(VertexShaderRHI);
 
 }
 
-void FEmptyDynamicRHI::RHISetShaderUniformBuffer(FHullShaderRHIParamRef HullShader, uint32 BufferIndex, FUniformBufferRHIParamRef BufferRHI)
+void FEmptyDynamicRHI::RHISetShaderUniformBuffer(FRHIHullShader* HullShader, uint32 BufferIndex, FUniformBufferRHIParamRef BufferRHI)
 {
 	FEmptyUniformBuffer* Buffer = ResourceCast(BufferRHI);
 }
 
-void FEmptyDynamicRHI::RHISetShaderUniformBuffer(FDomainShaderRHIParamRef DomainShader, uint32 BufferIndex, FUniformBufferRHIParamRef BufferRHI)
+void FEmptyDynamicRHI::RHISetShaderUniformBuffer(FRHIDomainShader* DomainShader, uint32 BufferIndex, FUniformBufferRHIParamRef BufferRHI)
 {
 	FEmptyUniformBuffer* Buffer = ResourceCast(BufferRHI);
 }
 
-void FEmptyDynamicRHI::RHISetShaderUniformBuffer(FGeometryShaderRHIParamRef GeometryShader, uint32 BufferIndex, FUniformBufferRHIParamRef BufferRHI)
+void FEmptyDynamicRHI::RHISetShaderUniformBuffer(FRHIGeometryShader* GeometryShader, uint32 BufferIndex, FUniformBufferRHIParamRef BufferRHI)
 {
 	FEmptyUniformBuffer* Buffer = ResourceCast(BufferRHI);
 }
 
-void FEmptyDynamicRHI::RHISetShaderUniformBuffer(FPixelShaderRHIParamRef PixelShader, uint32 BufferIndex, FUniformBufferRHIParamRef BufferRHI)
+void FEmptyDynamicRHI::RHISetShaderUniformBuffer(FRHIPixelShader* PixelShader, uint32 BufferIndex, FUniformBufferRHIParamRef BufferRHI)
 {
 	FEmptyUniformBuffer* Buffer = ResourceCast(BufferRHI);
 }
 
-void FEmptyDynamicRHI::RHISetShaderUniformBuffer(FComputeShaderRHIParamRef ComputeShader, uint32 BufferIndex, FUniformBufferRHIParamRef BufferRHI)
+void FEmptyDynamicRHI::RHISetShaderUniformBuffer(FRHIComputeShader* ComputeShader, uint32 BufferIndex, FUniformBufferRHIParamRef BufferRHI)
 {
 	FEmptyUniformBuffer* Buffer = ResourceCast(BufferRHI);
 }
diff --git a/Engine/Source/Runtime/EmptyRHI/Private/EmptyShaders.cpp b/Engine/Source/Runtime/EmptyRHI/Private/EmptyShaders.cpp
--- a/Engine/Source/Runtime/EmptyRHI/Private/EmptyShaders.cpp
+++ b/Engine/Source/Runtime/EmptyRHI/Private/EmptyShaders.cpp
@@ -69,11 +69,11 @@ FComputeShaderRHIRef FEmptyDynamicRHI::RHICreateComputeShader(const TArray<uint8
 
 FEmptyBoundShaderState::FEmptyBoundShaderState(
 			FRHIVertexDeclaration* InVertexDeclarationRHI,
-			FVertexShaderRHIParamRef InVertexShaderRHI,
-			FPixelShaderRHIParamRef InPixelShaderRHI,
-			FHullShaderRHIParamRef InHullShaderRHI,
-			FDomainShaderRHIParamRef InDomainShaderRHI,
-			FGeometryShaderRHIParamRef InGeometryShaderRHI)
+			FRHIVertexShader* InVertexShaderRHI,
+			FRHIPixelShader* InPixelShaderRHI,
+			FRHIHullShader* InHullShaderRHI,
+	FRHIDomainShader* InDomainShaderRHI,
+	FRHIGeometryShader* InGeometryShaderRHI)
 	:	CacheLink(InVertexDeclarationRHI,InVertexShaderRHI,InPixelShaderRHI,InHullShaderRHI,InDomainShaderRHI,InGeometryShaderRHI,this)
 {
 	FEmptyVertexDeclaration* InVertexDeclaration = FEmptyDynamicRHI::ResourceCast(InVertexDeclarationRHI);
@@ -99,11 +99,11 @@ FEmptyBoundShaderState::~FEmptyBoundShaderState()
 
 FBoundShaderStateRHIRef FEmptyDynamicRHI::RHICreateBoundShaderState(
 	FRHIVertexDeclaration* VertexDeclarationRHI,
-	FVertexShaderRHIParamRef VertexShaderRHI, 
-	FHullShaderRHIParamRef HullShaderRHI, 
-	FDomainShaderRHIParamRef DomainShaderRHI, 
-	FPixelShaderRHIParamRef PixelShaderRHI,
-	FGeometryShaderRHIParamRef GeometryShaderRHI
+	FRHIVertexShader* VertexShaderRHI,
+	FRHIHullShader* HullShaderRHI,
+	FRHIDomainShader* DomainShaderRHI,
+	FRHIPixelShader* PixelShaderRHI,
+	FRHIGeometryShader* GeometryShaderRHI
 	)
 {
 	check(IsInRenderingThread());
diff --git a/Engine/Source/Runtime/EmptyRHI/Public/EmptyRHI.h b/Engine/Source/Runtime/EmptyRHI/Public/EmptyRHI.h
--- a/Engine/Source/Runtime/EmptyRHI/Public/EmptyRHI.h
+++ b/Engine/Source/Runtime/EmptyRHI/Public/EmptyRHI.h
@@ -48,7 +48,7 @@ public:
 	virtual FGeometryShaderRHIRef RHICreateGeometryShader(const TArray<uint8>& Code) final override;
 	virtual FGeometryShaderRHIRef RHICreateGeometryShaderWithStreamOutput(const TArray<uint8>& Code, const FStreamOutElementList& ElementList, uint32 NumStrides, const uint32* Strides, int32 RasterizedStream) final override;
 	virtual FComputeShaderRHIRef RHICreateComputeShader(const TArray<uint8>& Code) final override;
-	virtual FBoundShaderStateRHIRef RHICreateBoundShaderState(FRHIVertexDeclaration* VertexDeclaration, FVertexShaderRHIParamRef VertexShader, FHullShaderRHIParamRef HullShader, FDomainShaderRHIParamRef DomainShader, FPixelShaderRHIParamRef PixelShader, FGeometryShaderRHIParamRef GeometryShader) final override;
+	virtual FBoundShaderStateRHIRef RHICreateBoundShaderState(FRHIVertexDeclaration* VertexDeclaration, FRHIVertexShader* VertexShader, FRHIHullShader* HullShader, FRHIDomainShaderFRHIDomainShader*, FRHIPixelShader* PixelShader, FRHIGeometryShader* GeometryShader) final override;
 	virtual FUniformBufferRHIRef RHICreateUniformBuffer(const void* Contents, const FRHIUniformBufferLayout& Layout, EUniformBufferUsage Usage, EUniformBufferValidation Validation) final override;
 	virtual FIndexBufferRHIRef RHICreateIndexBuffer(uint32 Stride, uint32 Size, uint32 InUsage, FRHIResourceCreateInfo& CreateInfo) final override;
 	virtual void* RHILockIndexBuffer(FIndexBufferRHIParamRef IndexBuffer, uint32 Offset, uint32 Size, EResourceLockMode LockMode) final override;
@@ -132,7 +132,7 @@ public:
 
 
 
-	virtual void RHISetComputeShader(FComputeShaderRHIParamRef ComputeShader) final override;
+	virtual void RHISetComputeShader(FRHIComputeShader* ComputeShader) final override;
 	virtual void RHIDispatchComputeShader(uint32 ThreadGroupCountX, uint32 ThreadGroupCountY, uint32 ThreadGroupCountZ) final override;
 	virtual void RHIDispatchIndirectComputeShader(FVertexBufferRHIParamRef ArgumentBuffer, uint32 ArgumentOffset) final override;
 	virtual void RHIAutomaticCacheFlushAfterComputeShader(bool bEnable) final override;
@@ -155,38 +155,38 @@ public:
 	virtual void RHISetViewport(uint32 MinX, uint32 MinY, float MinZ, uint32 MaxX, uint32 MaxY, float MaxZ) final override;
 	virtual void RHISetScissorRect(bool bEnable, uint32 MinX, uint32 MinY, uint32 MaxX, uint32 MaxY) final override;
 	virtual void RHISetBoundShaderState(FRHIBoundShaderState* BoundShaderState) final override;
-	virtual void RHISetShaderTexture(FVertexShaderRHIParamRef VertexShader, uint32 TextureIndex, FTextureRHIParamRef NewTexture) final override;
-	virtual void RHISetShaderTexture(FHullShaderRHIParamRef HullShader, uint32 TextureIndex, FTextureRHIParamRef NewTexture) final override;
-	virtual void RHISetShaderTexture(FDomainShaderRHIParamRef DomainShader, uint32 TextureIndex, FTextureRHIParamRef NewTexture) final override;
-	virtual void RHISetShaderTexture(FGeometryShaderRHIParamRef GeometryShader, uint32 TextureIndex, FTextureRHIParamRef NewTexture) final override;
-	virtual void RHISetShaderTexture(FPixelShaderRHIParamRef PixelShader, uint32 TextureIndex, FTextureRHIParamRef NewTexture) final override;
-	virtual void RHISetShaderTexture(FComputeShaderRHIParamRef PixelShader, uint32 TextureIndex, FTextureRHIParamRef NewTexture) final override;
-	virtual void RHISetShaderSampler(FComputeShaderRHIParamRef ComputeShader, uint32 SamplerIndex, FRHISamplerState* NewState) final override;
-	virtual void RHISetShaderSampler(FVertexShaderRHIParamRef VertexShader, uint32 SamplerIndex, FRHISamplerState* NewState) final override;
-	virtual void RHISetShaderSampler(FGeometryShaderRHIParamRef GeometryShader, uint32 SamplerIndex, FRHISamplerState* NewState) final override;
-	virtual void RHISetShaderSampler(FDomainShaderRHIParamRef DomainShader, uint32 SamplerIndex, FRHISamplerState* NewState) final override;
-	virtual void RHISetShaderSampler(FHullShaderRHIParamRef HullShader, uint32 SamplerIndex, FRHISamplerState* NewState) final override;
-	virtual void RHISetShaderSampler(FPixelShaderRHIParamRef PixelShader, uint32 SamplerIndex, FRHISamplerState* NewState) final override;
-	virtual void RHISetUAVParameter(FComputeShaderRHIParamRef ComputeShader, uint32 UAVIndex, FRHIUnorderedAccessView* UAV) final override;
-	virtual void RHISetUAVParameter(FComputeShaderRHIParamRef ComputeShader, uint32 UAVIndex, FRHIUnorderedAccessView* UAV, uint32 InitialCount) final override;
-	virtual void RHISetShaderResourceViewParameter(FPixelShaderRHIParamRef PixelShader, uint32 SamplerIndex, FRHIShaderResourceView* SRV) final override;
-	virtual void RHISetShaderResourceViewParameter(FVertexShaderRHIParamRef VertexShader, uint32 SamplerIndex, FRHIShaderResourceView* SRV) final override;
-	virtual void RHISetShaderResourceViewParameter(FComputeShaderRHIParamRef ComputeShader, uint32 SamplerIndex, FRHIShaderResourceView* SRV) final override;
-	virtual void RHISetShaderResourceViewParameter(FHullShaderRHIParamRef HullShader, uint32 SamplerIndex, FRHIShaderResourceView* SRV) final override;
-	virtual void RHISetShaderResourceViewParameter(FDomainShaderRHIParamRef DomainShader, uint32 SamplerIndex, FRHIShaderResourceView* SRV) final override;
-	virtual void RHISetShaderResourceViewParameter(FGeometryShaderRHIParamRef GeometryShader, uint32 SamplerIndex, FRHIShaderResourceView* SRV) final override;
-	virtual void RHISetShaderUniformBuffer(FVertexShaderRHIParamRef VertexShader, uint32 BufferIndex, FUniformBufferRHIParamRef Buffer) final override;
-	virtual void RHISetShaderUniformBuffer(FHullShaderRHIParamRef HullShader, uint32 BufferIndex, FUniformBufferRHIParamRef Buffer) final override;
-	virtual void RHISetShaderUniformBuffer(FDomainShaderRHIParamRef DomainShader, uint32 BufferIndex, FUniformBufferRHIParamRef Buffer) final override;
-	virtual void RHISetShaderUniformBuffer(FGeometryShaderRHIParamRef GeometryShader, uint32 BufferIndex, FUniformBufferRHIParamRef Buffer) final override;
-	virtual void RHISetShaderUniformBuffer(FPixelShaderRHIParamRef PixelShader, uint32 BufferIndex, FUniformBufferRHIParamRef Buffer) final override;
-	virtual void RHISetShaderUniformBuffer(FComputeShaderRHIParamRef ComputeShader, uint32 BufferIndex, FUniformBufferRHIParamRef Buffer) final override;
-	virtual void RHISetShaderParameter(FVertexShaderRHIParamRef VertexShader, uint32 BufferIndex, uint32 BaseIndex, uint32 NumBytes, const void* NewValue) final override;
-	virtual void RHISetShaderParameter(FPixelShaderRHIParamRef PixelShader, uint32 BufferIndex, uint32 BaseIndex, uint32 NumBytes, const void* NewValue) final override;
-	virtual void RHISetShaderParameter(FHullShaderRHIParamRef HullShader, uint32 BufferIndex, uint32 BaseIndex, uint32 NumBytes, const void* NewValue) final override;
-	virtual void RHISetShaderParameter(FDomainShaderRHIParamRef DomainShader, uint32 BufferIndex, uint32 BaseIndex, uint32 NumBytes, const void* NewValue) final override;
-	virtual void RHISetShaderParameter(FGeometryShaderRHIParamRef GeometryShader, uint32 BufferIndex, uint32 BaseIndex, uint32 NumBytes, const void* NewValue) final override;
-	virtual void RHISetShaderParameter(FComputeShaderRHIParamRef ComputeShader, uint32 BufferIndex, uint32 BaseIndex, uint32 NumBytes, const void* NewValue) final override;
+	virtual void RHISetShaderTexture(FRHIVertexShader* VertexShader, uint32 TextureIndex, FTextureRHIParamRef NewTexture) final override;
+	virtual void RHISetShaderTexture(FRHIHullShader* HullShader, uint32 TextureIndex, FTextureRHIParamRef NewTexture) final override;
+	virtual void RHISetShaderTexture(FRHIDomainShader* DomainShader, uint32 TextureIndex, FTextureRHIParamRef NewTexture) final override;
+	virtual void RHISetShaderTexture(FRHIGeometryShader* GeometryShader, uint32 TextureIndex, FTextureRHIParamRef NewTexture) final override;
+	virtual void RHISetShaderTexture(FRHIPixelShader* PixelShader, uint32 TextureIndex, FTextureRHIParamRef NewTexture) final override;
+	virtual void RHISetShaderTexture(FRHIComputeShader* PixelShader, uint32 TextureIndex, FTextureRHIParamRef NewTexture) final override;
+	virtual void RHISetShaderSampler(FRHIComputeShader* ComputeShader, uint32 SamplerIndex, FRHISamplerState* NewState) final override;
+	virtual void RHISetShaderSampler(FRHIVertexShader* VertexShader, uint32 SamplerIndex, FRHISamplerState* NewState) final override;
+	virtual void RHISetShaderSampler(FRHIGeometryShader* GeometryShader, uint32 SamplerIndex, FRHISamplerState* NewState) final override;
+	virtual void RHISetShaderSampler(FRHIDomainShader* DomainShader, uint32 SamplerIndex, FRHISamplerState* NewState) final override;
+	virtual void RHISetShaderSampler(FRHIHullShader* HullShader, uint32 SamplerIndex, FRHISamplerState* NewState) final override;
+	virtual void RHISetShaderSampler(FRHIPixelShader* PixelShader, uint32 SamplerIndex, FRHISamplerState* NewState) final override;
+	virtual void RHISetUAVParameter(FRHIComputeShader* ComputeShader, uint32 UAVIndex, FRHIUnorderedAccessView* UAV) final override;
+	virtual void RHISetUAVParameter(FRHIComputeShader* ComputeShader, uint32 UAVIndex, FRHIUnorderedAccessView* UAV, uint32 InitialCount) final override;
+	virtual void RHISetShaderResourceViewParameter(FRHIPixelShader* PixelShader, uint32 SamplerIndex, FRHIShaderResourceView* SRV) final override;
+	virtual void RHISetShaderResourceViewParameter(FRHIVertexShader* VertexShader, uint32 SamplerIndex, FRHIShaderResourceView* SRV) final override;
+	virtual void RHISetShaderResourceViewParameter(FRHIComputeShader* ComputeShader, uint32 SamplerIndex, FRHIShaderResourceView* SRV) final override;
+	virtual void RHISetShaderResourceViewParameter(FRHIHullShader* HullShader, uint32 SamplerIndex, FRHIShaderResourceView* SRV) final override;
+	virtual void RHISetShaderResourceViewParameter(FRHIDomainShader* DomainShader, uint32 SamplerIndex, FRHIShaderResourceView* SRV) final override;
+	virtual void RHISetShaderResourceViewParameter(FRHIGeometryShader* GeometryShader, uint32 SamplerIndex, FRHIShaderResourceView* SRV) final override;
+	virtual void RHISetShaderUniformBuffer(FRHIVertexShader* VertexShader, uint32 BufferIndex, FUniformBufferRHIParamRef Buffer) final override;
+	virtual void RHISetShaderUniformBuffer(FRHIHullShader* HullShader, uint32 BufferIndex, FUniformBufferRHIParamRef Buffer) final override;
+	virtual void RHISetShaderUniformBuffer(FRHIDomainShader* DomainShader, uint32 BufferIndex, FUniformBufferRHIParamRef Buffer) final override;
+	virtual void RHISetShaderUniformBuffer(FRHIGeometryShader* GeometryShader, uint32 BufferIndex, FUniformBufferRHIParamRef Buffer) final override;
+	virtual void RHISetShaderUniformBuffer(FRHIPixelShader* PixelShader, uint32 BufferIndex, FUniformBufferRHIParamRef Buffer) final override;
+	virtual void RHISetShaderUniformBuffer(FRHIComputeShader* ComputeShader, uint32 BufferIndex, FUniformBufferRHIParamRef Buffer) final override;
+	virtual void RHISetShaderParameter(FRHIVertexShader* VertexShader, uint32 BufferIndex, uint32 BaseIndex, uint32 NumBytes, const void* NewValue) final override;
+	virtual void RHISetShaderParameter(FRHIPixelShader* PixelShader, uint32 BufferIndex, uint32 BaseIndex, uint32 NumBytes, const void* NewValue) final override;
+	virtual void RHISetShaderParameter(FRHIHullShader* HullShader, uint32 BufferIndex, uint32 BaseIndex, uint32 NumBytes, const void* NewValue) final override;
+	virtual void RHISetShaderParameter(FRHIDomainShader* DomainShader, uint32 BufferIndex, uint32 BaseIndex, uint32 NumBytes, const void* NewValue) final override;
+	virtual void RHISetShaderParameter(FRHIGeometryShader* GeometryShader, uint32 BufferIndex, uint32 BaseIndex, uint32 NumBytes, const void* NewValue) final override;
+	virtual void RHISetShaderParameter(FRHIComputeShader* ComputeShader, uint32 BufferIndex, uint32 BaseIndex, uint32 NumBytes, const void* NewValue) final override;
 	virtual void RHISetDepthStencilState(FRHIDepthStencilState* NewState, uint32 StencilRef) final override;
 	virtual void RHISetBlendState(FRHIBlendState* NewState, const FLinearColor& BlendFactor) final override;
 	virtual void RHISetRenderTargets(uint32 NumSimultaneousRenderTargets, const FRHIRenderTargetView* NewRenderTargets, const FRHIDepthRenderTargetView* NewDepthStencilTarget, uint32 NumUAVs, FRHIUnorderedAccessView* const* UAVs) final override;
diff --git a/Engine/Source/Runtime/EmptyRHI/Public/EmptyResources.h b/Engine/Source/Runtime/EmptyRHI/Public/EmptyResources.h
--- a/Engine/Source/Runtime/EmptyRHI/Public/EmptyResources.h
+++ b/Engine/Source/Runtime/EmptyRHI/Public/EmptyResources.h
@@ -44,11 +44,11 @@ public:
 	/** Initialization constructor. */
 	FEmptyBoundShaderState(
 		FRHIVertexDeclaration* InVertexDeclarationRHI,
-		FVertexShaderRHIParamRef InVertexShaderRHI,
-		FPixelShaderRHIParamRef InPixelShaderRHI,
-		FHullShaderRHIParamRef InHullShaderRHI,
-		FDomainShaderRHIParamRef InDomainShaderRHI,
-		FGeometryShaderRHIParamRef InGeometryShaderRHI);
+		FRHIVertexShader* InVertexShaderRHI,
+		FRHIPixelShader* InPixelShaderRHI,
+		FRHIHullShader* InHullShaderRHI,
+		FRHIDomainShader* InDomainShaderRHI,
+		FRHIGeometryShader* InGeometryShaderRHI);
 
 	/**
 	 *Destructor
diff --git a/Engine/Source/Runtime/Engine/Private/CubemapUnwrapUtils.cpp b/Engine/Source/Runtime/Engine/Private/CubemapUnwrapUtils.cpp
--- a/Engine/Source/Runtime/Engine/Private/CubemapUnwrapUtils.cpp
+++ b/Engine/Source/Runtime/Engine/Private/CubemapUnwrapUtils.cpp
@@ -164,7 +164,7 @@ void FMipLevelBatchedElementParameters::BindShaders(FRHICommandList& RHICmdList,
 
 void FIESLightProfilePS::SetParameters( FRHICommandList& RHICmdList, const FTexture* Texture, float InBrightnessInLumens )
 {
-	FPixelShaderRHIParamRef ShaderRHI = GetPixelShader();
+	FRHIPixelShader* ShaderRHI = GetPixelShader();
 	SetTextureParameter(RHICmdList, ShaderRHI, IESTexture, IESTextureSampler, Texture);
 
 	SetShaderValue(RHICmdList, ShaderRHI, BrightnessInLumens, InBrightnessInLumens);
diff --git a/Engine/Source/Runtime/Engine/Private/GPUSkinCache.cpp b/Engine/Source/Runtime/Engine/Private/GPUSkinCache.cpp
--- a/Engine/Source/Runtime/Engine/Private/GPUSkinCache.cpp
+++ b/Engine/Source/Runtime/Engine/Private/GPUSkinCache.cpp
@@ -424,7 +424,7 @@ public:
 		FGPUSkinCacheEntry::FSectionDispatchData& DispatchData,
 		FRHIUnorderedAccessView* PositionUAV, FRHIUnorderedAccessView* TangentUAV)
 	{
-		FComputeShaderRHIParamRef ShaderRHI = GetComputeShader();
+		FRHIComputeShader* ShaderRHI = GetComputeShader();
 
 		SetShaderValue(RHICmdList, ShaderRHI, NumVertices, DispatchData.NumVertices);
 		SetShaderValue(RHICmdList, ShaderRHI, InputStreamStart, DispatchData.InputStreamStart);
@@ -468,7 +468,7 @@ public:
 
 	void UnsetParameters(FRHICommandList& RHICmdList)
 	{
-		FComputeShaderRHIParamRef ShaderRHI = GetComputeShader();
+		FRHIComputeShader* ShaderRHI = GetComputeShader();
 
 		SetUAVParameter(RHICmdList, ShaderRHI, PositionBufferUAV, 0);
 		SetUAVParameter(RHICmdList, ShaderRHI, TangentBufferUAV, 0);
@@ -704,7 +704,7 @@ public:
 
 	void SetParameters(FRHICommandListImmediate& RHICmdList, FGPUSkinCacheEntry* Entry, FGPUSkinCacheEntry::FSectionDispatchData& DispatchData, FRWBuffer& StagingBuffer)
 	{
-		const FComputeShaderRHIParamRef ShaderRHI = GetComputeShader();
+		FRHIComputeShader* ShaderRHI = GetComputeShader();
 
 //later		FGlobalShader::SetParameters<FViewUniformShaderParameters>(RHICmdList, ShaderRHI, View);
 
@@ -737,7 +737,7 @@ public:
 
 	void UnsetParameters(FRHICommandList& RHICmdList)
 	{
-		FComputeShaderRHIParamRef ShaderRHI = GetComputeShader();
+		FRHIComputeShader* ShaderRHI = GetComputeShader();
 
 		SetUAVParameter(RHICmdList, ShaderRHI, IntermediateAccumBufferUAV, 0);
 	}
@@ -836,7 +836,7 @@ public:
 
 	void SetParameters(FRHICommandListImmediate& RHICmdList, FGPUSkinCacheEntry* Entry, FGPUSkinCacheEntry::FSectionDispatchData& DispatchData, FRWBuffer& StagingBuffer)
 	{
-		const FComputeShaderRHIParamRef ShaderRHI = GetComputeShader();
+		FRHIComputeShader* ShaderRHI = GetComputeShader();
 
 		check(StagingBuffer.UAV);
 
@@ -853,7 +853,7 @@ public:
 
 	void UnsetParameters(FRHICommandList& RHICmdList)
 	{
-		FComputeShaderRHIParamRef ShaderRHI = GetComputeShader();
+		FRHIComputeShader* ShaderRHI = GetComputeShader();
 
 		SetUAVParameter(RHICmdList, ShaderRHI, TangentBufferUAV, 0);
 		SetUAVParameter(RHICmdList, ShaderRHI, IntermediateAccumBufferUAV, 0);
@@ -962,7 +962,7 @@ void FGPUSkinCache::DispatchUpdateSkinTangents(FRHICommandListImmediate& RHICmdL
 			SCOPED_DRAW_EVENTF(RHICmdList, SkinTangents_PerTrianglePass, TEXT("TangentsTri IndexStart=%d Tri=%d ExtraBoneInfluences=%d UVPrecision=%d"),
 				DispatchData.IndexBufferOffsetValue, DispatchData.NumTriangles, DispatchData.bExtraBoneInfluences, bFullPrecisionUV);
 
-			const FComputeShaderRHIParamRef ShaderRHI = Shader->GetComputeShader();
+			FRHIComputeShader* ShaderRHI = Shader->GetComputeShader();
 			RHICmdList.SetComputeShader(ShaderRHI);
 
 			RHICmdList.TransitionResource(EResourceTransitionAccess::ERWNoBarrier, EResourceTransitionPipeline::EComputeToCompute, StagingBuffer.UAV.GetReference());
@@ -985,7 +985,7 @@ void FGPUSkinCache::DispatchUpdateSkinTangents(FRHICommandListImmediate& RHICmdL
 				DispatchData.InputStreamStart, DispatchData.OutputStreamStart, DispatchData.NumVertices);
 			//#todo-gpuskin Feature level?
 			TShaderMapRef<FRecomputeTangentsPerVertexPassCS> ComputeShader(GetGlobalShaderMap(GMaxRHIFeatureLevel));
-			const FComputeShaderRHIParamRef ShaderRHI = ComputeShader->GetComputeShader();
+			FRHIComputeShader* ShaderRHI = ComputeShader->GetComputeShader();
 			RHICmdList.SetComputeShader(ComputeShader->GetComputeShader());
 
 			uint32 VertexCount = DispatchData.NumVertices;
@@ -1194,7 +1194,7 @@ void FGPUSkinCache::SetVertexStreams(FGPUSkinCacheEntry* Entry, int32 Section, F
 		RHICmdList.SetStreamSource(VertexFactory->GetTangentStreamIndex(), DispatchData.GetTangentRWBuffer()->Buffer, 0);
 	}
 
-	FVertexShaderRHIParamRef ShaderRHI = Shader->GetVertexShader();
+	FRHIVertexShader* ShaderRHI = Shader->GetVertexShader();
 	if (ShaderRHI && GPUSkinCachePreviousPositionBuffer.IsBound())
 	{
 		RHICmdList.SetShaderResourceViewParameter(ShaderRHI, GPUSkinCachePreviousPositionBuffer.GetBaseIndex(), DispatchData.GetPreviousPositionRWBuffer()->SRV);
diff --git a/Engine/Source/Runtime/Engine/Private/GPUSort.cpp b/Engine/Source/Runtime/Engine/Private/GPUSort.cpp
--- a/Engine/Source/Runtime/Engine/Private/GPUSort.cpp
+++ b/Engine/Source/Runtime/Engine/Private/GPUSort.cpp
@@ -260,7 +260,7 @@ public:
 	 */
 	void SetParameters(FRHICommandList& RHICmdList, const FRadixSortUniformBufferRef& UniformBuffer )
 	{
-		FComputeShaderRHIParamRef ComputeShaderRHI = GetComputeShader();
+		FRHIComputeShader* ComputeShaderRHI = GetComputeShader();
 		SetUniformBufferParameter(RHICmdList, ComputeShaderRHI, GetUniformBufferParameter<FRadixSortParameters>(), UniformBuffer );
 	}
 
@@ -269,7 +269,7 @@ public:
 	 */
 	void SetOutput(FRHICommandList& RHICmdList, FRHIUnorderedAccessView* OutOffsetsUAV)
 	{
-		FComputeShaderRHIParamRef ComputeShaderRHI = GetComputeShader();
+		FRHIComputeShader* ComputeShaderRHI = GetComputeShader();
 		if ( OutOffsets.IsBound() )
 		{
 			RHICmdList.SetUAVParameter(ComputeShaderRHI, OutOffsets.GetBaseIndex(), OutOffsetsUAV);
@@ -281,7 +281,7 @@ public:
 	 */
 	void UnbindBuffers(FRHICommandList& RHICmdList)
 	{
-		FComputeShaderRHIParamRef ComputeShaderRHI = GetComputeShader();
+		FRHIComputeShader* ComputeShaderRHI = GetComputeShader();
 		if ( OutOffsets.IsBound() )
 		{
 			RHICmdList.SetUAVParameter(ComputeShaderRHI, OutOffsets.GetBaseIndex(), nullptr);
@@ -357,7 +357,7 @@ public:
 	 */
 	void SetParameters(FRHICommandList& RHICmdList, FRHIShaderResourceView* InKeysSRV, FRadixSortUniformBufferRef& RadixSortUniformBuffer, FRHIShaderResourceView* RadixSortParameterBufferSRV)
 	{
-		FComputeShaderRHIParamRef ComputeShaderRHI = GetComputeShader();
+		FRHIComputeShader* ComputeShaderRHI = GetComputeShader();
 		SetUniformBufferParameter(RHICmdList, ComputeShaderRHI, GetUniformBufferParameter<FRadixSortParameters>(), RadixSortUniformBuffer );
 		if ( InKeys.IsBound() )
 		{
@@ -374,7 +374,7 @@ public:
 	 */
 	void SetOutput(FRHICommandList& RHICmdList, FRHIUnorderedAccessView* OutOffsetsUAV)
 	{
-		FComputeShaderRHIParamRef ComputeShaderRHI = GetComputeShader();
+		FRHIComputeShader* ComputeShaderRHI = GetComputeShader();
 		if ( OutOffsets.IsBound() )
 		{
 			RHICmdList.SetUAVParameter(ComputeShaderRHI, OutOffsets.GetBaseIndex(), OutOffsetsUAV);
@@ -386,7 +386,7 @@ public:
 	 */
 	void UnbindBuffers(FRHICommandList& RHICmdList)
 	{
-		FComputeShaderRHIParamRef ComputeShaderRHI = GetComputeShader();
+		FRHIComputeShader* ComputeShaderRHI = GetComputeShader();
 		if ( RadixSortParameterBuffer.IsBound() )
 		{
 			RHICmdList.SetShaderResourceViewParameter(ComputeShaderRHI, RadixSortParameterBuffer.GetBaseIndex(), nullptr);
@@ -463,7 +463,7 @@ public:
 	 */
 	void SetParameters(FRHICommandList& RHICmdList, FRHIShaderResourceView* InOffsetsSRV)
 	{
-		FComputeShaderRHIParamRef ComputeShaderRHI = GetComputeShader();
+		FRHIComputeShader* ComputeShaderRHI = GetComputeShader();
 		if ( InOffsets.IsBound() )
 		{
 			RHICmdList.SetShaderResourceViewParameter(ComputeShaderRHI, InOffsets.GetBaseIndex(), InOffsetsSRV);
@@ -475,7 +475,7 @@ public:
 	 */
 	void SetOutput(FRHICommandList& RHICmdList, FRHIUnorderedAccessView* OutOffsetsUAV)
 	{
-		FComputeShaderRHIParamRef ComputeShaderRHI = GetComputeShader();
+		FRHIComputeShader* ComputeShaderRHI = GetComputeShader();
 		if ( OutOffsets.IsBound() )
 		{
 			RHICmdList.SetUAVParameter(ComputeShaderRHI, OutOffsets.GetBaseIndex(), OutOffsetsUAV);
@@ -487,7 +487,7 @@ public:
 	 */
 	void UnbindBuffers(FRHICommandList& RHICmdList)
 	{
-		FComputeShaderRHIParamRef ComputeShaderRHI = GetComputeShader();
+		FRHIComputeShader* ComputeShaderRHI = GetComputeShader();
 		if ( InOffsets.IsBound() )
 		{
 			RHICmdList.SetShaderResourceViewParameter(ComputeShaderRHI, InOffsets.GetBaseIndex(), nullptr);
@@ -583,7 +583,7 @@ public:
 		FRadixSortUniformBufferRef& RadixSortUniformBuffer,
 		FRHIShaderResourceView* RadixSortParameterBufferSRV )
 	{
-		FComputeShaderRHIParamRef ComputeShaderRHI = GetComputeShader();
+		FRHIComputeShader* ComputeShaderRHI = GetComputeShader();
 		SetUniformBufferParameter(RHICmdList, ComputeShaderRHI, GetUniformBufferParameter<FRadixSortParameters>(), RadixSortUniformBuffer );
 		if ( RadixSortParameterBuffer.IsBound() )
 		{
@@ -608,7 +608,7 @@ public:
 	 */
 	void SetOutput(FRHICommandList& RHICmdList, FRHIUnorderedAccessView* OutKeysUAV, FRHIUnorderedAccessView* OutValuesUAV)
 	{
-		FComputeShaderRHIParamRef ComputeShaderRHI = GetComputeShader();
+		FRHIComputeShader* ComputeShaderRHI = GetComputeShader();
 		if ( OutKeys.IsBound() )
 		{
 			RHICmdList.SetUAVParameter(ComputeShaderRHI, OutKeys.GetBaseIndex(), OutKeysUAV);
@@ -624,7 +624,7 @@ public:
 	 */
 	void UnbindBuffers(FRHICommandList& RHICmdList)
 	{
-		FComputeShaderRHIParamRef ComputeShaderRHI = GetComputeShader();
+		FRHIComputeShader* ComputeShaderRHI = GetComputeShader();
 		if ( RadixSortParameterBuffer.IsBound() )
 		{
 			RHICmdList.SetShaderResourceViewParameter(ComputeShaderRHI, RadixSortParameterBuffer.GetBaseIndex(), nullptr);
diff --git a/Engine/Source/Runtime/Engine/Private/Particles/ParticleCurveTexture.cpp b/Engine/Source/Runtime/Engine/Private/Particles/ParticleCurveTexture.cpp
--- a/Engine/Source/Runtime/Engine/Private/Particles/ParticleCurveTexture.cpp
+++ b/Engine/Source/Runtime/Engine/Private/Particles/ParticleCurveTexture.cpp
@@ -82,7 +82,7 @@ public:
 		Parameters.PixelScale.Y = 1.0f / GParticleCurveTextureSizeY;
 		Parameters.CurveOffset = CurveOffset;
 		FParticleCurveInjectionBufferRef UniformBuffer = FParticleCurveInjectionBufferRef::CreateUniformBufferImmediate(Parameters, UniformBuffer_SingleDraw);
-		FVertexShaderRHIParamRef VertexShader = GetVertexShader();
+		FRHIVertexShader* VertexShader = GetVertexShader();
 		SetUniformBufferParameter(RHICmdList, VertexShader, GetUniformBufferParameter<FParticleCurveInjectionParameters>(), UniformBuffer);
 	}
 };
diff --git a/Engine/Source/Runtime/Engine/Private/Particles/ParticleGpuSimulation.cpp b/Engine/Source/Runtime/Engine/Private/Particles/ParticleGpuSimulation.cpp
--- a/Engine/Source/Runtime/Engine/Private/Particles/ParticleGpuSimulation.cpp
+++ b/Engine/Source/Runtime/Engine/Private/Particles/ParticleGpuSimulation.cpp
@@ -1027,7 +1027,7 @@ public:
 	/** Set parameters. */
 	void SetParameters(FRHICommandList& RHICmdList, FParticleShaderParamRef TileOffsetsRef)
 	{
-		FVertexShaderRHIParamRef VertexShaderRHI = GetVertexShader();
+		FRHIVertexShader* VertexShaderRHI = GetVertexShader();
 		if (TileOffsets.IsBound())
 		{
 			RHICmdList.SetShaderResourceViewParameter(VertexShaderRHI, TileOffsets.GetBaseIndex(), TileOffsetsRef);
@@ -1143,7 +1143,7 @@ public:
 		FUniformBufferRHIParamRef SceneTexturesUniformBuffer
 		)
 	{
-		FPixelShaderRHIParamRef PixelShaderRHI = GetPixelShader();
+		FRHIPixelShader* PixelShaderRHI = GetPixelShader();
 		FRHISamplerState* SamplerStatePoint = TStaticSamplerState<SF_Point>::GetRHI();
 		FRHISamplerState* SamplerStateLinear = TStaticSamplerState<SF_Bilinear,AM_Clamp,AM_Clamp,AM_Clamp>::GetRHI();
 		SetTextureParameter(RHICmdList, PixelShaderRHI, PositionTexture, PositionTextureSampler, SamplerStatePoint, TextureResources.PositionTextureRHI);
@@ -1188,7 +1188,7 @@ public:
 	 */
 	void SetVectorFieldParameters(FRHICommandList& RHICmdList, const FVectorFieldUniformBufferRef& UniformBuffer, const FTexture3DRHIParamRef VolumeTexturesRHI[])
 	{
-		FPixelShaderRHIParamRef PixelShaderRHI = GetPixelShader();
+		FRHIPixelShader* PixelShaderRHI = GetPixelShader();
 			SetUniformBufferParameter(RHICmdList, PixelShaderRHI, GetUniformBufferParameter<FVectorFieldUniformParameters>(), UniformBuffer);
 		
 			FRHISamplerState* SamplerStateLinear = TStaticSamplerState<SF_Bilinear,AM_Clamp,AM_Clamp,AM_Clamp>::GetRHI();
@@ -1205,7 +1205,7 @@ public:
 	 */
 	void SetInstanceParameters(FRHICommandList& RHICmdList, FUniformBufferRHIParamRef UniformBuffer, const FParticlePerFrameSimulationParameters& InPerFrameParameters, bool bUseFixDT)
 	{
-		FPixelShaderRHIParamRef PixelShaderRHI = GetPixelShader();
+		FRHIPixelShader* PixelShaderRHI = GetPixelShader();
 		SetUniformBufferParameter(RHICmdList, PixelShaderRHI, GetUniformBufferParameter<FParticleSimulationParameters>(), UniformBuffer);
 		PerFrameParameters.Set(RHICmdList, PixelShaderRHI, InPerFrameParameters, bUseFixDT);
 	}
@@ -1215,7 +1215,7 @@ public:
 	 */
 	void UnbindBuffers(FRHICommandList& RHICmdList)
 	{
-		FPixelShaderRHIParamRef PixelShaderRHI = GetPixelShader();
+		FRHIPixelShader* PixelShaderRHI = GetPixelShader();
 		FRHIShaderResourceView* NullSRV = nullptr;
 		for (int32 i = 0; i < MAX_VECTOR_FIELDS; ++i)
 		{
@@ -1746,7 +1746,7 @@ public:
 		Parameters.PixelScale.X = 1.0f / GParticleSimulationTextureSizeX;
 		Parameters.PixelScale.Y = 1.0f / GParticleSimulationTextureSizeY;
 		FParticleInjectionBufferRef UniformBuffer = FParticleInjectionBufferRef::CreateUniformBufferImmediate( Parameters, UniformBuffer_SingleDraw );
-		FVertexShaderRHIParamRef VertexShader = GetVertexShader();
+		FRHIVertexShader* VertexShader = GetVertexShader();
 		SetUniformBufferParameter(RHICmdList, VertexShader, GetUniformBufferParameter<FParticleInjectionParameters>(), UniformBuffer );
 	}
 };
@@ -1977,7 +1977,7 @@ public:
 	 */
 	void SetParameters(FRHICommandList& RHICmdList, const FParticleSimVisualizeBufferRef& UniformBuffer )
 	{
-		FVertexShaderRHIParamRef VertexShader = GetVertexShader();
+		FRHIVertexShader* VertexShader = GetVertexShader();
 		SetUniformBufferParameter(RHICmdList, VertexShader, GetUniformBufferParameter<FParticleSimVisualizeParameters>(), UniformBuffer );
 	}
 };
@@ -2029,7 +2029,7 @@ public:
 	 */
 	void SetParameters(FRHICommandList& RHICmdList, int32 InVisualizationMode, FTexture2DRHIParamRef PositionTextureRHI, FTexture2DRHIParamRef CurveTextureRHI )
 	{
-		FPixelShaderRHIParamRef PixelShader = GetPixelShader();
+		FRHIPixelShader* PixelShader = GetPixelShader();
 		SetShaderValue(RHICmdList, PixelShader, VisualizationMode, InVisualizationMode );
 		FRHISamplerState* SamplerStatePoint = TStaticSamplerState<SF_Point>::GetRHI();
 		SetTextureParameter(RHICmdList, PixelShader, PositionTexture, PositionTextureSampler, SamplerStatePoint, PositionTextureRHI );
@@ -2271,7 +2271,7 @@ public:
 	 */
 	void SetOutput(FRHICommandList& RHICmdList, FRHIUnorderedAccessView* OutBoundsUAV )
 	{
-		FComputeShaderRHIParamRef ComputeShaderRHI = GetComputeShader();
+		FRHIComputeShader* ComputeShaderRHI = GetComputeShader();
 		if ( OutBounds.IsBound() )
 		{
 			RHICmdList.SetUAVParameter(ComputeShaderRHI, OutBounds.GetBaseIndex(), OutBoundsUAV);
@@ -2288,7 +2288,7 @@ public:
 		FTexture2DRHIParamRef PositionTextureRHI
 		)
 	{
-		FComputeShaderRHIParamRef ComputeShaderRHI = GetComputeShader();
+		FRHIComputeShader* ComputeShaderRHI = GetComputeShader();
 		SetUniformBufferParameter(RHICmdList, ComputeShaderRHI, GetUniformBufferParameter<FParticleBoundsParameters>(), UniformBuffer );
 		if ( InParticleIndices.IsBound() )
 		{
@@ -2308,7 +2308,7 @@ public:
 	 */
 	void UnbindBuffers(FRHICommandList& RHICmdList)
 	{
-		FComputeShaderRHIParamRef ComputeShaderRHI = GetComputeShader();
+		FRHIComputeShader* ComputeShaderRHI = GetComputeShader();
 		if ( InParticleIndices.IsBound() )
 		{
 			RHICmdList.SetShaderResourceViewParameter(ComputeShaderRHI, InParticleIndices.GetBaseIndex(), nullptr);
diff --git a/Engine/Source/Runtime/Engine/Private/Particles/ParticleSortingGPU.cpp b/Engine/Source/Runtime/Engine/Private/Particles/ParticleSortingGPU.cpp
--- a/Engine/Source/Runtime/Engine/Private/Particles/ParticleSortingGPU.cpp
+++ b/Engine/Source/Runtime/Engine/Private/Particles/ParticleSortingGPU.cpp
@@ -93,7 +93,7 @@ public:
 	 */
 	void SetOutput(FRHICommandList& RHICmdList, FRHIUnorderedAccessView* OutKeysUAV, FRHIUnorderedAccessView* OutIndicesUAV )
 	{
-		FComputeShaderRHIParamRef ComputeShaderRHI = GetComputeShader();
+		FRHIComputeShader* ComputeShaderRHI = GetComputeShader();
 		if ( OutKeys.IsBound() )
 		{
 			RHICmdList.SetUAVParameter(ComputeShaderRHI, OutKeys.GetBaseIndex(), OutKeysUAV);
@@ -113,7 +113,7 @@ public:
 		FRHIShaderResourceView* InIndicesSRV
 		)
 	{
-		FComputeShaderRHIParamRef ComputeShaderRHI = GetComputeShader();
+		FRHIComputeShader* ComputeShaderRHI = GetComputeShader();
 		SetUniformBufferParameter(RHICmdList, ComputeShaderRHI, GetUniformBufferParameter<FParticleKeyGenParameters>(), UniformBuffer );
 		if ( InParticleIndices.IsBound() )
 		{
@@ -126,7 +126,7 @@ public:
 	 */
 	void SetPositionTextures(FRHICommandList& RHICmdList, FTexture2DRHIParamRef PositionTextureRHI)
 	{
-		FComputeShaderRHIParamRef ComputeShaderRHI = GetComputeShader();
+		FRHIComputeShader* ComputeShaderRHI = GetComputeShader();
 		if (PositionTexture.IsBound())
 		{
 			RHICmdList.SetShaderTexture(ComputeShaderRHI, PositionTexture.GetBaseIndex(), PositionTextureRHI);
@@ -138,7 +138,7 @@ public:
 	 */
 	void UnbindBuffers(FRHICommandList& RHICmdList)
 	{
-		FComputeShaderRHIParamRef ComputeShaderRHI = GetComputeShader();
+		FRHIComputeShader* ComputeShaderRHI = GetComputeShader();
 		if ( InParticleIndices.IsBound() )
 		{
 			RHICmdList.SetShaderResourceViewParameter(ComputeShaderRHI, InParticleIndices.GetBaseIndex(), nullptr);
diff --git a/Engine/Source/Runtime/Engine/Private/SkeletalRenderGPUSkin.cpp b/Engine/Source/Runtime/Engine/Private/SkeletalRenderGPUSkin.cpp
--- a/Engine/Source/Runtime/Engine/Private/SkeletalRenderGPUSkin.cpp
+++ b/Engine/Source/Runtime/Engine/Private/SkeletalRenderGPUSkin.cpp
@@ -606,7 +606,7 @@ TArray<float> FSkeletalMeshObjectGPUSkin::FSkeletalMeshObjectLOD::MorphAccumulat
 
 void FGPUMorphUpdateCS::SetParameters(FRHICommandList& RHICmdList, const FVector4& LocalScale, const FMorphTargetVertexInfoBuffers& MorphTargetVertexInfoBuffers, FMorphVertexBuffer& MorphVertexBuffer)
 {
-	FComputeShaderRHIRef CS = GetComputeShader();
+	FRHIComputeShader* CS = GetComputeShader();
 	RHICmdList.SetComputeShader(CS);
 
 	SetUAVParameter(RHICmdList, CS, MorphVertexBufferParameter, MorphVertexBuffer.GetUAV());
diff --git a/Engine/Source/Runtime/Engine/Private/VectorField.cpp b/Engine/Source/Runtime/Engine/Private/VectorField.cpp
--- a/Engine/Source/Runtime/Engine/Private/VectorField.cpp
+++ b/Engine/Source/Runtime/Engine/Private/VectorField.cpp
@@ -811,7 +811,7 @@ public:
 		FTextureRHIParamRef AtlasTextureRHI,
 		FTextureRHIParamRef NoiseVolumeTextureRHI )
 	{
-		FComputeShaderRHIParamRef ComputeShaderRHI = GetComputeShader();
+		FRHIComputeShader* ComputeShaderRHI = GetComputeShader();
 		FRHISamplerState* SamplerStateLinear = TStaticSamplerState<SF_Bilinear,AM_Clamp,AM_Clamp,AM_Clamp>::GetRHI();
 		SetUniformBufferParameter(RHICmdList, ComputeShaderRHI, GetUniformBufferParameter<FCompositeAnimatedVectorFieldUniformParameters>(), UniformBuffer );
 		SetTextureParameter(RHICmdList, ComputeShaderRHI, AtlasTexture, AtlasTextureSampler, SamplerStateLinear, AtlasTextureRHI );
@@ -823,7 +823,7 @@ public:
 	 */
 	void SetOutput(FRHICommandList& RHICmdList, FRHIUnorderedAccessView* VolumeTextureUAV)
 	{
-		FComputeShaderRHIParamRef ComputeShaderRHI = GetComputeShader();
+		FRHIComputeShader* ComputeShaderRHI = GetComputeShader();
 		if ( OutVolumeTexture.IsBound() )
 		{
 			RHICmdList.SetUAVParameter(ComputeShaderRHI, OutVolumeTexture.GetBaseIndex(), VolumeTextureUAV);
@@ -835,7 +835,7 @@ public:
 	 */
 	void UnbindBuffers(FRHICommandList& RHICmdList)
 	{
-		FComputeShaderRHIParamRef ComputeShaderRHI = GetComputeShader();
+		FRHIComputeShader* ComputeShaderRHI = GetComputeShader();
 		if ( OutVolumeTexture.IsBound() )
 		{
 			RHICmdList.SetUAVParameter(ComputeShaderRHI, OutVolumeTexture.GetBaseIndex(), nullptr);
diff --git a/Engine/Source/Runtime/Engine/Public/StereoLayerRendering.h b/Engine/Source/Runtime/Engine/Public/StereoLayerRendering.h
--- a/Engine/Source/Runtime/Engine/Public/StereoLayerRendering.h
+++ b/Engine/Source/Runtime/Engine/Public/StereoLayerRendering.h
@@ -34,7 +34,7 @@ public:
 
 	void SetParameters(FRHICommandList& RHICmdList, FVector2D QuadSize, FBox2D UVRect, const FMatrix& ViewProjection, const FMatrix& World)
 	{
-		FVertexShaderRHIParamRef VS = GetVertexShader();
+		FRHIVertexShader* VS = GetVertexShader();
 
 		if (InQuadAdjust.IsBound())
 		{
@@ -99,7 +99,7 @@ public:
 
 	void SetParameters(FRHICommandList& RHICmdList, FRHISamplerState* SamplerStateRHI, FTextureRHIParamRef TextureRHI)
 	{
-		FPixelShaderRHIParamRef PS = GetPixelShader();
+		FRHIPixelShader* PS = GetPixelShader();
 
 		SetTextureParameter(RHICmdList, PS,InTexture,InTextureSampler,SamplerStateRHI,TextureRHI);
 	}
diff --git a/Engine/Source/Runtime/NullDrv/Public/NullRHI.h b/Engine/Source/Runtime/NullDrv/Public/NullRHI.h
--- a/Engine/Source/Runtime/NullDrv/Public/NullRHI.h
+++ b/Engine/Source/Runtime/NullDrv/Public/NullRHI.h
@@ -78,12 +78,12 @@ public:
 	}
 
 
-	virtual FBoundShaderStateRHIRef RHICreateBoundShaderState(FRHIVertexDeclaration* VertexDeclaration, FVertexShaderRHIParamRef VertexShader, FHullShaderRHIParamRef HullShader, FDomainShaderRHIParamRef DomainShader, FPixelShaderRHIParamRef PixelShader, FGeometryShaderRHIParamRef GeometryShader) final override
+	virtual FBoundShaderStateRHIRef RHICreateBoundShaderState(FRHIVertexDeclaration* VertexDeclaration, FRHIVertexShader* VertexShader, FRHIHullShader* HullShader, FRHIDomainShader* DomainShader, FRHIPixelShader* PixelShader, FRHIGeometryShader* GeometryShader) final override
 	{ 
 		return new FRHIBoundShaderState(); 
 	}
 
-	virtual void RHISetComputeShader(FComputeShaderRHIParamRef ComputeShader) final override
+	virtual void RHISetComputeShader(FRHIComputeShader* ComputeShader) final override
 	{
 
 	}
@@ -537,167 +537,167 @@ public:
 	}
 
 
-	virtual void RHISetShaderTexture(FVertexShaderRHIParamRef VertexShader, uint32 TextureIndex, FTextureRHIParamRef NewTexture) final override
+	virtual void RHISetShaderTexture(FRHIVertexShader* VertexShader, uint32 TextureIndex, FTextureRHIParamRef NewTexture) final override
 	{
 
 	}
 
 
-	virtual void RHISetShaderTexture(FHullShaderRHIParamRef HullShader, uint32 TextureIndex, FTextureRHIParamRef NewTexture) final override
+	virtual void RHISetShaderTexture(FRHIHullShader* HullShader, uint32 TextureIndex, FTextureRHIParamRef NewTexture) final override
 	{
 
 	}
 
 
-	virtual void RHISetShaderTexture(FDomainShaderRHIParamRef DomainShader, uint32 TextureIndex, FTextureRHIParamRef NewTexture) final override
+	virtual void RHISetShaderTexture(FRHIDomainShader* DomainShader, uint32 TextureIndex, FTextureRHIParamRef NewTexture) final override
 	{
 
 	}
 
 
-	virtual void RHISetShaderTexture(FGeometryShaderRHIParamRef GeometryShader, uint32 TextureIndex, FTextureRHIParamRef NewTexture) final override
+	virtual void RHISetShaderTexture(FRHIGeometryShader* GeometryShader, uint32 TextureIndex, FTextureRHIParamRef NewTexture) final override
 	{
 
 	}
 
 
-	virtual void RHISetShaderTexture(FPixelShaderRHIParamRef PixelShader, uint32 TextureIndex, FTextureRHIParamRef NewTexture) final override
+	virtual void RHISetShaderTexture(FRHIPixelShader* PixelShader, uint32 TextureIndex, FTextureRHIParamRef NewTexture) final override
 	{
 
 	}
 
 
-	virtual void RHISetShaderTexture(FComputeShaderRHIParamRef PixelShader, uint32 TextureIndex, FTextureRHIParamRef NewTexture) final override
+	virtual void RHISetShaderTexture(FRHIComputeShader* PixelShader, uint32 TextureIndex, FTextureRHIParamRef NewTexture) final override
 	{
 
 	}
 
-	virtual void RHISetShaderSampler(FComputeShaderRHIParamRef ComputeShader, uint32 SamplerIndex, FRHISamplerState* NewState) final override
+	virtual void RHISetShaderSampler(FRHIComputeShader* ComputeShader, uint32 SamplerIndex, FRHISamplerState* NewState) final override
 	{
 
 	}
 
-	virtual void RHISetShaderSampler(FVertexShaderRHIParamRef VertexShader, uint32 SamplerIndex, FRHISamplerState* NewState) final override
+	virtual void RHISetShaderSampler(FRHIVertexShader* VertexShader, uint32 SamplerIndex, FRHISamplerState* NewState) final override
 	{
 
 	}
 
-	virtual void RHISetShaderSampler(FGeometryShaderRHIParamRef GeometryShader, uint32 SamplerIndex, FRHISamplerState* NewState) final override
+	virtual void RHISetShaderSampler(FRHIGeometryShader* GeometryShader, uint32 SamplerIndex, FRHISamplerState* NewState) final override
 	{
 
 	}
-	virtual void RHISetShaderSampler(FDomainShaderRHIParamRef DomainShader, uint32 SamplerIndex, FRHISamplerState* NewState) final override
+	virtual void RHISetShaderSampler(FRHIDomainShader* DomainShader, uint32 SamplerIndex, FRHISamplerState* NewState) final override
 	{
 
 	}
-	virtual void RHISetShaderSampler(FHullShaderRHIParamRef HullShader, uint32 SamplerIndex, FRHISamplerState* NewState) final override
+	virtual void RHISetShaderSampler(FRHIHullShader* HullShader, uint32 SamplerIndex, FRHISamplerState* NewState) final override
 	{
 
 	}
 
-	virtual void RHISetShaderSampler(FPixelShaderRHIParamRef PixelShader, uint32 SamplerIndex, FRHISamplerState* NewState) final override
+	virtual void RHISetShaderSampler(FRHIPixelShader* PixelShader, uint32 SamplerIndex, FRHISamplerState* NewState) final override
 	{
 
 	}
 
 
-	virtual void RHISetUAVParameter(FComputeShaderRHIParamRef ComputeShader, uint32 UAVIndex, FRHIUnorderedAccessView* UAV) final override
+	virtual void RHISetUAVParameter(FRHIComputeShader* ComputeShader, uint32 UAVIndex, FRHIUnorderedAccessView* UAV) final override
 	{
 
 	}
 
 
-	virtual void RHISetUAVParameter(FComputeShaderRHIParamRef ComputeShader, uint32 UAVIndex, FRHIUnorderedAccessView* UAV, uint32 InitialCount) final override
+	virtual void RHISetUAVParameter(FRHIComputeShader* ComputeShader, uint32 UAVIndex, FRHIUnorderedAccessView* UAV, uint32 InitialCount) final override
 	{
 
 	}
 
 
-	virtual void RHISetShaderResourceViewParameter(FPixelShaderRHIParamRef PixelShader, uint32 SamplerIndex, FRHIShaderResourceView* SRV) final override
+	virtual void RHISetShaderResourceViewParameter(FRHIPixelShader* PixelShader, uint32 SamplerIndex, FRHIShaderResourceView* SRV) final override
 	{
 
 	}
 
 
-	virtual void RHISetShaderResourceViewParameter(FVertexShaderRHIParamRef VertexShader, uint32 SamplerIndex, FRHIShaderResourceView* SRV) final override
+	virtual void RHISetShaderResourceViewParameter(FRHIVertexShader* VertexShader, uint32 SamplerIndex, FRHIShaderResourceView* SRV) final override
 	{
 
 	}
 
 
-	virtual void RHISetShaderResourceViewParameter(FComputeShaderRHIParamRef ComputeShader, uint32 SamplerIndex, FRHIShaderResourceView* SRV) final override
+	virtual void RHISetShaderResourceViewParameter(FRHIComputeShader* ComputeShader, uint32 SamplerIndex, FRHIShaderResourceView* SRV) final override
 	{
 
 	}
 
 
-	virtual void RHISetShaderResourceViewParameter(FHullShaderRHIParamRef HullShader, uint32 SamplerIndex, FRHIShaderResourceView* SRV) final override
+	virtual void RHISetShaderResourceViewParameter(FRHIHullShader* HullShader, uint32 SamplerIndex, FRHIShaderResourceView* SRV) final override
 	{
 
 	}
 
 
-	virtual void RHISetShaderResourceViewParameter(FDomainShaderRHIParamRef DomainShader, uint32 SamplerIndex, FRHIShaderResourceView* SRV) final override
+	virtual void RHISetShaderResourceViewParameter(FRHIDomainShader* DomainShader, uint32 SamplerIndex, FRHIShaderResourceView* SRV) final override
 	{
 
 	}
-	virtual void RHISetShaderResourceViewParameter(FGeometryShaderRHIParamRef GeometryShader, uint32 SamplerIndex, FRHIShaderResourceView* SRV) final override
+	virtual void RHISetShaderResourceViewParameter(FRHIGeometryShader* GeometryShader, uint32 SamplerIndex, FRHIShaderResourceView* SRV) final override
 	{
 
 	}
 
-	virtual void RHISetShaderUniformBuffer(FVertexShaderRHIParamRef VertexShader, uint32 BufferIndex, FUniformBufferRHIParamRef Buffer) final override
+	virtual void RHISetShaderUniformBuffer(FRHIVertexShader* VertexShader, uint32 BufferIndex, FUniformBufferRHIParamRef Buffer) final override
 	{
 
 	}
 
-	virtual void RHISetShaderUniformBuffer(FHullShaderRHIParamRef HullShader, uint32 BufferIndex, FUniformBufferRHIParamRef Buffer) final override
+	virtual void RHISetShaderUniformBuffer(FRHIHullShader* HullShader, uint32 BufferIndex, FUniformBufferRHIParamRef Buffer) final override
 	{
 
 	}
 
-	virtual void RHISetShaderUniformBuffer(FDomainShaderRHIParamRef DomainShader, uint32 BufferIndex, FUniformBufferRHIParamRef Buffer) final override
+	virtual void RHISetShaderUniformBuffer(FRHIDomainShader* DomainShader, uint32 BufferIndex, FUniformBufferRHIParamRef Buffer) final override
 	{
 
 	}
 
-	virtual void RHISetShaderUniformBuffer(FGeometryShaderRHIParamRef GeometryShader, uint32 BufferIndex, FUniformBufferRHIParamRef Buffer) final override
+	virtual void RHISetShaderUniformBuffer(FRHIGeometryShader* GeometryShader, uint32 BufferIndex, FUniformBufferRHIParamRef Buffer) final override
 	{
 
 	}
 
-	virtual void RHISetShaderUniformBuffer(FPixelShaderRHIParamRef PixelShader, uint32 BufferIndex, FUniformBufferRHIParamRef Buffer) final override
+	virtual void RHISetShaderUniformBuffer(FRHIPixelShader* PixelShader, uint32 BufferIndex, FUniformBufferRHIParamRef Buffer) final override
 	{
 
 	}
 
-	virtual void RHISetShaderUniformBuffer(FComputeShaderRHIParamRef ComputeShader, uint32 BufferIndex, FUniformBufferRHIParamRef Buffer) final override
+	virtual void RHISetShaderUniformBuffer(FRHIComputeShader* ComputeShader, uint32 BufferIndex, FUniformBufferRHIParamRef Buffer) final override
 	{
 
 	}
 
-	virtual void RHISetShaderParameter(FVertexShaderRHIParamRef VertexShader, uint32 BufferIndex, uint32 BaseIndex, uint32 NumBytes, const void* NewValue) final override
+	virtual void RHISetShaderParameter(FRHIVertexShader* VertexShader, uint32 BufferIndex, uint32 BaseIndex, uint32 NumBytes, const void* NewValue) final override
 	{
 
 	}
-	virtual void RHISetShaderParameter(FPixelShaderRHIParamRef PixelShader, uint32 BufferIndex, uint32 BaseIndex, uint32 NumBytes, const void* NewValue) final override
+	virtual void RHISetShaderParameter(FRHIPixelShader* PixelShader, uint32 BufferIndex, uint32 BaseIndex, uint32 NumBytes, const void* NewValue) final override
 	{
 
 	}
 
-	virtual void RHISetShaderParameter(FHullShaderRHIParamRef HullShader, uint32 BufferIndex, uint32 BaseIndex, uint32 NumBytes, const void* NewValue) final override
+	virtual void RHISetShaderParameter(FRHIHullShader* HullShader, uint32 BufferIndex, uint32 BaseIndex, uint32 NumBytes, const void* NewValue) final override
 	{
 
 	}
-	virtual void RHISetShaderParameter(FDomainShaderRHIParamRef DomainShader, uint32 BufferIndex, uint32 BaseIndex, uint32 NumBytes, const void* NewValue) final override
+	virtual void RHISetShaderParameter(FRHIDomainShader* DomainShader, uint32 BufferIndex, uint32 BaseIndex, uint32 NumBytes, const void* NewValue) final override
 	{
 
 	}
-	virtual void RHISetShaderParameter(FGeometryShaderRHIParamRef GeometryShader, uint32 BufferIndex, uint32 BaseIndex, uint32 NumBytes, const void* NewValue) final override
+	virtual void RHISetShaderParameter(FRHIGeometryShader* GeometryShader, uint32 BufferIndex, uint32 BaseIndex, uint32 NumBytes, const void* NewValue) final override
 	{
 
 	}
-	virtual void RHISetShaderParameter(FComputeShaderRHIParamRef ComputeShader, uint32 BufferIndex, uint32 BaseIndex, uint32 NumBytes, const void* NewValue) final override
+	virtual void RHISetShaderParameter(FRHIComputeShader* ComputeShader, uint32 BufferIndex, uint32 BaseIndex, uint32 NumBytes, const void* NewValue) final override
 	{
 
 	}
diff --git a/Engine/Source/Runtime/OpenGLDrv/Private/OpenGLCommands.cpp b/Engine/Source/Runtime/OpenGLDrv/Private/OpenGLCommands.cpp
--- a/Engine/Source/Runtime/OpenGLDrv/Private/OpenGLCommands.cpp
+++ b/Engine/Source/Runtime/OpenGLDrv/Private/OpenGLCommands.cpp
@@ -16,7 +16,7 @@
 #include "OpenGLDrvPrivate.h"
 #include "RenderUtils.h"
 
-#define DECLARE_ISBOUNDSHADER(ShaderType) inline void ValidateBoundShader(TRefCountPtr<FOpenGLBoundShaderState> InBoundShaderState, F##ShaderType##RHIParamRef ShaderType##RHI) \
+#define DECLARE_ISBOUNDSHADER(ShaderType) inline void ValidateBoundShader(TRefCountPtr<FOpenGLBoundShaderState> InBoundShaderState, FRHI##ShaderType* ShaderType##RHI) \
 { \
 	FOpenGL##ShaderType* ShaderType = FOpenGLDynamicRHI::ResourceCast(ShaderType##RHI); \
 	ensureMsgf(InBoundShaderState && ShaderType == InBoundShaderState->Get##ShaderType(), TEXT("Parameters are being set for a %s which is not currently bound"), TEXT(#ShaderType)); \
@@ -445,7 +445,7 @@ void FOpenGLDynamicRHI::RHISetBoundShaderState(FRHIBoundShaderState* BoundShader
 	BoundShaderStateHistory.Add(BoundShaderState);
 }
 
-void FOpenGLDynamicRHI::RHISetUAVParameter(FComputeShaderRHIParamRef ComputeShaderRHI,uint32 UAVIndex, FRHIUnorderedAccessView* UnorderedAccessViewRHI)
+void FOpenGLDynamicRHI::RHISetUAVParameter(FRHIComputeShader* ComputeShaderRHI,uint32 UAVIndex, FRHIUnorderedAccessView* UnorderedAccessViewRHI)
 {
 	check(FOpenGL::SupportsComputeShaders());
 	
@@ -461,7 +461,7 @@ void FOpenGLDynamicRHI::RHISetUAVParameter(FComputeShaderRHIParamRef ComputeShad
 	}
 }
 
-void FOpenGLDynamicRHI::RHISetUAVParameter(FComputeShaderRHIParamRef ComputeShaderRHI,uint32 UAVIndex, FRHIUnorderedAccessView* UAVRHI, uint32 InitialCount )
+void FOpenGLDynamicRHI::RHISetUAVParameter(FRHIComputeShader* ComputeShaderRHI,uint32 UAVIndex, FRHIUnorderedAccessView* UAVRHI, uint32 InitialCount )
 {
 	// TODO: Implement for OpenGL
 	check(0);
@@ -843,7 +843,7 @@ void FOpenGLDynamicRHI::UpdateSRV(FOpenGLShaderResourceView* SRV)
 #endif
 }
 
-void FOpenGLDynamicRHI::RHISetShaderResourceViewParameter(FPixelShaderRHIParamRef PixelShaderRHI,uint32 TextureIndex, FRHIShaderResourceView* SRVRHI)
+void FOpenGLDynamicRHI::RHISetShaderResourceViewParameter(FRHIPixelShader* PixelShaderRHI,uint32 TextureIndex, FRHIShaderResourceView* SRVRHI)
 {
 	VERIFY_GL_SCOPE();
 	VALIDATE_BOUND_SHADER(PixelShaderRHI);
@@ -869,7 +869,7 @@ void FOpenGLDynamicRHI::RHISetShaderResourceViewParameter(FPixelShaderRHIParamRe
 	RHISetShaderSampler(PixelShaderRHI,TextureIndex,PointSamplerState);
 }
 
-void FOpenGLDynamicRHI::RHISetShaderResourceViewParameter(FVertexShaderRHIParamRef VertexShaderRHI,uint32 TextureIndex, FRHIShaderResourceView* SRVRHI)
+void FOpenGLDynamicRHI::RHISetShaderResourceViewParameter(FRHIVertexShader* VertexShaderRHI,uint32 TextureIndex, FRHIShaderResourceView* SRVRHI)
 {
 	VALIDATE_BOUND_SHADER(VertexShaderRHI);
 
@@ -891,7 +891,7 @@ void FOpenGLDynamicRHI::RHISetShaderResourceViewParameter(FVertexShaderRHIParamR
 	RHISetShaderSampler(VertexShaderRHI,TextureIndex,PointSamplerState);
 }
 
-void FOpenGLDynamicRHI::RHISetShaderResourceViewParameter(FComputeShaderRHIParamRef ComputeShaderRHI,uint32 TextureIndex, FRHIShaderResourceView* SRVRHI)
+void FOpenGLDynamicRHI::RHISetShaderResourceViewParameter(FRHIComputeShader* ComputeShaderRHI,uint32 TextureIndex, FRHIShaderResourceView* SRVRHI)
 {
 	VERIFY_GL_SCOPE();
 	check(FOpenGL::SupportsComputeShaders());
@@ -912,7 +912,7 @@ void FOpenGLDynamicRHI::RHISetShaderResourceViewParameter(FComputeShaderRHIParam
 	RHISetShaderSampler(ComputeShaderRHI,TextureIndex,PointSamplerState);
 }
 
-void FOpenGLDynamicRHI::RHISetShaderResourceViewParameter(FHullShaderRHIParamRef HullShaderRHI,uint32 TextureIndex, FRHIShaderResourceView* SRVRHI)
+void FOpenGLDynamicRHI::RHISetShaderResourceViewParameter(FRHIHullShader* HullShaderRHI,uint32 TextureIndex, FRHIShaderResourceView* SRVRHI)
 {
 	VALIDATE_BOUND_SHADER(HullShaderRHI);
 
@@ -934,7 +934,7 @@ void FOpenGLDynamicRHI::RHISetShaderResourceViewParameter(FHullShaderRHIParamRef
 	InternalSetShaderTexture(NULL, SRV, FOpenGL::GetFirstHullTextureUnit() + TextureIndex, Target, Resource, 0, LimitMip);
 }
 
-void FOpenGLDynamicRHI::RHISetShaderResourceViewParameter(FDomainShaderRHIParamRef DomainShaderRHI,uint32 TextureIndex, FRHIShaderResourceView* SRVRHI)
+void FOpenGLDynamicRHI::RHISetShaderResourceViewParameter(FRHIDomainShader* DomainShaderRHI,uint32 TextureIndex, FRHIShaderResourceView* SRVRHI)
 {
 	VALIDATE_BOUND_SHADER(DomainShaderRHI);
 
@@ -956,7 +956,7 @@ void FOpenGLDynamicRHI::RHISetShaderResourceViewParameter(FDomainShaderRHIParamR
 	InternalSetShaderTexture(NULL, SRV, FOpenGL::GetFirstDomainTextureUnit() + TextureIndex, Target, Resource, 0, LimitMip);
 }
 
-void FOpenGLDynamicRHI::RHISetShaderResourceViewParameter(FGeometryShaderRHIParamRef GeometryShaderRHI,uint32 TextureIndex, FRHIShaderResourceView* SRVRHI)
+void FOpenGLDynamicRHI::RHISetShaderResourceViewParameter(FRHIGeometryShader* GeometryShaderRHI,uint32 TextureIndex, FRHIShaderResourceView* SRVRHI)
 {
 	VALIDATE_BOUND_SHADER(GeometryShaderRHI);
 
@@ -978,7 +978,7 @@ void FOpenGLDynamicRHI::RHISetShaderResourceViewParameter(FGeometryShaderRHIPara
 	RHISetShaderSampler(GeometryShaderRHI,TextureIndex,PointSamplerState);
 }
 
-void FOpenGLDynamicRHI::RHISetShaderTexture(FVertexShaderRHIParamRef VertexShaderRHI,uint32 TextureIndex,FTextureRHIParamRef NewTextureRHI)
+void FOpenGLDynamicRHI::RHISetShaderTexture(FRHIVertexShader* VertexShaderRHI,uint32 TextureIndex,FTextureRHIParamRef NewTextureRHI)
 {
 	VALIDATE_BOUND_SHADER(VertexShaderRHI);
 
@@ -995,7 +995,7 @@ void FOpenGLDynamicRHI::RHISetShaderTexture(FVertexShaderRHIParamRef VertexShade
 	}
 }
 
-void FOpenGLDynamicRHI::RHISetShaderTexture(FHullShaderRHIParamRef HullShaderRHI,uint32 TextureIndex,FTextureRHIParamRef NewTextureRHI)
+void FOpenGLDynamicRHI::RHISetShaderTexture(FRHIHullShader* HullShaderRHI,uint32 TextureIndex,FTextureRHIParamRef NewTextureRHI)
 {
 	VALIDATE_BOUND_SHADER(HullShaderRHI);
 	
@@ -1013,7 +1013,7 @@ void FOpenGLDynamicRHI::RHISetShaderTexture(FHullShaderRHIParamRef HullShaderRHI
 	}
 }
 
-void FOpenGLDynamicRHI::RHISetShaderTexture(FDomainShaderRHIParamRef DomainShaderRHI,uint32 TextureIndex,FTextureRHIParamRef NewTextureRHI)
+void FOpenGLDynamicRHI::RHISetShaderTexture(FRHIDomainShader* DomainShaderRHI,uint32 TextureIndex,FTextureRHIParamRef NewTextureRHI)
 {
 	VALIDATE_BOUND_SHADER(DomainShaderRHI);
 
@@ -1031,7 +1031,7 @@ void FOpenGLDynamicRHI::RHISetShaderTexture(FDomainShaderRHIParamRef DomainShade
 	}
 }
 
-void FOpenGLDynamicRHI::RHISetShaderTexture(FGeometryShaderRHIParamRef GeometryShaderRHI,uint32 TextureIndex,FTextureRHIParamRef NewTextureRHI)
+void FOpenGLDynamicRHI::RHISetShaderTexture(FRHIGeometryShader* GeometryShaderRHI,uint32 TextureIndex,FTextureRHIParamRef NewTextureRHI)
 {
 	VALIDATE_BOUND_SHADER(GeometryShaderRHI);
 
@@ -1048,7 +1048,7 @@ void FOpenGLDynamicRHI::RHISetShaderTexture(FGeometryShaderRHIParamRef GeometryS
 	}
 }
 
-void FOpenGLDynamicRHI::RHISetShaderTexture(FPixelShaderRHIParamRef PixelShaderRHI,uint32 TextureIndex,FTextureRHIParamRef NewTextureRHI)
+void FOpenGLDynamicRHI::RHISetShaderTexture(FRHIPixelShader* PixelShaderRHI,uint32 TextureIndex,FTextureRHIParamRef NewTextureRHI)
 {
 	VALIDATE_BOUND_SHADER(PixelShaderRHI);
 
@@ -1065,7 +1065,7 @@ void FOpenGLDynamicRHI::RHISetShaderTexture(FPixelShaderRHIParamRef PixelShaderR
 	}
 }
 
-void FOpenGLDynamicRHI::RHISetShaderSampler(FVertexShaderRHIParamRef VertexShaderRHI,uint32 SamplerIndex, FRHISamplerState* NewStateRHI)
+void FOpenGLDynamicRHI::RHISetShaderSampler(FRHIVertexShader* VertexShaderRHI,uint32 SamplerIndex, FRHISamplerState* NewStateRHI)
 {
 	VALIDATE_BOUND_SHADER(VertexShaderRHI);
 
@@ -1074,7 +1074,7 @@ void FOpenGLDynamicRHI::RHISetShaderSampler(FVertexShaderRHIParamRef VertexShade
 	InternalSetSamplerStates(FOpenGL::GetFirstVertexTextureUnit() + SamplerIndex, NewState);
 }
 
-void FOpenGLDynamicRHI::RHISetShaderSampler(FHullShaderRHIParamRef HullShaderRHI,uint32 SamplerIndex, FRHISamplerState* NewStateRHI)
+void FOpenGLDynamicRHI::RHISetShaderSampler(FRHIHullShader* HullShaderRHI,uint32 SamplerIndex, FRHISamplerState* NewStateRHI)
 {
 	VALIDATE_BOUND_SHADER(HullShaderRHI);
 
@@ -1083,7 +1083,7 @@ void FOpenGLDynamicRHI::RHISetShaderSampler(FHullShaderRHIParamRef HullShaderRHI
 	FOpenGLSamplerState* NewState = ResourceCast(NewStateRHI);
 	InternalSetSamplerStates(FOpenGL::GetFirstHullTextureUnit() + SamplerIndex, NewState);
 }
-void FOpenGLDynamicRHI::RHISetShaderSampler(FDomainShaderRHIParamRef DomainShaderRHI,uint32 SamplerIndex, FRHISamplerState* NewStateRHI)
+void FOpenGLDynamicRHI::RHISetShaderSampler(FRHIDomainShader* DomainShaderRHI,uint32 SamplerIndex, FRHISamplerState* NewStateRHI)
 {
 	VALIDATE_BOUND_SHADER(DomainShaderRHI);
 
@@ -1093,7 +1093,7 @@ void FOpenGLDynamicRHI::RHISetShaderSampler(FDomainShaderRHIParamRef DomainShade
 	InternalSetSamplerStates(FOpenGL::GetFirstDomainTextureUnit() + SamplerIndex, NewState);
 }
 
-void FOpenGLDynamicRHI::RHISetShaderSampler(FGeometryShaderRHIParamRef GeometryShaderRHI,uint32 SamplerIndex, FRHISamplerState* NewStateRHI)
+void FOpenGLDynamicRHI::RHISetShaderSampler(FRHIGeometryShader* GeometryShaderRHI,uint32 SamplerIndex, FRHISamplerState* NewStateRHI)
 {
 	VALIDATE_BOUND_SHADER(GeometryShaderRHI);
 
@@ -1102,7 +1102,7 @@ void FOpenGLDynamicRHI::RHISetShaderSampler(FGeometryShaderRHIParamRef GeometryS
 	InternalSetSamplerStates(FOpenGL::GetFirstGeometryTextureUnit() + SamplerIndex, NewState);
 }
 
-void FOpenGLDynamicRHI::RHISetShaderTexture(FComputeShaderRHIParamRef ComputeShaderRHI,uint32 TextureIndex,FTextureRHIParamRef NewTextureRHI)
+void FOpenGLDynamicRHI::RHISetShaderTexture(FRHIComputeShader* ComputeShaderRHI,uint32 TextureIndex,FTextureRHIParamRef NewTextureRHI)
 {
 	check(GMaxRHIFeatureLevel >= ERHIFeatureLevel::SM5);
 	VERIFY_GL_SCOPE();
@@ -1118,7 +1118,7 @@ void FOpenGLDynamicRHI::RHISetShaderTexture(FComputeShaderRHIParamRef ComputeSha
 	}
 }
 
-void FOpenGLDynamicRHI::RHISetShaderSampler(FPixelShaderRHIParamRef PixelShaderRHI,uint32 SamplerIndex, FRHISamplerState* NewStateRHI)
+void FOpenGLDynamicRHI::RHISetShaderSampler(FRHIPixelShader* PixelShaderRHI,uint32 SamplerIndex, FRHISamplerState* NewStateRHI)
 {
 	VALIDATE_BOUND_SHADER(PixelShaderRHI);
 
@@ -1127,7 +1127,7 @@ void FOpenGLDynamicRHI::RHISetShaderSampler(FPixelShaderRHIParamRef PixelShaderR
 	InternalSetSamplerStates(FOpenGL::GetFirstPixelTextureUnit() + SamplerIndex, NewState);
 }
 
-void FOpenGLDynamicRHI::RHISetShaderUniformBuffer(FVertexShaderRHIParamRef VertexShaderRHI,uint32 BufferIndex,FUniformBufferRHIParamRef BufferRHI)
+void FOpenGLDynamicRHI::RHISetShaderUniformBuffer(FRHIVertexShader* VertexShaderRHI,uint32 BufferIndex,FUniformBufferRHIParamRef BufferRHI)
 {
 	VALIDATE_BOUND_SHADER(VertexShaderRHI);
 	VERIFY_GL_SCOPE();
@@ -1136,7 +1136,7 @@ void FOpenGLDynamicRHI::RHISetShaderUniformBuffer(FVertexShaderRHIParamRef Verte
 	PendingState.bAnyDirtyGraphicsUniformBuffers = true;
 }
 
-void FOpenGLDynamicRHI::RHISetShaderUniformBuffer(FHullShaderRHIParamRef HullShaderRHI,uint32 BufferIndex,FUniformBufferRHIParamRef BufferRHI)
+void FOpenGLDynamicRHI::RHISetShaderUniformBuffer(FRHIHullShader* HullShaderRHI,uint32 BufferIndex,FUniformBufferRHIParamRef BufferRHI)
 {
 	VALIDATE_BOUND_SHADER(HullShaderRHI);
 	VERIFY_GL_SCOPE();
@@ -1146,7 +1146,7 @@ void FOpenGLDynamicRHI::RHISetShaderUniformBuffer(FHullShaderRHIParamRef HullSha
 	PendingState.bAnyDirtyGraphicsUniformBuffers = true;
 }
 
-void FOpenGLDynamicRHI::RHISetShaderUniformBuffer(FDomainShaderRHIParamRef DomainShaderRHI,uint32 BufferIndex,FUniformBufferRHIParamRef BufferRHI)
+void FOpenGLDynamicRHI::RHISetShaderUniformBuffer(FRHIDomainShader* DomainShaderRHI,uint32 BufferIndex,FUniformBufferRHIParamRef BufferRHI)
 {
 	VALIDATE_BOUND_SHADER(DomainShaderRHI);
 	VERIFY_GL_SCOPE();
@@ -1156,7 +1156,7 @@ void FOpenGLDynamicRHI::RHISetShaderUniformBuffer(FDomainShaderRHIParamRef Domai
 	PendingState.bAnyDirtyGraphicsUniformBuffers = true;
 }
 
-void FOpenGLDynamicRHI::RHISetShaderUniformBuffer(FGeometryShaderRHIParamRef GeometryShaderRHI,uint32 BufferIndex,FUniformBufferRHIParamRef BufferRHI)
+void FOpenGLDynamicRHI::RHISetShaderUniformBuffer(FRHIGeometryShader* GeometryShaderRHI,uint32 BufferIndex,FUniformBufferRHIParamRef BufferRHI)
 {
 	VALIDATE_BOUND_SHADER(GeometryShaderRHI);
 	VERIFY_GL_SCOPE();
@@ -1165,7 +1165,7 @@ void FOpenGLDynamicRHI::RHISetShaderUniformBuffer(FGeometryShaderRHIParamRef Geo
 	PendingState.bAnyDirtyGraphicsUniformBuffers = true;
 }
 
-void FOpenGLDynamicRHI::RHISetShaderSampler(FComputeShaderRHIParamRef ComputeShaderRHI,uint32 SamplerIndex, FRHISamplerState* NewStateRHI)
+void FOpenGLDynamicRHI::RHISetShaderSampler(FRHIComputeShader* ComputeShaderRHI,uint32 SamplerIndex, FRHISamplerState* NewStateRHI)
 {
 	check(FOpenGL::SupportsComputeShaders());
 	VERIFY_GL_SCOPE();
@@ -1180,7 +1180,7 @@ void FOpenGLDynamicRHI::RHISetShaderSampler(FComputeShaderRHIParamRef ComputeSha
 	}
 }
 
-void FOpenGLDynamicRHI::RHISetShaderUniformBuffer(FPixelShaderRHIParamRef PixelShaderRHI,uint32 BufferIndex,FUniformBufferRHIParamRef BufferRHI)
+void FOpenGLDynamicRHI::RHISetShaderUniformBuffer(FRHIPixelShader* PixelShaderRHI,uint32 BufferIndex,FUniformBufferRHIParamRef BufferRHI)
 {
 	VERIFY_GL_SCOPE();
 	VALIDATE_BOUND_SHADER(PixelShaderRHI);
@@ -1189,7 +1189,7 @@ void FOpenGLDynamicRHI::RHISetShaderUniformBuffer(FPixelShaderRHIParamRef PixelS
 	PendingState.bAnyDirtyGraphicsUniformBuffers = true;
 }
 
-void FOpenGLDynamicRHI::RHISetShaderUniformBuffer(FComputeShaderRHIParamRef ComputeShaderRHI,uint32 BufferIndex,FUniformBufferRHIParamRef BufferRHI)
+void FOpenGLDynamicRHI::RHISetShaderUniformBuffer(FRHIComputeShader* ComputeShaderRHI,uint32 BufferIndex,FUniformBufferRHIParamRef BufferRHI)
 {
 	VERIFY_GL_SCOPE();
 	check(GMaxRHIFeatureLevel >= ERHIFeatureLevel::SM5);
@@ -1197,7 +1197,7 @@ void FOpenGLDynamicRHI::RHISetShaderUniformBuffer(FComputeShaderRHIParamRef Comp
 	PendingState.DirtyUniformBuffers[SF_Compute] |= 1 << BufferIndex;
 }
 
-void FOpenGLDynamicRHI::RHISetShaderParameter(FVertexShaderRHIParamRef VertexShaderRHI,uint32 BufferIndex,uint32 BaseIndex,uint32 NumBytes,const void* NewValue)
+void FOpenGLDynamicRHI::RHISetShaderParameter(FRHIVertexShader* VertexShaderRHI,uint32 BufferIndex,uint32 BaseIndex,uint32 NumBytes,const void* NewValue)
 {
 	VALIDATE_BOUND_SHADER(VertexShaderRHI);
 
@@ -1206,7 +1206,7 @@ void FOpenGLDynamicRHI::RHISetShaderParameter(FVertexShaderRHIParamRef VertexSha
 	PendingState.LinkedProgramAndDirtyFlag = nullptr;
 }
 
-void FOpenGLDynamicRHI::RHISetShaderParameter(FPixelShaderRHIParamRef PixelShaderRHI,uint32 BufferIndex,uint32 BaseIndex,uint32 NumBytes,const void* NewValue)
+void FOpenGLDynamicRHI::RHISetShaderParameter(FRHIPixelShader* PixelShaderRHI,uint32 BufferIndex,uint32 BaseIndex,uint32 NumBytes,const void* NewValue)
 {
 	VALIDATE_BOUND_SHADER(PixelShaderRHI);
 
@@ -1215,7 +1215,7 @@ void FOpenGLDynamicRHI::RHISetShaderParameter(FPixelShaderRHIParamRef PixelShade
 	PendingState.LinkedProgramAndDirtyFlag = nullptr;
 }
 
-void FOpenGLDynamicRHI::RHISetShaderParameter(FHullShaderRHIParamRef HullShaderRHI,uint32 BufferIndex,uint32 BaseIndex,uint32 NumBytes,const void* NewValue)
+void FOpenGLDynamicRHI::RHISetShaderParameter(FRHIHullShader* HullShaderRHI,uint32 BufferIndex,uint32 BaseIndex,uint32 NumBytes,const void* NewValue)
 {
 	VALIDATE_BOUND_SHADER(HullShaderRHI);
 
@@ -1225,7 +1225,7 @@ void FOpenGLDynamicRHI::RHISetShaderParameter(FHullShaderRHIParamRef HullShaderR
 	PendingState.LinkedProgramAndDirtyFlag = nullptr;
 }
 
-void FOpenGLDynamicRHI::RHISetShaderParameter(FDomainShaderRHIParamRef DomainShaderRHI,uint32 BufferIndex,uint32 BaseIndex,uint32 NumBytes,const void* NewValue)
+void FOpenGLDynamicRHI::RHISetShaderParameter(FRHIDomainShader* DomainShaderRHI,uint32 BufferIndex,uint32 BaseIndex,uint32 NumBytes,const void* NewValue)
 {
 	VALIDATE_BOUND_SHADER(DomainShaderRHI);
 
@@ -1235,7 +1235,7 @@ void FOpenGLDynamicRHI::RHISetShaderParameter(FDomainShaderRHIParamRef DomainSha
 	PendingState.LinkedProgramAndDirtyFlag = nullptr;
 }
 
-void FOpenGLDynamicRHI::RHISetShaderParameter(FGeometryShaderRHIParamRef GeometryShaderRHI,uint32 BufferIndex,uint32 BaseIndex,uint32 NumBytes,const void* NewValue)
+void FOpenGLDynamicRHI::RHISetShaderParameter(FRHIGeometryShader* GeometryShaderRHI,uint32 BufferIndex,uint32 BaseIndex,uint32 NumBytes,const void* NewValue)
 {
 	VALIDATE_BOUND_SHADER(GeometryShaderRHI);
 
@@ -1244,7 +1244,7 @@ void FOpenGLDynamicRHI::RHISetShaderParameter(FGeometryShaderRHIParamRef Geometr
 	PendingState.LinkedProgramAndDirtyFlag = nullptr;
 }
 
-void FOpenGLDynamicRHI::RHISetShaderParameter(FComputeShaderRHIParamRef ComputeShaderRHI,uint32 BufferIndex,uint32 BaseIndex,uint32 NumBytes,const void* NewValue)
+void FOpenGLDynamicRHI::RHISetShaderParameter(FRHIComputeShader* ComputeShaderRHI,uint32 BufferIndex,uint32 BaseIndex,uint32 NumBytes,const void* NewValue)
 { 
 	VERIFY_GL_SCOPE();
 	check(FOpenGL::SupportsComputeShaders());
@@ -3174,7 +3174,7 @@ uint32 FOpenGLDynamicRHI::RHIGetGPUFrameCycles()
 	return GGPUFrameTime;
 }
 
-void FOpenGLDynamicRHI::RHISetComputeShader(FComputeShaderRHIParamRef ComputeShaderRHI)
+void FOpenGLDynamicRHI::RHISetComputeShader(FRHIComputeShader* ComputeShaderRHI)
 {
 	if (OpenGLConsoleVariables::bSkipCompute)
 	{
@@ -3202,7 +3202,7 @@ void FOpenGLDynamicRHI::RHIDispatchComputeShader(uint32 ThreadGroupCountX, uint3
 	{
 		VERIFY_GL_SCOPE();
 		
-		FComputeShaderRHIParamRef ComputeShaderRHI = PendingState.CurrentComputeShader;
+		FRHIComputeShader* ComputeShaderRHI = PendingState.CurrentComputeShader;
 		check(ComputeShaderRHI);
 
 		FOpenGLComputeShader* ComputeShader = ResourceCast(ComputeShaderRHI);
@@ -3233,7 +3233,7 @@ void FOpenGLDynamicRHI::RHIDispatchIndirectComputeShader(FVertexBufferRHIParamRe
 	{
 		VERIFY_GL_SCOPE();
 		
-		FComputeShaderRHIParamRef ComputeShaderRHI = PendingState.CurrentComputeShader;
+		FRHIComputeShader* ComputeShaderRHI = PendingState.CurrentComputeShader;
 		check(ComputeShaderRHI);
 
 		FOpenGLComputeShader* ComputeShader = ResourceCast(ComputeShaderRHI);
diff --git a/Engine/Source/Runtime/OpenGLDrv/Private/OpenGLShaders.cpp b/Engine/Source/Runtime/OpenGLDrv/Private/OpenGLShaders.cpp
--- a/Engine/Source/Runtime/OpenGLDrv/Private/OpenGLShaders.cpp
+++ b/Engine/Source/Runtime/OpenGLDrv/Private/OpenGLShaders.cpp
@@ -3580,11 +3580,11 @@ static void BindShaderStage(FOpenGLLinkedProgramConfiguration& Config, CrossComp
 
 FBoundShaderStateRHIRef FOpenGLDynamicRHI::RHICreateBoundShaderState_OnThisThread(
 	FRHIVertexDeclaration* VertexDeclarationRHI,
-	FVertexShaderRHIParamRef VertexShaderRHI, 
-	FHullShaderRHIParamRef HullShaderRHI,
-	FDomainShaderRHIParamRef DomainShaderRHI, 
-	FPixelShaderRHIParamRef PixelShaderRHI, 
-	FGeometryShaderRHIParamRef GeometryShaderRHI,
+	FRHIVertexShader* VertexShaderRHI,
+	FRHIHullShader* HullShaderRHI,
+	FRHIDomainShader* DomainShaderRHI,
+	FRHIPixelShader* PixelShaderRHI,
+	FRHIGeometryShader* GeometryShaderRHI,
 	bool bFromPSOFileCache
 	)
 {
@@ -4026,11 +4026,11 @@ void FOpenGLDynamicRHI::BindPendingShaderState( FOpenGLContextState& ContextStat
 FOpenGLBoundShaderState::FOpenGLBoundShaderState(
 	FOpenGLLinkedProgram* InLinkedProgram,
 	FRHIVertexDeclaration* InVertexDeclarationRHI,
-	FVertexShaderRHIParamRef InVertexShaderRHI,
-	FPixelShaderRHIParamRef InPixelShaderRHI,
-	FGeometryShaderRHIParamRef InGeometryShaderRHI,
-	FHullShaderRHIParamRef InHullShaderRHI,
-	FDomainShaderRHIParamRef InDomainShaderRHI
+	FRHIVertexShader* InVertexShaderRHI,
+	FRHIPixelShader* InPixelShaderRHI,
+	FRHIGeometryShader* InGeometryShaderRHI,
+	FRHIHullShader* InHullShaderRHI,
+	FRHIDomainShader* InDomainShaderRHI
 	)
 	:	CacheLink(InVertexDeclarationRHI, InVertexShaderRHI, InPixelShaderRHI,
 		InHullShaderRHI, InDomainShaderRHI,	InGeometryShaderRHI, this)
diff --git a/Engine/Source/Runtime/OpenGLDrv/Public/OpenGLDrv.h b/Engine/Source/Runtime/OpenGLDrv/Public/OpenGLDrv.h
--- a/Engine/Source/Runtime/OpenGLDrv/Public/OpenGLDrv.h
+++ b/Engine/Source/Runtime/OpenGLDrv/Public/OpenGLDrv.h
@@ -445,7 +445,7 @@ public:
 	virtual void RHIBeginRenderPass(const FRHIRenderPassInfo& InInfo, const TCHAR* InName) final override;
 	virtual void RHIEndRenderPass() final override;
 
-	virtual void RHISetComputeShader(FComputeShaderRHIParamRef ComputeShader) final override;
+	virtual void RHISetComputeShader(FRHIComputeShader* ComputeShader) final override;
 	virtual void RHIDispatchComputeShader(uint32 ThreadGroupCountX, uint32 ThreadGroupCountY, uint32 ThreadGroupCountZ) final override;
 	virtual void RHIDispatchIndirectComputeShader(FVertexBufferRHIParamRef ArgumentBuffer, uint32 ArgumentOffset) final override;
 	virtual void RHIAutomaticCacheFlushAfterComputeShader(bool bEnable) final override;
@@ -467,38 +467,38 @@ public:
 	virtual void RHISetViewport(uint32 MinX, uint32 MinY, float MinZ, uint32 MaxX, uint32 MaxY, float MaxZ) final override;
 	virtual void RHISetScissorRect(bool bEnable, uint32 MinX, uint32 MinY, uint32 MaxX, uint32 MaxY) final override;
 	virtual void RHISetBoundShaderState(FRHIBoundShaderState* BoundShaderState) final override;
-	virtual void RHISetShaderTexture(FVertexShaderRHIParamRef VertexShader, uint32 TextureIndex, FTextureRHIParamRef NewTexture) final override;
-	virtual void RHISetShaderTexture(FHullShaderRHIParamRef HullShader, uint32 TextureIndex, FTextureRHIParamRef NewTexture) final override;
-	virtual void RHISetShaderTexture(FDomainShaderRHIParamRef DomainShader, uint32 TextureIndex, FTextureRHIParamRef NewTexture) final override;
-	virtual void RHISetShaderTexture(FGeometryShaderRHIParamRef GeometryShader, uint32 TextureIndex, FTextureRHIParamRef NewTexture) final override;
-	virtual void RHISetShaderTexture(FPixelShaderRHIParamRef PixelShader, uint32 TextureIndex, FTextureRHIParamRef NewTexture) final override;
-	virtual void RHISetShaderTexture(FComputeShaderRHIParamRef PixelShader, uint32 TextureIndex, FTextureRHIParamRef NewTexture) final override;
-	virtual void RHISetShaderSampler(FComputeShaderRHIParamRef ComputeShader, uint32 SamplerIndex, FRHISamplerState* NewState) final override;
-	virtual void RHISetShaderSampler(FVertexShaderRHIParamRef VertexShader, uint32 SamplerIndex, FRHISamplerState* NewState) final override;
-	virtual void RHISetShaderSampler(FGeometryShaderRHIParamRef GeometryShader, uint32 SamplerIndex, FRHISamplerState* NewState) final override;
-	virtual void RHISetShaderSampler(FDomainShaderRHIParamRef DomainShader, uint32 SamplerIndex, FRHISamplerState* NewState) final override;
-	virtual void RHISetShaderSampler(FHullShaderRHIParamRef HullShader, uint32 SamplerIndex, FRHISamplerState* NewState) final override;
-	virtual void RHISetShaderSampler(FPixelShaderRHIParamRef PixelShader, uint32 SamplerIndex, FRHISamplerState* NewState) final override;
-	virtual void RHISetUAVParameter(FComputeShaderRHIParamRef ComputeShader, uint32 UAVIndex, FRHIUnorderedAccessView* UAV) final override;
-	virtual void RHISetUAVParameter(FComputeShaderRHIParamRef ComputeShader, uint32 UAVIndex, FRHIUnorderedAccessView* UAV, uint32 InitialCount) final override;
-	virtual void RHISetShaderResourceViewParameter(FPixelShaderRHIParamRef PixelShader, uint32 SamplerIndex, FRHIShaderResourceView* SRV) final override;
-	virtual void RHISetShaderResourceViewParameter(FVertexShaderRHIParamRef VertexShader, uint32 SamplerIndex, FRHIShaderResourceView* SRV) final override;
-	virtual void RHISetShaderResourceViewParameter(FComputeShaderRHIParamRef ComputeShader, uint32 SamplerIndex, FRHIShaderResourceView* SRV) final override;
-	virtual void RHISetShaderResourceViewParameter(FHullShaderRHIParamRef HullShader, uint32 SamplerIndex, FRHIShaderResourceView* SRV) final override;
-	virtual void RHISetShaderResourceViewParameter(FDomainShaderRHIParamRef DomainShader, uint32 SamplerIndex, FRHIShaderResourceView* SRV) final override;
-	virtual void RHISetShaderResourceViewParameter(FGeometryShaderRHIParamRef GeometryShader, uint32 SamplerIndex, FRHIShaderResourceView* SRV) final override;
-	virtual void RHISetShaderUniformBuffer(FVertexShaderRHIParamRef VertexShader, uint32 BufferIndex, FUniformBufferRHIParamRef Buffer) final override;
-	virtual void RHISetShaderUniformBuffer(FHullShaderRHIParamRef HullShader, uint32 BufferIndex, FUniformBufferRHIParamRef Buffer) final override;
-	virtual void RHISetShaderUniformBuffer(FDomainShaderRHIParamRef DomainShader, uint32 BufferIndex, FUniformBufferRHIParamRef Buffer) final override;
-	virtual void RHISetShaderUniformBuffer(FGeometryShaderRHIParamRef GeometryShader, uint32 BufferIndex, FUniformBufferRHIParamRef Buffer) final override;
-	virtual void RHISetShaderUniformBuffer(FPixelShaderRHIParamRef PixelShader, uint32 BufferIndex, FUniformBufferRHIParamRef Buffer) final override;
-	virtual void RHISetShaderUniformBuffer(FComputeShaderRHIParamRef ComputeShader, uint32 BufferIndex, FUniformBufferRHIParamRef Buffer) final override;
-	virtual void RHISetShaderParameter(FVertexShaderRHIParamRef VertexShader, uint32 BufferIndex, uint32 BaseIndex, uint32 NumBytes, const void* NewValue) final override;
-	virtual void RHISetShaderParameter(FPixelShaderRHIParamRef PixelShader, uint32 BufferIndex, uint32 BaseIndex, uint32 NumBytes, const void* NewValue) final override;
-	virtual void RHISetShaderParameter(FHullShaderRHIParamRef HullShader, uint32 BufferIndex, uint32 BaseIndex, uint32 NumBytes, const void* NewValue) final override;
-	virtual void RHISetShaderParameter(FDomainShaderRHIParamRef DomainShader, uint32 BufferIndex, uint32 BaseIndex, uint32 NumBytes, const void* NewValue) final override;
-	virtual void RHISetShaderParameter(FGeometryShaderRHIParamRef GeometryShader, uint32 BufferIndex, uint32 BaseIndex, uint32 NumBytes, const void* NewValue) final override;
-	virtual void RHISetShaderParameter(FComputeShaderRHIParamRef ComputeShader, uint32 BufferIndex, uint32 BaseIndex, uint32 NumBytes, const void* NewValue) final override;
+	virtual void RHISetShaderTexture(FRHIVertexShader* VertexShader, uint32 TextureIndex, FTextureRHIParamRef NewTexture) final override;
+	virtual void RHISetShaderTexture(FRHIHullShader* HullShader, uint32 TextureIndex, FTextureRHIParamRef NewTexture) final override;
+	virtual void RHISetShaderTexture(FRHIDomainShader* DomainShader, uint32 TextureIndex, FTextureRHIParamRef NewTexture) final override;
+	virtual void RHISetShaderTexture(FRHIGeometryShader* GeometryShader, uint32 TextureIndex, FTextureRHIParamRef NewTexture) final override;
+	virtual void RHISetShaderTexture(FRHIPixelShader* PixelShader, uint32 TextureIndex, FTextureRHIParamRef NewTexture) final override;
+	virtual void RHISetShaderTexture(FRHIComputeShader* PixelShader, uint32 TextureIndex, FTextureRHIParamRef NewTexture) final override;
+	virtual void RHISetShaderSampler(FRHIComputeShader* ComputeShader, uint32 SamplerIndex, FRHISamplerState* NewState) final override;
+	virtual void RHISetShaderSampler(FRHIVertexShader* VertexShader, uint32 SamplerIndex, FRHISamplerState* NewState) final override;
+	virtual void RHISetShaderSampler(FRHIGeometryShader* GeometryShader, uint32 SamplerIndex, FRHISamplerState* NewState) final override;
+	virtual void RHISetShaderSampler(FRHIDomainShader* DomainShader, uint32 SamplerIndex, FRHISamplerState* NewState) final override;
+	virtual void RHISetShaderSampler(FRHIHullShader* HullShader, uint32 SamplerIndex, FRHISamplerState* NewState) final override;
+	virtual void RHISetShaderSampler(FRHIPixelShader* PixelShader, uint32 SamplerIndex, FRHISamplerState* NewState) final override;
+	virtual void RHISetUAVParameter(FRHIComputeShader* ComputeShader, uint32 UAVIndex, FRHIUnorderedAccessView* UAV) final override;
+	virtual void RHISetUAVParameter(FRHIComputeShader* ComputeShader, uint32 UAVIndex, FRHIUnorderedAccessView* UAV, uint32 InitialCount) final override;
+	virtual void RHISetShaderResourceViewParameter(FRHIPixelShader* PixelShader, uint32 SamplerIndex, FRHIShaderResourceView* SRV) final override;
+	virtual void RHISetShaderResourceViewParameter(FRHIVertexShader* VertexShader, uint32 SamplerIndex, FRHIShaderResourceView* SRV) final override;
+	virtual void RHISetShaderResourceViewParameter(FRHIComputeShader* ComputeShader, uint32 SamplerIndex, FRHIShaderResourceView* SRV) final override;
+	virtual void RHISetShaderResourceViewParameter(FRHIHullShader* HullShader, uint32 SamplerIndex, FRHIShaderResourceView* SRV) final override;
+	virtual void RHISetShaderResourceViewParameter(FRHIDomainShader* DomainShader, uint32 SamplerIndex, FRHIShaderResourceView* SRV) final override;
+	virtual void RHISetShaderResourceViewParameter(FRHIGeometryShader* GeometryShader, uint32 SamplerIndex, FRHIShaderResourceView* SRV) final override;
+	virtual void RHISetShaderUniformBuffer(FRHIVertexShader* VertexShader, uint32 BufferIndex, FUniformBufferRHIParamRef Buffer) final override;
+	virtual void RHISetShaderUniformBuffer(FRHIHullShader* HullShader, uint32 BufferIndex, FUniformBufferRHIParamRef Buffer) final override;
+	virtual void RHISetShaderUniformBuffer(FRHIDomainShader* DomainShader, uint32 BufferIndex, FUniformBufferRHIParamRef Buffer) final override;
+	virtual void RHISetShaderUniformBuffer(FRHIGeometryShader* GeometryShader, uint32 BufferIndex, FUniformBufferRHIParamRef Buffer) final override;
+	virtual void RHISetShaderUniformBuffer(FRHIPixelShader* PixelShader, uint32 BufferIndex, FUniformBufferRHIParamRef Buffer) final override;
+	virtual void RHISetShaderUniformBuffer(FRHIComputeShader* ComputeShader, uint32 BufferIndex, FUniformBufferRHIParamRef Buffer) final override;
+	virtual void RHISetShaderParameter(FRHIVertexShader* VertexShader, uint32 BufferIndex, uint32 BaseIndex, uint32 NumBytes, const void* NewValue) final override;
+	virtual void RHISetShaderParameter(FRHIPixelShader* PixelShader, uint32 BufferIndex, uint32 BaseIndex, uint32 NumBytes, const void* NewValue) final override;
+	virtual void RHISetShaderParameter(FRHIHullShader* HullShader, uint32 BufferIndex, uint32 BaseIndex, uint32 NumBytes, const void* NewValue) final override;
+	virtual void RHISetShaderParameter(FRHIDomainShader* DomainShader, uint32 BufferIndex, uint32 BaseIndex, uint32 NumBytes, const void* NewValue) final override;
+	virtual void RHISetShaderParameter(FRHIGeometryShader* GeometryShader, uint32 BufferIndex, uint32 BaseIndex, uint32 NumBytes, const void* NewValue) final override;
+	virtual void RHISetShaderParameter(FRHIComputeShader* ComputeShader, uint32 BufferIndex, uint32 BaseIndex, uint32 NumBytes, const void* NewValue) final override;
 	virtual void RHISetDepthStencilState(FRHIDepthStencilState* NewState, uint32 StencilRef) final override;
 	virtual void RHISetStencilRef(uint32 StencilRef) final override;
 	virtual void RHISetBlendState(FRHIBlendState* NewState, const FLinearColor& BlendFactor) final override;
@@ -999,11 +999,11 @@ public:
 
 	FBoundShaderStateRHIRef RHICreateBoundShaderState_internal(
 		FRHIVertexDeclaration* VertexDeclarationRHI,
-		FVertexShaderRHIParamRef VertexShaderRHI,
-		FHullShaderRHIParamRef HullShaderRHI,
-		FDomainShaderRHIParamRef DomainShaderRHI,
-		FPixelShaderRHIParamRef PixelShaderRHI,
-		FGeometryShaderRHIParamRef GeometryShaderRHI,
+		FRHIVertexShader* VertexShaderRHI,
+		FRHIHullShader* HullShaderRHI,
+		FRHIDomainShader* DomainShaderRHI,
+		FRHIPixelShader* PixelShaderRHI,
+		FRHIGeometryShader* GeometryShaderRHI,
 		bool FromPSOFileCache
 	)
 	{
@@ -1022,11 +1022,11 @@ public:
 
 	virtual FBoundShaderStateRHIRef RHICreateBoundShaderState(
 		FRHIVertexDeclaration* VertexDeclarationRHI,
-		FVertexShaderRHIParamRef VertexShaderRHI,
-		FHullShaderRHIParamRef HullShaderRHI,
-		FDomainShaderRHIParamRef DomainShaderRHI,
-		FPixelShaderRHIParamRef PixelShaderRHI,
-		FGeometryShaderRHIParamRef GeometryShaderRHI
+		FRHIVertexShader* VertexShaderRHI,
+		FRHIHullShader* HullShaderRHI,
+		FRHIDomainShader* DomainShaderRHI,
+		FRHIPixelShader* PixelShaderRHI,
+		FRHIGeometryShader* GeometryShaderRHI
 	) final override
 	{
 		return RHICreateBoundShaderState_internal(
@@ -1040,7 +1040,7 @@ public:
 	}
 
 
-	FBoundShaderStateRHIRef RHICreateBoundShaderState_OnThisThread(FRHIVertexDeclaration* VertexDeclaration, FVertexShaderRHIParamRef VertexShader, FHullShaderRHIParamRef HullShader, FDomainShaderRHIParamRef DomainShader, FPixelShaderRHIParamRef PixelShader, FGeometryShaderRHIParamRef GeometryShader, bool FromPSOFileCache);
+	FBoundShaderStateRHIRef RHICreateBoundShaderState_OnThisThread(FRHIVertexDeclaration* VertexDeclaration, FRHIVertexShader* VertexShader, FRHIHullShader* HullShader, FRHIDomainShader* DomainShader, FRHIPixelShader* PixelShader, FRHIGeometryShader* GeometryShader, bool FromPSOFileCache);
 	void RHIPerFrameRHIFlushComplete();
 
 	FOpenGLGPUProfiler& GetGPUProfilingData() {
diff --git a/Engine/Source/Runtime/OpenGLDrv/Public/OpenGLResources.h b/Engine/Source/Runtime/OpenGLDrv/Public/OpenGLResources.h
--- a/Engine/Source/Runtime/OpenGLDrv/Public/OpenGLResources.h
+++ b/Engine/Source/Runtime/OpenGLDrv/Public/OpenGLResources.h
@@ -1169,11 +1169,11 @@ public:
 	FOpenGLBoundShaderState(
 		FOpenGLLinkedProgram* InLinkedProgram,
 		FRHIVertexDeclaration* InVertexDeclarationRHI,
-		FVertexShaderRHIParamRef InVertexShaderRHI,
-		FPixelShaderRHIParamRef InPixelShaderRHI,
-		FGeometryShaderRHIParamRef InGeometryShaderRHI,
-		FHullShaderRHIParamRef InHullShaderRHI,
-		FDomainShaderRHIParamRef InDomainShaderRHI
+		FRHIVertexShader* InVertexShaderRHI,
+		FRHIPixelShader* InPixelShaderRHI,
+		FRHIGeometryShader* InGeometryShaderRHI,
+		FRHIHullShader* InHullShaderRHI,
+		FRHIDomainShader* InDomainShaderRHI
 		);
 
 	const TBitArray<>& GetTextureNeeds(int32& OutMaxTextureStageUsed);
diff --git a/Engine/Source/Runtime/RHI/Private/BoundShaderStateCache.cpp b/Engine/Source/Runtime/RHI/Private/BoundShaderStateCache.cpp
--- a/Engine/Source/Runtime/RHI/Private/BoundShaderStateCache.cpp
+++ b/Engine/Source/Runtime/RHI/Private/BoundShaderStateCache.cpp
@@ -31,11 +31,11 @@ static FCriticalSection BoundShaderStateCacheLock;
 
 FCachedBoundShaderStateLink::FCachedBoundShaderStateLink(
 	FRHIVertexDeclaration* VertexDeclaration,
-	FVertexShaderRHIParamRef VertexShader,
-	FPixelShaderRHIParamRef PixelShader,
-	FHullShaderRHIParamRef HullShader,
-	FDomainShaderRHIParamRef DomainShader,
-	FGeometryShaderRHIParamRef GeometryShader,
+	FRHIVertexShader* VertexShader,
+	FRHIPixelShader* PixelShader,
+	FRHIHullShader* HullShader,
+	FRHIDomainShader* DomainShader,
+	FRHIGeometryShader* GeometryShader,
 	FRHIBoundShaderState* InBoundShaderState,
 	bool bAddToSingleThreadedCache
 	):
@@ -51,8 +51,8 @@ FCachedBoundShaderStateLink::FCachedBoundShaderStateLink(
 
 FCachedBoundShaderStateLink::FCachedBoundShaderStateLink(
 	FRHIVertexDeclaration* VertexDeclaration,
-	FVertexShaderRHIParamRef VertexShader,
-	FPixelShaderRHIParamRef PixelShader,
+	FRHIVertexShader* VertexShader,
+	FRHIPixelShader* PixelShader,
 	FRHIBoundShaderState* InBoundShaderState,
 	bool bAddToSingleThreadedCache
 	):
@@ -77,11 +77,11 @@ FCachedBoundShaderStateLink::~FCachedBoundShaderStateLink()
 
 FCachedBoundShaderStateLink* GetCachedBoundShaderState(
 	FRHIVertexDeclaration* VertexDeclaration,
-	FVertexShaderRHIParamRef VertexShader,
-	FPixelShaderRHIParamRef PixelShader,
-	FHullShaderRHIParamRef HullShader,
-	FDomainShaderRHIParamRef DomainShader,
-	FGeometryShaderRHIParamRef GeometryShader
+	FRHIVertexShader* VertexShader,
+	FRHIPixelShader* PixelShader,
+	FRHIHullShader* HullShader,
+	FRHIDomainShader* DomainShader,
+	FRHIGeometryShader* GeometryShader
 	)
 {
 	// Find the existing bound shader state in the cache.
@@ -105,11 +105,11 @@ void FCachedBoundShaderStateLink_Threadsafe::RemoveFromCache()
 
 FBoundShaderStateRHIRef GetCachedBoundShaderState_Threadsafe(
 	FRHIVertexDeclaration* VertexDeclaration,
-	FVertexShaderRHIParamRef VertexShader,
-	FPixelShaderRHIParamRef PixelShader,
-	FHullShaderRHIParamRef HullShader,
-	FDomainShaderRHIParamRef DomainShader,
-	FGeometryShaderRHIParamRef GeometryShader
+	FRHIVertexShader* VertexShader,
+	FRHIPixelShader* PixelShader,
+	FRHIHullShader* HullShader,
+	FRHIDomainShader* DomainShader,
+	FRHIGeometryShader* GeometryShader
 	)
 {
 	FScopeLock Lock(&BoundShaderStateCacheLock);
