Added support for MIP-level selection in RHIReadSurfaceData function for Vulkan, D3D12, and OpenGL.

#rnx

diff --git a/Engine/Source/Runtime/D3D12RHI/Private/D3D12RenderTarget.cpp b/Engine/Source/Runtime/D3D12RHI/Private/D3D12RenderTarget.cpp
--- a/Engine/Source/Runtime/D3D12RHI/Private/D3D12RenderTarget.cpp
+++ b/Engine/Source/Runtime/D3D12RHI/Private/D3D12RenderTarget.cpp
@@ -719,7 +719,11 @@ TRefCountPtr<FD3D12Resource> FD3D12DynamicRHI::GetStagingTexture(FTextureRHIPara
 	if (InTexture2D->IsCubemap())
 	{
 		uint32 D3DFace = GetD3D12CubeFace(InFlags.GetCubeFace());
-		Subresource = CalcSubresource(0, D3DFace, 1);
+		Subresource = CalcSubresource(InFlags.GetMip(), D3DFace, TextureRHI->GetNumMips());
+	}
+	else
+	{
+		Subresource = CalcSubresource(InFlags.GetMip(), 0, TextureRHI->GetNumMips());
 	}
 
 	D3D12_BOX* RectPtr = nullptr; // API prefers NULL for entire texture.
@@ -789,10 +793,13 @@ void FD3D12DynamicRHI::ReadSurfaceDataNoMSAARaw(FTextureRHIParamRef TextureRHI,
 	const uint32 XBytesAligned = Align((uint32)readBackHeapDesc.Footprint.Width * BytesPerPixel, FD3D12_TEXTURE_DATA_PITCH_ALIGNMENT);
 	uint32 SrcStart = StagingRect.Min.X * BytesPerPixel + StagingRect.Min.Y * XBytesAligned;
 
+	// Determine the subresource index.
+	uint32 Subresource = CalcSubresource(InFlags.GetMip(), 0, TextureRHI->GetNumMips());
+
 	// Lock the staging resource.
 	void* pData;
 	D3D12_RANGE ReadRange = { SrcStart, SrcStart + XBytesAligned * (SizeY - 1) + BytesPerLine };
-	VERIFYD3D12RESULT(TempTexture2D->GetResource()->Map(0, &ReadRange, &pData));
+	VERIFYD3D12RESULT(TempTexture2D->GetResource()->Map(Subresource, &ReadRange, &pData));
 
 	uint8* DestPtr = OutData.GetData();
 	uint8* SrcPtr = (uint8*)pData + SrcStart;
@@ -1236,7 +1243,11 @@ void FD3D12DynamicRHI::ReadSurfaceDataMSAARaw(FRHICommandList_RecursiveHazardous
 	if (InTexture2D->IsCubemap())
 	{
 		uint32 D3DFace = GetD3D12CubeFace(InFlags.GetCubeFace());
-		Subresource = CalcSubresource(0, D3DFace, 1);
+		Subresource = CalcSubresource(InFlags.GetMip(), D3DFace, TextureRHI->GetNumMips());
+	}
+	else
+	{
+		Subresource = CalcSubresource(InFlags.GetMip(), 0, TextureRHI->GetNumMips());
 	}
 
 	// Setup the descriptions for the copy to the readback heap.
diff --git a/Engine/Source/Runtime/OpenGLDrv/Private/OpenGLRenderTarget.cpp b/Engine/Source/Runtime/OpenGLDrv/Private/OpenGLRenderTarget.cpp
--- a/Engine/Source/Runtime/OpenGLDrv/Private/OpenGLRenderTarget.cpp
+++ b/Engine/Source/Runtime/OpenGLDrv/Private/OpenGLRenderTarget.cpp
@@ -39,8 +39,8 @@ public:
 	FOpenGLFramebufferKey(
 		uint32 InNumRenderTargets,
 		FOpenGLTextureBase** InRenderTargets,
-		uint32* InRenderTargetArrayIndices,
-		uint32* InRenderTargetMipmapLevels,
+		const uint32* InRenderTargetArrayIndices,
+		const uint32* InRenderTargetMipmapLevels,
 		FOpenGLTextureBase* InDepthStencilTarget,
 		EOpenGLCurrentContext InContext
 		)
@@ -102,7 +102,7 @@ static FOpenGLFramebufferCache& GetOpenGLFramebufferCache()
 	return OpenGLFramebufferCache;
 }
 
-GLuint FOpenGLDynamicRHI::GetOpenGLFramebuffer(uint32 NumSimultaneousRenderTargets, FOpenGLTextureBase** RenderTargets, uint32* ArrayIndices, uint32* MipmapLevels, FOpenGLTextureBase* DepthStencilTarget )
+GLuint FOpenGLDynamicRHI::GetOpenGLFramebuffer(uint32 NumSimultaneousRenderTargets, FOpenGLTextureBase** RenderTargets, const uint32* ArrayIndices, const uint32* MipmapLevels, FOpenGLTextureBase* DepthStencilTarget)
 {
 	VERIFY_GL_SCOPE();
 
@@ -482,7 +482,7 @@ void FOpenGLDynamicRHI::RHICopyToResolveTarget(FTextureRHIParamRef SourceTexture
 		}
 
 		GPUProfilingData.RegisterGPUWork();
-		uint32 MipmapLevel = ResolveParams.MipIndex;
+		const uint32 MipmapLevel = ResolveParams.MipIndex;
 
 		const bool bTrueBlit = !SourceTextureRHI->IsMultisampled()
 			&& !DestTextureRHI->IsMultisampled()
@@ -648,7 +648,7 @@ void FOpenGLDynamicRHI::ReadSurfaceDataRaw(FOpenGLContextState& ContextState, FT
 	const GLenum Attachment = bDepthFormat ? (FOpenGL::SupportsPackedDepthStencil() && bDepthStencilFormat ? GL_DEPTH_STENCIL_ATTACHMENT : GL_DEPTH_ATTACHMENT) : GL_COLOR_ATTACHMENT0;
 	const bool bIsColorBuffer = Texture->Attachment == GL_COLOR_ATTACHMENT0;
 
-	uint32 MipmapLevel = 0;
+	const uint32 MipmapLevel = InFlags.GetMip();
 	GLuint SourceFramebuffer = bIsColorBuffer ? GetOpenGLFramebuffer(1, &Texture, NULL, &MipmapLevel, NULL) : GetOpenGLFramebuffer(0, NULL, NULL, NULL, Texture);
 	if (TextureRHI->IsMultisampled())
 	{
@@ -912,7 +912,7 @@ void FOpenGLDynamicRHI::RHIReadSurfaceData(FTextureRHIParamRef TextureRHI, FIntR
 	}
 
 	// Get framebuffer for texture
-	uint32 MipmapLevel = 0;
+	const uint32 MipmapLevel = InFlags.GetMip();
 	GLuint SourceFramebuffer = GetOpenGLFramebuffer(1, &Texture, NULL, &MipmapLevel, NULL);
 
 	uint32 SizeX = Rect.Width();
@@ -984,7 +984,7 @@ void FOpenGLDynamicRHI::RHIReadSurfaceFloatData(FTextureRHIParamRef TextureRHI,F
 	FOpenGLTextureBase* Texture = GetOpenGLTextureFromRHITexture(TextureRHI);
 	check(TextureRHI->GetFormat() == PF_FloatRGBA);
 
-	uint32 MipmapLevel = MipIndex;
+	const uint32 MipmapLevel = MipIndex;
 
 	// Temp FBO is introduced to prevent a ballooning of FBO objects, which can have a detrimental
 	// impact on object management performance in the driver, only for CubeMapArray presently
diff --git a/Engine/Source/Runtime/OpenGLDrv/Public/OpenGLDrv.h b/Engine/Source/Runtime/OpenGLDrv/Public/OpenGLDrv.h
--- a/Engine/Source/Runtime/OpenGLDrv/Public/OpenGLDrv.h
+++ b/Engine/Source/Runtime/OpenGLDrv/Public/OpenGLDrv.h
@@ -1099,7 +1099,7 @@ private:
 	FCriticalSection CustomPresentSection;
 	TRefCountPtr<class FRHICustomPresent> CustomPresent;
 
-	GLuint GetOpenGLFramebuffer(uint32 NumSimultaneousRenderTargets, FOpenGLTextureBase** RenderTargets, uint32* ArrayIndices, uint32* MipmapLevels, FOpenGLTextureBase* DepthStencilTarget);
+	GLuint GetOpenGLFramebuffer(uint32 NumSimultaneousRenderTargets, FOpenGLTextureBase** RenderTargets, const uint32* ArrayIndices, const uint32* MipmapLevels, FOpenGLTextureBase* DepthStencilTarget);
 
 	void InitializeStateResources();
 
diff --git a/Engine/Source/Runtime/VulkanRHI/Private/VulkanRenderTarget.cpp b/Engine/Source/Runtime/VulkanRHI/Private/VulkanRenderTarget.cpp
--- a/Engine/Source/Runtime/VulkanRHI/Private/VulkanRenderTarget.cpp
+++ b/Engine/Source/Runtime/VulkanRHI/Private/VulkanRenderTarget.cpp
@@ -783,12 +783,12 @@ void FVulkanDynamicRHI::RHIReadSurfaceData(FTextureRHIParamRef TextureRHI, FIntR
 	{
 		VkBufferImageCopy CopyRegion;
 		FMemory::Memzero(CopyRegion);
-		//Region.bufferOffset = 0;
+		//CopyRegion.bufferOffset = 0;
 		CopyRegion.bufferRowLength = TextureRHI2D->GetSizeX();
 		CopyRegion.bufferImageHeight = TextureRHI2D->GetSizeY();
 		CopyRegion.imageSubresource.aspectMask = Texture2D->Surface.GetFullAspectMask();
-		//Region.imageSubresource.mipLevel = 0;
-		//Region.imageSubresource.baseArrayLayer = 0;
+		CopyRegion.imageSubresource.mipLevel = InFlags.GetMip();
+		//CopyRegion.imageSubresource.baseArrayLayer = 0;
 		CopyRegion.imageSubresource.layerCount = 1;
 		CopyRegion.imageExtent.width = TextureRHI2D->GetSizeX();
 		CopyRegion.imageExtent.height = TextureRHI2D->GetSizeY();
