Migrated OpenGL to new RHIShaderResourceView API.

diff --git a/Engine/Source/Runtime/OpenGLDrv/Private/OpenGLTexture.cpp b/Engine/Source/Runtime/OpenGLDrv/Private/OpenGLTexture.cpp
--- a/Engine/Source/Runtime/OpenGLDrv/Private/OpenGLTexture.cpp
+++ b/Engine/Source/Runtime/OpenGLDrv/Private/OpenGLTexture.cpp
@@ -1859,245 +1859,201 @@ void FOpenGLDynamicRHI::RHIGetResourceInfo(FTextureRHIParamRef Ref, FRHIResource
 
 FShaderResourceViewRHIRef FOpenGLDynamicRHI::RHICreateShaderResourceView(FTextureRHIParamRef Texture, const FRHITextureSRVCreateInfo& CreateInfo)
 {
-#if 0
-	FOpenGLShaderResourceViewProxy *ViewProxy = new FOpenGLShaderResourceViewProxy([=](FRHIShaderResourceView* OwnerRHI)
-	{
-		FOpenGLTexture2D* Texture2D = ResourceCast(Texture2DRHI);
+	const uint32 MipLevel = CreateInfo.MipLevel;
+	const uint32 NumMipLevels = CreateInfo.NumMipLevels;
+	const uint8 Format = CreateInfo.Format;
 
-		if (FOpenGL::SupportsTextureView())
+	FOpenGLShaderResourceViewProxy *ViewProxy = new FOpenGLShaderResourceViewProxy([this, Texture, MipLevel, NumMipLevels, Format](FRHIShaderResourceView* OwnerRHI) -> FOpenGLShaderResourceView*
+	{
+		if (FRHITexture2D* Texture2DRHI = Texture->GetTexture2D())
 		{
-			VERIFY_GL_SCOPE();
+			FOpenGLTexture2D* Texture2D = ResourceCast(Texture2DRHI);
+			FOpenGLShaderResourceView *View = 0;
 
-			GLuint Resource = 0;
+			if (FOpenGL::SupportsTextureView())
+			{
+				VERIFY_GL_SCOPE();
 
-			FOpenGL::GenTextures(1, &Resource);
-			const FOpenGLTextureFormat& GLFormat = GOpenGLTextureFormats[Texture2D->GetFormat()];
-			const bool bSRGB = (Texture2D->GetFlags()&TexCreate_SRGB) != 0;
+				GLuint Resource = 0;
 
-			FOpenGL::TextureView(Resource, Texture2D->Target, Texture2D->Resource, GLFormat.InternalFormat[bSRGB], MipLevel, 1, 0, 1);
+				FOpenGL::GenTextures(1, &Resource);
 
-			return new FOpenGLShaderResourceView(this, Resource, Texture2D->Target, MipLevel, true);
-		}
-		else
-		{
-			return new FOpenGLShaderResourceView(this, Texture2D->Resource, Texture2D->Target, MipLevel, false);
-		}
-
-	});
-
-	return ViewProxy;
-#endif
-	return nullptr;
-}
+				if (Format != PF_X24_G8)
+				{
+					const FOpenGLTextureFormat& GLFormat = GOpenGLTextureFormats[Format];
+					const bool bSRGB = (Texture2D->GetFlags()&TexCreate_SRGB) != 0;
 
-#if 0
-FShaderResourceViewRHIRef FOpenGLDynamicRHI::RHICreateShaderResourceView(FTexture2DRHIParamRef Texture2DRHI, uint8 MipLevel, uint8 NumMipLevels, uint8 Format)
-{
-	FOpenGLShaderResourceViewProxy *ViewProxy = new FOpenGLShaderResourceViewProxy([=](FRHIShaderResourceView* OwnerRHI)
-	{
-		FOpenGLTexture2D* Texture2D = ResourceCast(Texture2DRHI);
-		FOpenGLShaderResourceView *View = 0;
+					FOpenGL::TextureView(Resource, Texture2D->Target, Texture2D->Resource, GLFormat.InternalFormat[bSRGB], MipLevel, NumMipLevels, 0, 1);
+				}
+				else
+				{
+					// PF_X24_G8 doesn't correspond to a real format under OpenGL
+					// The solution is to create a view with the original format, and convert it to return the stencil index
+					// To match component locations, texture swizzle needs to be setup too
+					const FOpenGLTextureFormat& GLFormat = GOpenGLTextureFormats[Texture2D->GetFormat()];
 
-		if (FOpenGL::SupportsTextureView())
-		{
-			VERIFY_GL_SCOPE();
+					// create a second depth/stencil view
+					FOpenGL::TextureView(Resource, Texture2D->Target, Texture2D->Resource, GLFormat.InternalFormat[0], MipLevel, NumMipLevels, 0, 1);
 
-			GLuint Resource = 0;
+					// Use a texture stage that's not likely to be used for draws, to avoid waiting
+					FOpenGLContextState& ContextState = GetContextStateForCurrentContext();
+					CachedSetupTextureStage(ContextState, FOpenGL::GetMaxCombinedTextureImageUnits() - 1, Texture2D->Target, Resource, 0, NumMipLevels);
 
-			FOpenGL::GenTextures( 1, &Resource);
+					//set the texture to return the stencil index, and then force the components to match D3D
+					glTexParameteri(Texture2D->Target, GL_DEPTH_STENCIL_TEXTURE_MODE, GL_STENCIL_INDEX);
+					glTexParameteri(Texture2D->Target, GL_TEXTURE_SWIZZLE_R, GL_ZERO);
+					glTexParameteri(Texture2D->Target, GL_TEXTURE_SWIZZLE_G, GL_RED);
+					glTexParameteri(Texture2D->Target, GL_TEXTURE_SWIZZLE_B, GL_ZERO);
+					glTexParameteri(Texture2D->Target, GL_TEXTURE_SWIZZLE_A, GL_ZERO);
+				}
 
-			if (Format != PF_X24_G8)
-			{
-				const FOpenGLTextureFormat& GLFormat = GOpenGLTextureFormats[Format];
-				const bool bSRGB = (Texture2D->GetFlags()&TexCreate_SRGB) != 0;
-		
-				FOpenGL::TextureView( Resource, Texture2D->Target, Texture2D->Resource, GLFormat.InternalFormat[bSRGB], MipLevel, NumMipLevels, 0, 1);
+				View = new FOpenGLShaderResourceView(this, Resource, Texture2D->Target, MipLevel, true);
 			}
 			else
 			{
-				// PF_X24_G8 doesn't correspond to a real format under OpenGL
-				// The solution is to create a view with the original format, and convert it to return the stencil index
-				// To match component locations, texture swizzle needs to be setup too
-				const FOpenGLTextureFormat& GLFormat = GOpenGLTextureFormats[Texture2D->GetFormat()];
+				uint32 const Target = Texture2D->Target;
+				GLuint Resource = Texture2D->Resource;
 
-				// create a second depth/stencil view
-				FOpenGL::TextureView( Resource, Texture2D->Target, Texture2D->Resource, GLFormat.InternalFormat[0], MipLevel, NumMipLevels, 0, 1);
+				FTexture2DRHIParamRef DepthStencilTex = nullptr;
 
-				// Use a texture stage that's not likely to be used for draws, to avoid waiting
-				FOpenGLContextState& ContextState = GetContextStateForCurrentContext();
-				CachedSetupTextureStage(ContextState, FOpenGL::GetMaxCombinedTextureImageUnits() - 1, Texture2D->Target, Resource, 0, NumMipLevels);
-
-				//set the texture to return the stencil index, and then force the components to match D3D
-				glTexParameteri( Texture2D->Target, GL_DEPTH_STENCIL_TEXTURE_MODE, GL_STENCIL_INDEX);
-				glTexParameteri( Texture2D->Target, GL_TEXTURE_SWIZZLE_R, GL_ZERO);
-				glTexParameteri( Texture2D->Target, GL_TEXTURE_SWIZZLE_G, GL_RED);
-				glTexParameteri( Texture2D->Target, GL_TEXTURE_SWIZZLE_B, GL_ZERO);
-				glTexParameteri( Texture2D->Target, GL_TEXTURE_SWIZZLE_A, GL_ZERO);
-			}
-		
-			View = new FOpenGLShaderResourceView(this, Resource, Texture2D->Target, MipLevel, true);
-		}
-		else
-		{
-			uint32 const Target = Texture2D->Target;
-			GLuint Resource = Texture2D->Resource;
-		
-			FTexture2DRHIParamRef DepthStencilTex = nullptr;
-		
-			// For stencil sampling we have to use a separate single channel texture to blit stencil data into
-	#if PLATFORM_DESKTOP || PLATFORM_ANDROIDESDEFERRED
-			if (FOpenGL::GetFeatureLevel() >= ERHIFeatureLevel::SM4 && Format == PF_X24_G8 && FOpenGL::SupportsPixelBuffers())
-			{
-				check(NumMipLevels == 1 && MipLevel == 0);
-			
-				if (!Texture2D->SRVResource)
+				// For stencil sampling we have to use a separate single channel texture to blit stencil data into
+#if PLATFORM_DESKTOP || PLATFORM_ANDROIDESDEFERRED
+				if (FOpenGL::GetFeatureLevel() >= ERHIFeatureLevel::SM4 && Format == PF_X24_G8 && FOpenGL::SupportsPixelBuffers())
 				{
-					FOpenGL::GenTextures(1, &Texture2D->SRVResource);
-				
-					GLenum const InternalFormat = GL_R8UI;
-					GLenum const ChannelFormat = GL_RED_INTEGER;
-					uint32 const SizeX = Texture2D->GetSizeX();
-					uint32 const SizeY = Texture2D->GetSizeY();
-					GLenum const Type = GL_UNSIGNED_BYTE;
-					uint32 const Flags = 0;
-				
-					FOpenGLContextState& ContextState = GetContextStateForCurrentContext();
-					CachedSetupTextureStage(ContextState, FOpenGL::GetMaxCombinedTextureImageUnits() - 1, Target, Texture2D->SRVResource, MipLevel, NumMipLevels);
-				
-					if (!FOpenGL::TexStorage2D(Target, NumMipLevels, InternalFormat, SizeX, SizeY, ChannelFormat, Type, Flags))
+					check(NumMipLevels == 1 && MipLevel == 0);
+
+					if (!Texture2D->SRVResource)
 					{
-						glTexImage2D(Target, 0, InternalFormat, SizeX, SizeY, 0, ChannelFormat, Type, nullptr);
-					}
-				
-					TArray<uint8> ZeroData;
-					ZeroData.AddZeroed(SizeX * SizeY);
-				
-					glPixelStorei(GL_UNPACK_ALIGNMENT, 1);
-					glTexSubImage2D(
-									Target,
-									0,
-									0,
-									0,
-									SizeX,
-									SizeY,
-									ChannelFormat,
-									Type,
-									ZeroData.GetData());
-					glPixelStorei(GL_UNPACK_ALIGNMENT, 4);
-				
-					//set the texture to return the stencil index, and then force the components to match D3D
-					glTexParameteri( Target, GL_TEXTURE_SWIZZLE_R, GL_ZERO);
-					glTexParameteri( Target, GL_TEXTURE_SWIZZLE_G, GL_RED);
-					glTexParameteri( Target, GL_TEXTURE_SWIZZLE_B, GL_ZERO);
-					glTexParameteri( Target, GL_TEXTURE_SWIZZLE_A, GL_ZERO);
-				}
-				check(Texture2D->SRVResource);
-			
-				Resource = Texture2D->SRVResource;
-				DepthStencilTex = Texture2DRHI;
-			}
-	#endif
-		
-			View = new FOpenGLShaderResourceView(this, Resource, Target, MipLevel, false);
-			View->Texture2D = DepthStencilTex;
-		}
-		return View;
-	});
+						FOpenGL::GenTextures(1, &Texture2D->SRVResource);
 
-	return ViewProxy;
-}
+						GLenum const InternalFormat = GL_R8UI;
+						GLenum const ChannelFormat = GL_RED_INTEGER;
+						uint32 const SizeX = Texture2D->GetSizeX();
+						uint32 const SizeY = Texture2D->GetSizeY();
+						GLenum const Type = GL_UNSIGNED_BYTE;
+						uint32 const Flags = 0;
 
-FShaderResourceViewRHIRef FOpenGLDynamicRHI::RHICreateShaderResourceView(FTexture3DRHIParamRef Texture3DRHI, uint8 MipLevel)
-{
-	FOpenGLShaderResourceViewProxy *ViewProxy = new FOpenGLShaderResourceViewProxy([=](FRHIShaderResourceView* OwnerRHI)
-	{
-		FOpenGLTexture3D* Texture3D = ResourceCast(Texture3DRHI);
+						FOpenGLContextState& ContextState = GetContextStateForCurrentContext();
+						CachedSetupTextureStage(ContextState, FOpenGL::GetMaxCombinedTextureImageUnits() - 1, Target, Texture2D->SRVResource, MipLevel, NumMipLevels);
 
-		FOpenGLShaderResourceView *View = 0;
+						if (!FOpenGL::TexStorage2D(Target, NumMipLevels, InternalFormat, SizeX, SizeY, ChannelFormat, Type, Flags))
+						{
+							glTexImage2D(Target, 0, InternalFormat, SizeX, SizeY, 0, ChannelFormat, Type, nullptr);
+						}
 
-		if (FOpenGL::SupportsTextureView())
-		{
-			VERIFY_GL_SCOPE();
+						TArray<uint8> ZeroData;
+						ZeroData.AddZeroed(SizeX * SizeY);
 
-			GLuint Resource = 0;
+						glPixelStorei(GL_UNPACK_ALIGNMENT, 1);
+						glTexSubImage2D(
+							Target,
+							0,
+							0,
+							0,
+							SizeX,
+							SizeY,
+							ChannelFormat,
+							Type,
+							ZeroData.GetData());
+						glPixelStorei(GL_UNPACK_ALIGNMENT, 4);
 
-			FOpenGL::GenTextures( 1, &Resource);
-			const FOpenGLTextureFormat& GLFormat = GOpenGLTextureFormats[Texture3D->GetFormat()];
-			const bool bSRGB = (Texture3D->GetFlags()&TexCreate_SRGB) != 0;
-		
-			FOpenGL::TextureView( Resource, Texture3D->Target, Texture3D->Resource, GLFormat.InternalFormat[bSRGB], MipLevel, 1, 0, 1);
-		
-			return new FOpenGLShaderResourceView(this, Resource, Texture3D->Target, MipLevel, true);
+						//set the texture to return the stencil index, and then force the components to match D3D
+						glTexParameteri(Target, GL_TEXTURE_SWIZZLE_R, GL_ZERO);
+						glTexParameteri(Target, GL_TEXTURE_SWIZZLE_G, GL_RED);
+						glTexParameteri(Target, GL_TEXTURE_SWIZZLE_B, GL_ZERO);
+						glTexParameteri(Target, GL_TEXTURE_SWIZZLE_A, GL_ZERO);
+					}
+					check(Texture2D->SRVResource);
+
+					Resource = Texture2D->SRVResource;
+					DepthStencilTex = Texture2DRHI;
+				}
+#endif
+
+				View = new FOpenGLShaderResourceView(this, Resource, Target, MipLevel, false);
+				View->Texture2D = DepthStencilTex;
+			}
+			return View;
 		}
-		else
+		else if (FRHITexture2DArray* Texture2DArrayRHI = Texture->GetTexture2DArray())
 		{
-			return new FOpenGLShaderResourceView(this, Texture3D->Resource, Texture3D->Target, MipLevel, false);
-		}
-	});
+			FOpenGLTexture2DArray* Texture2DArray = ResourceCast(Texture2DArrayRHI);
+			FOpenGLShaderResourceView *View = 0;
 
-	return ViewProxy;
-}
+			if (FOpenGL::SupportsTextureView())
+			{
+				VERIFY_GL_SCOPE();
 
-FShaderResourceViewRHIRef FOpenGLDynamicRHI::RHICreateShaderResourceView(FTexture2DArrayRHIParamRef Texture2DArrayRHI, uint8 MipLevel)
-{
-	FOpenGLShaderResourceViewProxy *ViewProxy = new FOpenGLShaderResourceViewProxy([=](FRHIShaderResourceView* OwnerRHI)
-	{
-		FOpenGLTexture2DArray* Texture2DArray = ResourceCast(Texture2DArrayRHI);
-		FOpenGLShaderResourceView *View = 0;
+				GLuint Resource = 0;
+
+				FOpenGL::GenTextures(1, &Resource);
+				const FOpenGLTextureFormat& GLFormat = GOpenGLTextureFormats[Texture2DArray->GetFormat()];
+				const bool bSRGB = (Texture2DArray->GetFlags()&TexCreate_SRGB) != 0;
+
+				FOpenGL::TextureView(Resource, Texture2DArray->Target, Texture2DArray->Resource, GLFormat.InternalFormat[bSRGB], MipLevel, 1, 0, 1);
 
-		if (FOpenGL::SupportsTextureView())
+				return new FOpenGLShaderResourceView(this, Resource, Texture2DArray->Target, MipLevel, true);
+			}
+			else
+			{
+				return new FOpenGLShaderResourceView(this, Texture2DArray->Resource, Texture2DArray->Target, MipLevel, false);
+			}
+		}
+		else if (FRHITextureCube* TextureCubeRHI = Texture->GetTextureCube())
 		{
-			VERIFY_GL_SCOPE();
+			FOpenGLTextureCube* TextureCube = ResourceCast(TextureCubeRHI);
+			if (FOpenGL::SupportsTextureView())
+			{
+				VERIFY_GL_SCOPE();
 
-			GLuint Resource = 0;
+				GLuint Resource = 0;
 
-			FOpenGL::GenTextures(1, &Resource);
-			const FOpenGLTextureFormat& GLFormat = GOpenGLTextureFormats[Texture2DArray->GetFormat()];
-			const bool bSRGB = (Texture2DArray->GetFlags()&TexCreate_SRGB) != 0;
+				FOpenGL::GenTextures(1, &Resource);
+				const FOpenGLTextureFormat& GLFormat = GOpenGLTextureFormats[TextureCube->GetFormat()];
+				const bool bSRGB = (TextureCube->GetFlags()&TexCreate_SRGB) != 0;
 
-			FOpenGL::TextureView(Resource, Texture2DArray->Target, Texture2DArray->Resource, GLFormat.InternalFormat[bSRGB], MipLevel, 1, 0, 1);
+				FOpenGL::TextureView(Resource, TextureCube->Target, TextureCube->Resource, GLFormat.InternalFormat[bSRGB], MipLevel, 1, 0, 6);
 
-			return new FOpenGLShaderResourceView(this, Resource, Texture2DArray->Target, MipLevel, true);
+				return new FOpenGLShaderResourceView(this, Resource, TextureCube->Target, MipLevel, true);
+			}
+			else
+			{
+				return new FOpenGLShaderResourceView(this, TextureCube->Resource, TextureCube->Target, MipLevel, false);
+			}
 		}
-		else
+		else if (FRHITexture3D* Texture3DRHI = Texture->GetTexture3D())
 		{
-			return new FOpenGLShaderResourceView(this, Texture2DArray->Resource, Texture2DArray->Target, MipLevel, false);
-		}
-	});
-	
-	return ViewProxy;
-}
+			FOpenGLTexture3D* Texture3D = ResourceCast(Texture3DRHI);
 
-FShaderResourceViewRHIRef FOpenGLDynamicRHI::RHICreateShaderResourceView(FTextureCubeRHIParamRef TextureCubeRHI, uint8 MipLevel)
-{
+			FOpenGLShaderResourceView *View = 0;
 
-	FOpenGLShaderResourceViewProxy *ViewProxy = new FOpenGLShaderResourceViewProxy([=](FRHIShaderResourceView* OwnerRHI)
-	{
-		FOpenGLTextureCube* TextureCube = ResourceCast(TextureCubeRHI);
-		if (FOpenGL::SupportsTextureView())
-		{
-			VERIFY_GL_SCOPE();
+			if (FOpenGL::SupportsTextureView())
+			{
+				VERIFY_GL_SCOPE();
 
-			GLuint Resource = 0;
+				GLuint Resource = 0;
 
-			FOpenGL::GenTextures(1, &Resource);
-			const FOpenGLTextureFormat& GLFormat = GOpenGLTextureFormats[TextureCube->GetFormat()];
-			const bool bSRGB = (TextureCube->GetFlags()&TexCreate_SRGB) != 0;
+				FOpenGL::GenTextures(1, &Resource);
+				const FOpenGLTextureFormat& GLFormat = GOpenGLTextureFormats[Texture3D->GetFormat()];
+				const bool bSRGB = (Texture3D->GetFlags()&TexCreate_SRGB) != 0;
 
-			FOpenGL::TextureView(Resource, TextureCube->Target, TextureCube->Resource, GLFormat.InternalFormat[bSRGB], MipLevel, 1, 0, 6);
+				FOpenGL::TextureView(Resource, Texture3D->Target, Texture3D->Resource, GLFormat.InternalFormat[bSRGB], MipLevel, 1, 0, 1);
 
-			return new FOpenGLShaderResourceView(this, Resource, TextureCube->Target, MipLevel, true);
+				return new FOpenGLShaderResourceView(this, Resource, Texture3D->Target, MipLevel, true);
+			}
+			else
+			{
+				return new FOpenGLShaderResourceView(this, Texture3D->Resource, Texture3D->Target, MipLevel, false);
+			}
 		}
 		else
 		{
-			return new FOpenGLShaderResourceView(this, TextureCube->Resource, TextureCube->Target, MipLevel, false);
+			check(false);
+			return nullptr;
 		}
 	});
-
 	return ViewProxy;
 }
-#endif
-
 
 /** Generates mip maps for the surface. */
 void FOpenGLDynamicRHI::RHIGenerateMips(FTextureRHIParamRef SurfaceRHI)
