Fixed format selection for GL texture view when PF_Unknown is specified.

#rnx

diff --git a/Engine/Source/Runtime/OpenGLDrv/Private/OpenGLTexture.cpp b/Engine/Source/Runtime/OpenGLDrv/Private/OpenGLTexture.cpp
--- a/Engine/Source/Runtime/OpenGLDrv/Private/OpenGLTexture.cpp
+++ b/Engine/Source/Runtime/OpenGLDrv/Private/OpenGLTexture.cpp
@@ -1880,7 +1880,8 @@ FShaderResourceViewRHIRef FOpenGLDynamicRHI::RHICreateShaderResourceView(FRHITex
 
 				if (Format != PF_X24_G8)
 				{
-					const FOpenGLTextureFormat& GLFormat = GOpenGLTextureFormats[Format];
+					// Choose original format when PF_Unknown is specified (as stated for FRHITextureSRVCreateInfo::Format)
+					const FOpenGLTextureFormat& GLFormat = GOpenGLTextureFormats[Format == PF_Unknown ? Texture2D->GetFormat() : Format];
 					const bool bSRGB = (Texture2D->GetFlags()&TexCreate_SRGB) != 0;
 
 					FOpenGL::TextureView(Resource, Texture2D->Target, Texture2D->Resource, GLFormat.InternalFormat[bSRGB], MipLevel, NumMipLevels, 0, 1);
