DRS - Compile fix
#rnx

diff --git a/Engine/Source/Runtime/OpenGLDrv/Public/OpenGLDrv.h b/Engine/Source/Runtime/OpenGLDrv/Public/OpenGLDrv.h
--- a/Engine/Source/Runtime/OpenGLDrv/Public/OpenGLDrv.h
+++ b/Engine/Source/Runtime/OpenGLDrv/Public/OpenGLDrv.h
@@ -21,8 +21,6 @@
 #include "Runtime/OpenGLDrv/Private/Windows/OpenGLWindows.h"
 #elif PLATFORM_LINUX
 #include "Runtime/OpenGLDrv/Private/Linux/OpenGLLinux.h"
-#elif PLATFORM_IOS
-#include "IOS/IOSOpenGL.h"
 #elif PLATFORM_LUMIN
 // these guys will self-select
 #include "Lumin/LuminESDeferredOpenGL.h"
