Use sized formats with TexStorage3D, fixes gl enum error seen on Android.
#rnx

diff --git a/Engine/Source/Runtime/OpenGLDrv/Private/OpenGLTexture.cpp b/Engine/Source/Runtime/OpenGLDrv/Private/OpenGLTexture.cpp
--- a/Engine/Source/Runtime/OpenGLDrv/Private/OpenGLTexture.cpp
+++ b/Engine/Source/Runtime/OpenGLDrv/Private/OpenGLTexture.cpp
@@ -1664,7 +1664,7 @@ FTexture2DArrayRHIRef FOpenGLDynamicRHI::RHICreateTexture2DArray(uint32 SizeX,ui
 	uint8* Data = Info.BulkData ? (uint8*)Info.BulkData->GetResourceBulkData() : NULL;
 	uint32 MipOffset = 0;
 
-	FOpenGL::TexStorage3D( Target, NumMips, GLFormat.InternalFormat[bSRGB], SizeX, SizeY, SizeZ, GLFormat.Format, GLFormat.Type );
+	FOpenGL::TexStorage3D( Target, NumMips, GLFormat.SizedInternalFormat[bSRGB], SizeX, SizeY, SizeZ, GLFormat.Format, GLFormat.Type );
 
 	if (Data)
 	{
@@ -1781,7 +1781,7 @@ FTexture3DRHIRef FOpenGLDynamicRHI::RHICreateTexture3D(uint32 SizeX,uint32 SizeY
 	uint32 DataSize = CreateInfo.BulkData ? CreateInfo.BulkData->GetResourceBulkDataSize() : 0;
 	uint32 MipOffset = 0;
 
-	FOpenGL::TexStorage3D( Target, NumMips, GLFormat.InternalFormat[bSRGB], SizeX, SizeY, SizeZ, GLFormat.Format, GLFormat.Type );
+	FOpenGL::TexStorage3D( Target, NumMips, GLFormat.SizedInternalFormat[bSRGB], SizeX, SizeY, SizeZ, GLFormat.Format, GLFormat.Type );
 
 	if (Data)
 	{
