Adding two-channel ETC2 compressed format needed for RVT on mobile.

Also added single-channel format while at it.

diff --git a/Engine/Source/Runtime/Core/Public/PixelFormat.h b/Engine/Source/Runtime/Core/Public/PixelFormat.h
--- a/Engine/Source/Runtime/Core/Public/PixelFormat.h
+++ b/Engine/Source/Runtime/Core/Public/PixelFormat.h
@@ -80,7 +80,9 @@ enum EPixelFormat
 	PF_PLATFORM_HDR_2		=66,	// Reserved.
 	PF_NV12					=67,
 	PF_R32G32_UINT          =68,
-	PF_MAX					=69,
+	PF_ETC2_R11_EAC			=69,
+	PF_ETC2_RG11_EAC		=70,
+	PF_MAX					=71,
 };
 #define FOREACH_ENUM_EPIXELFORMAT(op) \
 	op(PF_Unknown) \
diff --git a/Engine/Source/Runtime/CoreUObject/Public/UObject/NoExportTypes.h b/Engine/Source/Runtime/CoreUObject/Public/UObject/NoExportTypes.h
--- a/Engine/Source/Runtime/CoreUObject/Public/UObject/NoExportTypes.h
+++ b/Engine/Source/Runtime/CoreUObject/Public/UObject/NoExportTypes.h
@@ -239,6 +239,8 @@ enum EPixelFormat
 	PF_PLATFORM_HDR_2,
 	PF_NV12,
 	PF_R32G32_UINT,
+	PF_ETC2_R11_EAC,
+	PF_ETC2_RG11_EAC,
 	PF_MAX,
 };
 
diff --git a/Engine/Source/Runtime/OpenGLDrv/Private/OpenGLDevice.cpp b/Engine/Source/Runtime/OpenGLDrv/Private/OpenGLDevice.cpp
--- a/Engine/Source/Runtime/OpenGLDrv/Private/OpenGLDevice.cpp
+++ b/Engine/Source/Runtime/OpenGLDrv/Private/OpenGLDevice.cpp
@@ -1132,8 +1132,10 @@ static void InitRHICapabilitiesForGL()
 #if PLATFORM_ANDROID && !PLATFORM_LUMINGL4
 	if ( FOpenGL::SupportsETC2() )
 	{
-		SetupTextureFormat( PF_ETC2_RGB,	FOpenGLTextureFormat(GL_COMPRESSED_RGB8_ETC2,		GL_COMPRESSED_SRGB8_ETC2,					GL_RGBA,	GL_UNSIGNED_BYTE,	true,		false));
-		SetupTextureFormat( PF_ETC2_RGBA,	FOpenGLTextureFormat(GL_COMPRESSED_RGBA8_ETC2_EAC,	GL_COMPRESSED_SRGB8_ALPHA8_ETC2_EAC,		GL_RGBA,	GL_UNSIGNED_BYTE,	true,		false));
+		SetupTextureFormat( PF_ETC2_RGB,		FOpenGLTextureFormat(GL_COMPRESSED_RGB8_ETC2,		GL_COMPRESSED_SRGB8_ETC2,					GL_RGBA,		GL_UNSIGNED_BYTE,	true,		false));
+		SetupTextureFormat( PF_ETC2_RGBA,		FOpenGLTextureFormat(GL_COMPRESSED_RGBA8_ETC2_EAC,	GL_COMPRESSED_SRGB8_ALPHA8_ETC2_EAC,		GL_RGBA,		GL_UNSIGNED_BYTE,	true,		false));
+		SetupTextureFormat( PF_ETC2_R11_EAC,	FOpenGLTextureFormat(GL_COMPRESSED_R11_EAC,			GL_COMPRESSED_R11_EAC,						GL_RED,			GL_UNSIGNED_BYTE,	true,		false));
+		SetupTextureFormat( PF_ETC2_RG11_EAC,	FOpenGLTextureFormat(GL_COMPRESSED_RG11_EAC,		GL_COMPRESSED_RG11_EAC,						GL_RG,			GL_UNSIGNED_BYTE,	true,		false));
 
 		// ETC2 is a superset of ETC1 with sRGB support
 		if (FOpenGL::SupportsSRGB())
diff --git a/Engine/Source/Runtime/RenderCore/Private/RenderUtils.cpp b/Engine/Source/Runtime/RenderCore/Private/RenderUtils.cpp
--- a/Engine/Source/Runtime/RenderCore/Private/RenderUtils.cpp
+++ b/Engine/Source/Runtime/RenderCore/Private/RenderUtils.cpp
@@ -147,6 +147,9 @@ FPixelFormatInfo	GPixelFormats[PF_MAX] =
 	{ TEXT("NV12"),				1,			1,			1,			1,			1,				0,				0,				PF_NV12             },
 
 	{ TEXT("PF_R32G32_UINT"),   1,   		1,			1,			8,			2,				0,				1,				PF_R32G32_UINT      },
+
+	{ TEXT("PF_ETC2_R11_EAC"),  4,   		4,			1,			8,			1,				0,				0,				PF_ETC2_R11_EAC     },
+	{ TEXT("PF_ETC2_RG11_EAC"), 4,   		4,			1,			16,			2,				0,				0,				PF_ETC2_RG11_EAC    },
 };
 
 static struct FValidatePixelFormats
diff --git a/Engine/Source/Runtime/VulkanRHI/Private/VulkanDevice.cpp b/Engine/Source/Runtime/VulkanRHI/Private/VulkanDevice.cpp
--- a/Engine/Source/Runtime/VulkanRHI/Private/VulkanDevice.cpp
+++ b/Engine/Source/Runtime/VulkanRHI/Private/VulkanDevice.cpp
@@ -657,6 +657,18 @@ void FVulkanDevice::SetupFormats()
 		{
 			SetComponentMapping(PF_ETC2_RGBA, VK_COMPONENT_SWIZZLE_R, VK_COMPONENT_SWIZZLE_G, VK_COMPONENT_SWIZZLE_B, VK_COMPONENT_SWIZZLE_A);
 		}
+
+		MapFormatSupport(PF_ETC2_R11_EAC, VK_FORMAT_EAC_R11_UNORM_BLOCK);
+		if (GPixelFormats[PF_ETC2_R11_EAC].Supported)
+		{
+			SetComponentMapping(PF_ETC2_R11_EAC, VK_COMPONENT_SWIZZLE_R, VK_COMPONENT_SWIZZLE_ZERO, VK_COMPONENT_SWIZZLE_ZERO, VK_COMPONENT_SWIZZLE_ONE);
+		}
+
+		MapFormatSupport(PF_ETC2_RG11_EAC, VK_FORMAT_EAC_R11G11_UNORM_BLOCK);
+		if (GPixelFormats[PF_ETC2_RG11_EAC].Supported)
+		{
+			SetComponentMapping(PF_ETC2_RG11_EAC, VK_COMPONENT_SWIZZLE_R, VK_COMPONENT_SWIZZLE_G, VK_COMPONENT_SWIZZLE_ZERO, VK_COMPONENT_SWIZZLE_ONE);
+		}
 	}
 
 	// Verify available Vertex Formats
