#UE4 CIS fix

diff --git a/Engine/Source/Runtime/OpenGLDrv/Private/Android/AndroidOpenGL.cpp b/Engine/Source/Runtime/OpenGLDrv/Private/Android/AndroidOpenGL.cpp
--- a/Engine/Source/Runtime/OpenGLDrv/Private/Android/AndroidOpenGL.cpp
+++ b/Engine/Source/Runtime/OpenGLDrv/Private/Android/AndroidOpenGL.cpp
@@ -126,7 +126,7 @@ static TAutoConsoleVariable<int32> CVarEnableAdrenoTilingHint(
 	TEXT("Whether Adreno-based Android devices should hint to the driver to use tiling mode for the mobile base pass.\n")
 	TEXT("  0 = hinting disabled\n")
 	TEXT("  1 = hinting enabled for Adreno devices running Andorid 8 or earlier [default]\n")
-	TEXT("  2 = hinting always enabled for Adreno devices\n")
+	TEXT("  2 = hinting always enabled for Adreno devices\n"));
 
 struct FPlatformOpenGLDevice
 {
