UE4R - Deprecate ES2 from the programming side
#rnx

diff --git a/Engine/Plugins/FX/Niagara/Source/Niagara/Private/NiagaraComponent.cpp b/Engine/Plugins/FX/Niagara/Source/Niagara/Private/NiagaraComponent.cpp
--- a/Engine/Plugins/FX/Niagara/Source/Niagara/Private/NiagaraComponent.cpp
+++ b/Engine/Plugins/FX/Niagara/Source/Niagara/Private/NiagaraComponent.cpp
@@ -701,10 +701,12 @@ void UNiagaraComponent::Activate(bool bReset /* = false */)
 {
 	bAwaitingActivationDueToNotReady = false;
 
+	PRAGMA_DISABLE_DEPRECATION_WARNINGS
 	if (IsES2Platform(GShaderPlatformForFeatureLevel[GMaxRHIFeatureLevel]))
 	{
 		GbSuppressNiagaraSystems = 1;
 	}
+	PRAGMA_ENABLE_DEPRECATION_WARNINGS
 
 	if (GbSuppressNiagaraSystems != 0)
 	{
diff --git a/Engine/Plugins/Lumin/MagicLeap/Source/MagicLeap/Private/MagicLeapCustomPresentGL.cpp b/Engine/Plugins/Lumin/MagicLeap/Source/MagicLeap/Private/MagicLeapCustomPresentGL.cpp
--- a/Engine/Plugins/Lumin/MagicLeap/Source/MagicLeap/Private/MagicLeapCustomPresentGL.cpp
+++ b/Engine/Plugins/Lumin/MagicLeap/Source/MagicLeap/Private/MagicLeapCustomPresentGL.cpp
@@ -82,7 +82,9 @@ void FMagicLeapCustomPresentOpenGL::FinishRendering()
 		glBindFramebuffer(GL_READ_FRAMEBUFFER, Framebuffers[0]);
 		glBindFramebuffer(GL_DRAW_FRAMEBUFFER, Framebuffers[1]);
 
+		PRAGMA_DISABLE_DEPRECATION_WARNINGS
 		const bool bShouldFlipVertically = !IsES2Platform(GMaxRHIShaderPlatform);
+		PRAGMA_ENABLE_DEPRECATION_WARNINGS
 
 		bShouldFlipVertically ?
 			FOpenGL::BlitFramebuffer(0, 0, SizeX / 2, SizeY, 0, vp_height, vp_width, 0, GL_COLOR_BUFFER_BIT, GL_NEAREST) :
diff --git a/Engine/Source/Runtime/Engine/Private/StaticMesh.cpp b/Engine/Source/Runtime/Engine/Private/StaticMesh.cpp
--- a/Engine/Source/Runtime/Engine/Private/StaticMesh.cpp
+++ b/Engine/Source/Runtime/Engine/Private/StaticMesh.cpp
@@ -1069,6 +1069,7 @@ void FStaticMeshLODResources::ConditionalForce16BitIndexBuffer(EShaderPlatform M
 {
 	// Initialize the vertex and index buffers.
 	// All platforms supporting Metal also support 32-bit indices.
+	PRAGMA_DISABLE_DEPRECATION_WARNINGS
 	if (IsES2Platform(MaxShaderPlatform) && !IsMetalPlatform(MaxShaderPlatform))
 	{
 		if (IndexBuffer.Is32Bit())
@@ -1080,6 +1081,7 @@ void FStaticMeshLODResources::ConditionalForce16BitIndexBuffer(EShaderPlatform M
 			UE_LOG(LogStaticMesh, Warning, TEXT("[%s] Mesh has more that 65535 vertices, incompatible with mobile; forcing 16-bit (will probably cause rendering issues)."), *Parent->GetName());
 		}
 	}
+	PRAGMA_ENABLE_DEPRECATION_WARNINGS
 }
 
 template <bool bIncrement>
diff --git a/Engine/Source/Runtime/Engine/Public/SimpleElementShaders.h b/Engine/Source/Runtime/Engine/Public/SimpleElementShaders.h
--- a/Engine/Source/Runtime/Engine/Public/SimpleElementShaders.h
+++ b/Engine/Source/Runtime/Engine/Public/SimpleElementShaders.h
@@ -353,7 +353,9 @@ public:
 
 	static bool ShouldCompilePermutation(const FGlobalShaderPermutationParameters& Parameters)
 	{
+		PRAGMA_DISABLE_DEPRECATION_WARNINGS
 		return ((IsES2Platform(Parameters.Platform) && IsPCPlatform(Parameters.Platform)) || Parameters.Platform == SP_OPENGL_ES2_ANDROID) && TSimpleElementBase::ShouldCompilePermutation(Parameters);
+		PRAGMA_ENABLE_DEPRECATION_WARNINGS
 	}
 };
 
diff --git a/Engine/Source/Runtime/OpenGLDrv/Private/OpenGLCommands.cpp b/Engine/Source/Runtime/OpenGLDrv/Private/OpenGLCommands.cpp
--- a/Engine/Source/Runtime/OpenGLDrv/Private/OpenGLCommands.cpp
+++ b/Engine/Source/Runtime/OpenGLDrv/Private/OpenGLCommands.cpp
@@ -1645,6 +1645,7 @@ void FOpenGLDynamicRHI::RHISetRenderTargets(
 
 	FOpenGLTextureBase* NewDepthStencilRT = GetOpenGLTextureFromRHITexture(NewDepthStencilTargetRHI ? NewDepthStencilTargetRHI->Texture : nullptr);
 
+	PRAGMA_DISABLE_DEPRECATION_WARNINGS
 	if (IsES2Platform(GMaxRHIShaderPlatform) && !IsPCPlatform(GMaxRHIShaderPlatform))
 	{
 		// @todo-mobile
@@ -1679,6 +1680,7 @@ void FOpenGLDynamicRHI::RHISetRenderTargets(
 				ContextState.LastES2ColorTargetType = NewColorTargetType;
 		}
 	}
+	PRAGMA_ENABLE_DEPRECATION_WARNINGS
 	
 	PendingState.DepthStencil = NewDepthStencilRT;
 	PendingState.StencilStoreAction = NewDepthStencilTargetRHI ? NewDepthStencilTargetRHI->GetStencilStoreAction() : ERenderTargetStoreAction::ENoAction;
diff --git a/Engine/Source/Runtime/OpenGLDrv/Private/OpenGLDevice.cpp b/Engine/Source/Runtime/OpenGLDrv/Private/OpenGLDevice.cpp
--- a/Engine/Source/Runtime/OpenGLDrv/Private/OpenGLDevice.cpp
+++ b/Engine/Source/Runtime/OpenGLDrv/Private/OpenGLDevice.cpp
@@ -807,7 +807,9 @@ static void InitRHICapabilitiesForGL()
 	
 	// Emulate uniform buffers on ES2
 	// Optional emulation on ES3.1 and PC
+	PRAGMA_DISABLE_DEPRECATION_WARNINGS
 	GUseEmulatedUniformBuffers = (IsES2Platform(GMaxRHIShaderPlatform) && !IsPCPlatform(GMaxRHIShaderPlatform)) || bUseEmulatedUBs;
+	PRAGMA_ENABLE_DEPRECATION_WARNINGS
 
 	FString FeatureLevelName;
 	GetFeatureLevelName(GMaxRHIFeatureLevel, FeatureLevelName);
@@ -1304,6 +1306,7 @@ static void CheckVaryingLimit()
 {
 #if PLATFORM_ANDROID && !PLATFORM_LUMINGL4
 	FOpenGL::bRequiresGLFragCoordVaryingLimitHack = false;
+	PRAGMA_DISABLE_DEPRECATION_WARNINGS
 	if (IsES2Platform(GMaxRHIShaderPlatform))
 	{
 		// Some mobile GPUs require an available varying vector to support gl_FragCoord.
@@ -1443,17 +1446,21 @@ static void CheckVaryingLimit()
 			AttemptLinkProgramExecute(VertexShader->Resource, PixelShader->Resource);
 		}
 	}
+	PRAGMA_ENABLE_DEPRECATION_WARNINGS
 #elif PLATFORM_IOS
+	PRAGMA_DISABLE_DEPRECATION_WARNINGS
 	if (IsES2Platform(GMaxRHIShaderPlatform))
 	{
 		FOpenGL::bIsLimitingShaderCompileCount = FPlatformMisc::GetIOSDeviceType() == FPlatformMisc::IOS_IPad4;
 	}
+	PRAGMA_ENABLE_DEPRECATION_WARNINGS
 #endif
 }
 
 static void CheckTextureCubeLodSupport()
 {
 #if PLATFORM_ANDROID && !PLATFORM_LUMINGL4
+	PRAGMA_DISABLE_DEPRECATION_WARNINGS
 	if (IsES2Platform(GMaxRHIShaderPlatform))
 	{
 		UE_LOG(LogRHI, Display, TEXT("Testing for shader compiler compatibility"));
@@ -1580,6 +1587,7 @@ static void CheckTextureCubeLodSupport()
 		UE_LOG(LogRHI, Warning, TEXT("Unable to find a test shader that compiles try running anyway"));
 		FOpenGL::bIsCheckingShaderCompilerHacks = false;
 	}
+	PRAGMA_ENABLE_DEPRECATION_WARNINGS
 #endif
 }
 
@@ -1587,6 +1595,7 @@ static void CheckRoundFunction()
 {
 #if PLATFORM_ANDROID && !PLATFORM_LUMINGL4
 	FOpenGL::bRequiresRoundFunctionHack = false;
+	PRAGMA_DISABLE_DEPRECATION_WARNINGS
 	if (IsES2Platform(GMaxRHIShaderPlatform))
 	{
 		UE_LOG(LogRHI, Display, TEXT("Testing for round() function availability"));
@@ -1635,6 +1644,7 @@ static void CheckRoundFunction()
 
 		UE_LOG(LogRHI, Warning, TEXT("Disabling the need for the round() function hack"));
 	}
+	PRAGMA_ENABLE_DEPRECATION_WARNINGS
 #endif
 }
 
diff --git a/Engine/Source/Runtime/OpenGLDrv/Private/OpenGLShaders.cpp b/Engine/Source/Runtime/OpenGLDrv/Private/OpenGLShaders.cpp
--- a/Engine/Source/Runtime/OpenGLDrv/Private/OpenGLShaders.cpp
+++ b/Engine/Source/Runtime/OpenGLDrv/Private/OpenGLShaders.cpp
@@ -996,6 +996,7 @@ void OPENGLDRV_API GLSLToDeviceCompatibleGLSL(FAnsiCharArray& GlslCodeOriginal,
 	if (Capabilities.TargetPlatform == EOpenGLShaderTargetPlatform::OGLSTP_Android || Capabilities.TargetPlatform == EOpenGLShaderTargetPlatform::OGLSTP_HTML5)
 	{
 		bNeedsExtDrawInstancedDefine = !bES31;
+		PRAGMA_DISABLE_DEPRECATION_WARNINGS
 		if (bES31)
 		{
 			// @todo Lumin hack: This is needed for AEP on Lumin, so that some shaders compile that need version 320
@@ -1022,7 +1023,7 @@ void OPENGLDRV_API GLSLToDeviceCompatibleGLSL(FAnsiCharArray& GlslCodeOriginal,
 			}
 			ReplaceCString(GlslCodeOriginal, "#version 100", "");
 		}
-
+		PRAGMA_ENABLE_DEPRECATION_WARNINGS
 	}
 	else if (Capabilities.TargetPlatform == EOpenGLShaderTargetPlatform::OGLSTP_iOS)
 	{
@@ -1178,6 +1179,7 @@ void OPENGLDRV_API GLSLToDeviceCompatibleGLSL(FAnsiCharArray& GlslCodeOriginal,
 			}
 		}
 
+		PRAGMA_DISABLE_DEPRECATION_WARNINGS
 		if (IsES2Platform(Capabilities.MaxRHIShaderPlatform) && !bES31)
 		{
 			const ANSICHAR * EncodeModeDefine = nullptr;
@@ -1366,6 +1368,7 @@ void OPENGLDRV_API GLSLToDeviceCompatibleGLSL(FAnsiCharArray& GlslCodeOriginal,
 				}
 			}
 		}
+		PRAGMA_ENABLE_DEPRECATION_WARNINGS
 	}
 	else if (Capabilities.TargetPlatform == EOpenGLShaderTargetPlatform::OGLSTP_HTML5)
 	{
diff --git a/Engine/Source/Runtime/RHI/Public/RHIDefinitions.h b/Engine/Source/Runtime/RHI/Public/RHIDefinitions.h
--- a/Engine/Source/Runtime/RHI/Public/RHIDefinitions.h
+++ b/Engine/Source/Runtime/RHI/Public/RHIDefinitions.h
@@ -1037,6 +1037,7 @@ inline bool IsPCPlatform(const EShaderPlatform Platform)
 }
 
 /** Whether the shader platform corresponds to the ES2 feature level. */
+UE_DEPRECATED(4.23, "Feature level ES2 is getting deprecated. Please use ES3.1.")
 inline bool IsES2Platform(const EShaderPlatform Platform)
 {
 	return Platform == SP_PCD3D_ES2 || Platform == SP_OPENGL_PCES2 || Platform == SP_OPENGL_ES2_ANDROID || Platform == SP_OPENGL_ES2_WEBGL || Platform == SP_METAL_MACES2
@@ -1046,7 +1047,10 @@ inline bool IsES2Platform(const EShaderPlatform Platform)
 /** Whether the shader platform corresponds to the ES2/ES3.1 feature level. */
 inline bool IsMobilePlatform(const EShaderPlatform Platform)
 {
-	return IsES2Platform(Platform)
+	return 
+		PRAGMA_DISABLE_DEPRECATION_WARNINGS
+		IsES2Platform(Platform)
+		PRAGMA_ENABLE_DEPRECATION_WARNINGS
 		|| Platform == SP_METAL || Platform == SP_METAL_MACES3_1 || Platform == SP_METAL_TVOS
 		|| Platform == SP_PCD3D_ES3_1
 		|| Platform == SP_OPENGL_PCES3_1 || Platform == SP_OPENGL_ES3_1_ANDROID
