Move implementation of unused function to cpp in hopes of resolving undefined symbol errors.
#rnx

diff --git a/Engine/Source/Runtime/OpenGLDrv/Private/OpenGLUtil.cpp b/Engine/Source/Runtime/OpenGLDrv/Private/OpenGLUtil.cpp
--- a/Engine/Source/Runtime/OpenGLDrv/Private/OpenGLUtil.cpp
+++ b/Engine/Source/Runtime/OpenGLDrv/Private/OpenGLUtil.cpp
@@ -31,6 +31,26 @@ void VerifyOpenGLResult(GLenum ErrorCode, const TCHAR* Msg1, const TCHAR* Msg2,
 	}
 }
 
+GLenum GetOpenGLCubeFace(ECubeFace Face)
+{
+	switch (Face)
+	{
+	case CubeFace_PosX:
+	default:
+		return GL_TEXTURE_CUBE_MAP_POSITIVE_X;
+	case CubeFace_NegX:
+		return GL_TEXTURE_CUBE_MAP_NEGATIVE_X;
+	case CubeFace_PosY:
+		return GL_TEXTURE_CUBE_MAP_POSITIVE_Y;
+	case CubeFace_NegY:
+		return GL_TEXTURE_CUBE_MAP_NEGATIVE_Y;
+	case CubeFace_PosZ:
+		return GL_TEXTURE_CUBE_MAP_POSITIVE_Z;
+	case CubeFace_NegZ:
+		return GL_TEXTURE_CUBE_MAP_NEGATIVE_Z;
+	};
+}
+
 //
 // Stat declarations.
 //
diff --git a/Engine/Source/Runtime/OpenGLDrv/Public/OpenGLUtil.h b/Engine/Source/Runtime/OpenGLDrv/Public/OpenGLUtil.h
--- a/Engine/Source/Runtime/OpenGLDrv/Public/OpenGLUtil.h
+++ b/Engine/Source/Runtime/OpenGLDrv/Public/OpenGLUtil.h
@@ -39,25 +39,7 @@
 * @param Face - ECubeFace type to convert
 * @return OpenGL cube face enum value
 */
-FORCEINLINE GLenum GetOpenGLCubeFace(ECubeFace Face)
-{
-	switch(Face)
-	{
-	case CubeFace_PosX:
-	default:
-		return GL_TEXTURE_CUBE_MAP_POSITIVE_X;
-	case CubeFace_NegX:
-		return GL_TEXTURE_CUBE_MAP_NEGATIVE_X;
-	case CubeFace_PosY:
-		return GL_TEXTURE_CUBE_MAP_POSITIVE_Y;
-	case CubeFace_NegY:
-		return GL_TEXTURE_CUBE_MAP_NEGATIVE_Y;
-	case CubeFace_PosZ:
-		return GL_TEXTURE_CUBE_MAP_POSITIVE_Z;
-	case CubeFace_NegZ:
-		return GL_TEXTURE_CUBE_MAP_NEGATIVE_Z;
-	};
-}
+GLenum GetOpenGLCubeFace(ECubeFace Face);
 
 extern bool PlatformOpenGLContextValid();
 
