Per CL# 3337801 OpenGLDrv does not support Mac.
#rnx

diff --git a/Engine/Source/Runtime/OpenGLDrv/OpenGLDrv.Build.cs b/Engine/Source/Runtime/OpenGLDrv/OpenGLDrv.Build.cs
--- a/Engine/Source/Runtime/OpenGLDrv/OpenGLDrv.Build.cs
+++ b/Engine/Source/Runtime/OpenGLDrv/OpenGLDrv.Build.cs
@@ -2,7 +2,7 @@
 
 using UnrealBuildTool;
 
-[SupportedPlatformsAttribute(new string[] {"Win32", "Win64", "Mac", "Linux", "Android", "LinuxAArch64"})]
+[SupportedPlatformsAttribute(new string[] {"Win32", "Win64", "Linux", "Android", "LinuxAArch64"})]
 public class OpenGLDrv : ModuleRules
 {
 	public OpenGLDrv(ReadOnlyTargetRules Target) : base(Target)
