Fixed invalid copy extents in OpenGL when passing a default-constructed source or destination rectangle to RHICopyToResolveTarget().

The default constructor for FResolveRect sets the origin to (-1, -1), so we must explicitly set X and Y to 0 when the rectangle is empty (which means copy the whole texture).

diff --git a/Engine/Source/Runtime/OpenGLDrv/Private/OpenGLRenderTarget.cpp b/Engine/Source/Runtime/OpenGLDrv/Private/OpenGLRenderTarget.cpp
--- a/Engine/Source/Runtime/OpenGLDrv/Private/OpenGLRenderTarget.cpp
+++ b/Engine/Source/Runtime/OpenGLDrv/Private/OpenGLRenderTarget.cpp
@@ -465,6 +465,8 @@ void FOpenGLDynamicRHI::RHICopyToResolveTarget(FRHITexture* SourceTextureRHI, FR
 		{
 			// Empty rect mans that the entire source is to be copied. Note that we can't use ResolveParams.Rect.IsValid(), because it
 			// returns false if the rectangle is "inside out" (e.g. X1 > X2), and we want to perform flipping when that's the case.
+			SrcRect.Min.X = 0;
+			SrcRect.Min.Y = 0;
 			SrcRect.Max.X = GetOpenGLTextureSizeXFromRHITexture(SourceTextureRHI);
 			SrcRect.Max.Y = GetOpenGLTextureSizeYFromRHITexture(SourceTextureRHI);
 			SrcRect.Max.X = FMath::Max<int32>(1, SrcRect.Max.X >> ResolveParams.MipIndex);
@@ -474,6 +476,8 @@ void FOpenGLDynamicRHI::RHICopyToResolveTarget(FRHITexture* SourceTextureRHI, FR
 		FIntRect DestRect(ResolveParams.DestRect.X1, ResolveParams.DestRect.Y1, ResolveParams.DestRect.X2, ResolveParams.DestRect.Y2);
 		if(DestRect.IsEmpty())
 		{
+			DestRect.Min.X = 0;
+			DestRect.Min.Y = 0;
 			DestRect.Max.X = GetOpenGLTextureSizeXFromRHITexture(DestTextureRHI);
 			DestRect.Max.Y = GetOpenGLTextureSizeYFromRHITexture(DestTextureRHI);
 			DestRect.Max.X = FMath::Max<int32>(1, DestRect.Max.X >> ResolveParams.MipIndex);
