Fixed ES2 test shaders crash with enabled RHIT
#rnx

diff --git a/Engine/Source/Runtime/OpenGLDrv/Private/OpenGLDevice.cpp b/Engine/Source/Runtime/OpenGLDrv/Private/OpenGLDevice.cpp
--- a/Engine/Source/Runtime/OpenGLDrv/Private/OpenGLDevice.cpp
+++ b/Engine/Source/Runtime/OpenGLDrv/Private/OpenGLDevice.cpp
@@ -1399,14 +1399,15 @@ static void CheckVaryingLimit()
 
 		// Try to compile test shaders
 		GOpenGLShaderHackLastCompileSuccess = false;
-		TRefCountPtr<FOpenGLVertexShader> VertexShader = (FOpenGLVertexShader*)(RHICreateVertexShader(VertexShaderCode.GetReadAccess()).GetReference());
+		FVertexShaderRHIRef VertexShaderRHI = RHICreateVertexShader(VertexShaderCode.GetReadAccess());
+		FOpenGLVertexShader* OGLVertexShader = (FOpenGLVertexShader*)VertexShaderRHI.GetReference();
 
 		if (IsRunningRHIInSeparateThread())
 		{
 			RHICmdList.ImmediateFlush(EImmediateFlushType::FlushRHIThread);
 		}
 
-		if (!VerifyCompiledShader(VertexShader->Resource, TestVertexProgram, false))
+		if (!VerifyCompiledShader(OGLVertexShader->Resource, TestVertexProgram, false))
 		{
 			UE_LOG(LogRHI, Warning, TEXT("Vertex shader for varying test failed to compile. Try running anyway."));
 			FOpenGL::bIsCheckingShaderCompilerHacks = false;
@@ -1414,14 +1415,15 @@ static void CheckVaryingLimit()
 		}
 
 		GOpenGLShaderHackLastCompileSuccess = false;
-		TRefCountPtr<FOpenGLPixelShader> PixelShader = (FOpenGLPixelShader*)(RHICreatePixelShader(FragmentShaderCode.GetReadAccess()).GetReference());
+		FPixelShaderRHIRef PixelShaderRHI = RHICreatePixelShader(FragmentShaderCode.GetReadAccess());
+		FOpenGLPixelShader* OGLPixelShader = (FOpenGLPixelShader*)PixelShaderRHI.GetReference();
 
 		if (IsRunningRHIInSeparateThread())
 		{
 			RHICmdList.ImmediateFlush(EImmediateFlushType::FlushRHIThread);
 		}
 
-		if (!VerifyCompiledShader(PixelShader->Resource, TestFragmentProgram, false))
+		if (!VerifyCompiledShader(OGLPixelShader->Resource, TestFragmentProgram, false))
 		{
 			UE_LOG(LogRHI, Warning, TEXT("Fragment shader for varying test failed to compile. Try running anyway."));
 			FOpenGL::bIsCheckingShaderCompilerHacks = false;
@@ -1433,13 +1435,13 @@ static void CheckVaryingLimit()
 		// Now try linking them.. this is where gl_FragCoord may cause a failure
 		if (IsRunningRHIInSeparateThread())
 		{
-			new (RHICmdList.AllocCommand<FRHICommandAttemptLinkProgram>()) FRHICommandAttemptLinkProgram(VertexShader->Resource, PixelShader->Resource);
+			new (RHICmdList.AllocCommand<FRHICommandAttemptLinkProgram>()) FRHICommandAttemptLinkProgram(OGLVertexShader->Resource, OGLPixelShader->Resource);
 			// wait for AttemptLinkProgramExecute to complete.
 			RHICmdList.ImmediateFlush(EImmediateFlushType::FlushRHIThread);
 		}
 		else
 		{
-			AttemptLinkProgramExecute(VertexShader->Resource, PixelShader->Resource);
+			AttemptLinkProgramExecute(OGLVertexShader->Resource, OGLPixelShader->Resource);
 		}
 	}
 	PRAGMA_ENABLE_DEPRECATION_WARNINGS
@@ -1502,14 +1504,15 @@ static void CheckTextureCubeLodSupport()
 
 		// try to compile without any hacks
 		GOpenGLShaderHackLastCompileSuccess = false;
-		TRefCountPtr<FOpenGLPixelShader> PixelShader = (FOpenGLPixelShader*)(RHICreatePixelShader(Code).GetReference());
+		FPixelShaderRHIRef PixelShaderRHI = RHICreatePixelShader(Code);
+		FOpenGLPixelShader* OGLPixelShader = (FOpenGLPixelShader*)(PixelShaderRHI.GetReference());
 
 		if (IsRunningRHIInSeparateThread())
 		{
 			RHICmdList.ImmediateFlush(EImmediateFlushType::FlushRHIThread);
 		}
 
-		if (VerifyCompiledShader(PixelShader->Resource, TestFragmentProgram, false))
+		if (VerifyCompiledShader(OGLPixelShader->Resource, TestFragmentProgram, false))
 		{
 			UE_LOG(LogRHI, Display, TEXT("Shaders compile fine no need to enable hacks"));
 			// we are done
@@ -1524,14 +1527,15 @@ static void CheckTextureCubeLodSupport()
 		// try to compile without using precision for texture samplers
 		// Samsung Galaxy Express	Samsung Galaxy S3	Samsung Galaxy S3 mini	Samsung Galaxy Tab GT-P1000	Samsung Galaxy Tab 2
 		GOpenGLShaderHackLastCompileSuccess = false;
-		PixelShader = (FOpenGLPixelShader*)(RHICreatePixelShader(Code).GetReference());
+		PixelShaderRHI = RHICreatePixelShader(Code);
+		OGLPixelShader = (FOpenGLPixelShader*)(PixelShaderRHI.GetReference());
 
 		if (IsRunningRHIInSeparateThread())
 		{
 			RHICmdList.ImmediateFlush(EImmediateFlushType::FlushRHIThread);
 		}
 
-		if (VerifyCompiledShader(PixelShader->Resource, TestFragmentProgram, false))
+		if (VerifyCompiledShader(OGLPixelShader->Resource, TestFragmentProgram, false))
 		{
 			UE_LOG(LogRHI, Warning, TEXT("Enabling shader compiler hack to remove precision modifiers for texture samplers"));
 
@@ -1545,14 +1549,15 @@ static void CheckTextureCubeLodSupport()
 
 		// third most likely Samsung Galaxy Tab GT-P1000
 		GOpenGLShaderHackLastCompileSuccess = false;
-		PixelShader = (FOpenGLPixelShader*)(RHICreatePixelShader(Code).GetReference());
+		PixelShaderRHI = RHICreatePixelShader(Code);
+		OGLPixelShader = (FOpenGLPixelShader*)(PixelShaderRHI.GetReference());
 
 		if (IsRunningRHIInSeparateThread())
 		{
 			RHICmdList.ImmediateFlush(EImmediateFlushType::FlushRHIThread);
 		}
 
-		if (VerifyCompiledShader(PixelShader->Resource, TestFragmentProgram, false))
+		if (VerifyCompiledShader(OGLPixelShader->Resource, TestFragmentProgram, false))
 		{
 			UE_LOG(LogRHI, Warning, TEXT("Enabling shader compiler hack to redefine textureCubeLodEXT to textureCubeLod"));
 			// we are done
@@ -1565,14 +1570,15 @@ static void CheckTextureCubeLodSupport()
 
 		// try both hacks
 		GOpenGLShaderHackLastCompileSuccess = false;
-		PixelShader = (FOpenGLPixelShader*)(RHICreatePixelShader(Code).GetReference());
-
+		PixelShaderRHI = RHICreatePixelShader(Code);
+		OGLPixelShader = (FOpenGLPixelShader*)(PixelShaderRHI.GetReference());
+		
 		if (IsRunningRHIInSeparateThread())
 		{
 			RHICmdList.ImmediateFlush(EImmediateFlushType::FlushRHIThread);
 		}
 
-		if (VerifyCompiledShader(PixelShader->Resource, TestFragmentProgram, false))
+		if (VerifyCompiledShader(OGLPixelShader->Resource, TestFragmentProgram, false))
 		{
 			UE_LOG(LogRHI, Warning, TEXT("Enabling shader compiler hack to redefine textureCubeLodEXT to textureCubeLod and remove precision modifiers"));
 			// we are done
@@ -1620,15 +1626,16 @@ static void CheckRoundFunction()
 
 		// Try to compile test shaders
 		GOpenGLShaderHackLastCompileSuccess = false;
-		TRefCountPtr<FOpenGLVertexShader> VertexShader = (FOpenGLVertexShader*)(RHICreateVertexShader(VertexShaderCode.GetReadAccess()).GetReference());
+		FVertexShaderRHIRef VertexShaderRHI = RHICreateVertexShader(VertexShaderCode.GetReadAccess());
+		FOpenGLVertexShader* OGLVertexShader = (FOpenGLVertexShader*)VertexShaderRHI.GetReference();
 
 		if (IsRunningRHIInSeparateThread())
 		{
 			FRHICommandListImmediate& RHICmdList = FRHICommandListExecutor::GetImmediateCommandList();
 			RHICmdList.ImmediateFlush(EImmediateFlushType::FlushRHIThread);
 		}
-
-		if (!VerifyCompiledShader(VertexShader->Resource, TestVertexProgram, false))
+				
+		if (!VerifyCompiledShader(OGLVertexShader->Resource, TestVertexProgram, false))
 		{
 			UE_LOG(LogRHI, Warning, TEXT("Using the round() function hack, because the vertex shader for round function test failed to compile."));
 			FOpenGL::bRequiresRoundFunctionHack = true;
