Support virtual textures on mobile platforms. Requires Vulkan, OpenGL ES3.1 + SRV, iPhone8 and up.
Disabled by default, set r.Mobile.VirtualTextures=1 to enable
known issues: RVT compression is not implementted, Android OpenGL will have R and G channels swapped for uncompressed streaming VT, no sRGB support

diff --git a/Engine/Shaders/Private/BasePassPixelShader.usf b/Engine/Shaders/Private/BasePassPixelShader.usf
--- a/Engine/Shaders/Private/BasePassPixelShader.usf
+++ b/Engine/Shaders/Private/BasePassPixelShader.usf
@@ -51,9 +51,9 @@
 
 #if LIGHTMAP_VT_ENABLED || NUM_VIRTUALTEXTURE_SAMPLES
 #if COMPILER_HLSLCC
-	RWTexture2D<uint> VTFeedbackBuffer;
+	RWBuffer<uint> VTFeedbackBuffer;
 #else
-	RWTexture2D<uint> VTFeedbackBuffer : register(VT_FEEDBACK_REGISTER);
+	RWBuffer<uint> VTFeedbackBuffer : register(VT_FEEDBACK_REGISTER);
 #endif
 #endif
 
diff --git a/Engine/Shaders/Private/DeferredDecal.usf b/Engine/Shaders/Private/DeferredDecal.usf
--- a/Engine/Shaders/Private/DeferredDecal.usf
+++ b/Engine/Shaders/Private/DeferredDecal.usf
@@ -31,7 +31,7 @@ float2 DecalParams;
 #endif
 
 #if NUM_VIRTUALTEXTURE_SAMPLES
-	RWTexture2D<uint> VTFeedbackBuffer REGISTER(VT_FEEDBACK_REGISTER); //MUST match FDecalRenderTargetManager::VirtualTextureFeedbackIndex
+	RWBuffer<uint> VTFeedbackBuffer REGISTER(VT_FEEDBACK_REGISTER); //MUST match FDecalRenderTargetManager::VirtualTextureFeedbackIndex
 #endif
 
 // from component to clip space (for decal frustum)
diff --git a/Engine/Shaders/Private/MeshDecals.usf b/Engine/Shaders/Private/MeshDecals.usf
--- a/Engine/Shaders/Private/MeshDecals.usf
+++ b/Engine/Shaders/Private/MeshDecals.usf
@@ -9,7 +9,7 @@
 #include "/Engine/Generated/VertexFactory.ush"
 
 #if NUM_VIRTUALTEXTURE_SAMPLES
-	RWTexture2D<uint> VTFeedbackBuffer REGISTER(VT_FEEDBACK_REGISTER); //MUST match FDecalRenderTargetManager::VirtualTextureFeedbackIndex
+	RWBuffer<uint> VTFeedbackBuffer REGISTER(VT_FEEDBACK_REGISTER); //MUST match FDecalRenderTargetManager::VirtualTextureFeedbackIndex
 #endif
 
 // Additional guard to remove all struct usage if empty to avoid platform conflicts.
diff --git a/Engine/Shaders/Private/MobileBasePassPixelShader.usf b/Engine/Shaders/Private/MobileBasePassPixelShader.usf
--- a/Engine/Shaders/Private/MobileBasePassPixelShader.usf
+++ b/Engine/Shaders/Private/MobileBasePassPixelShader.usf
@@ -88,6 +88,14 @@ float4 SpotLightDirectionAndSpecularScale[MAX_DYNAMIC_POINT_LIGHTS];
 	#endif
 #endif
 
+#if LIGHTMAP_VT_ENABLED || NUM_VIRTUALTEXTURE_SAMPLES
+#if COMPILER_HLSLCC
+	#define VTFeedbackBuffer MobileSceneTextures.VirtualTextureFeedbackUAV
+#else
+	RWBuffer<uint> VTFeedbackBuffer : register(VT_FEEDBACK_REGISTER);
+#endif
+#endif
+
 half PhongApprox( half Roughness, half RoL )
 {
 	half a = Roughness * Roughness;			// 1 mul
@@ -375,6 +383,7 @@ void ApplyPixelDepthOffsetForMobileBasePass(inout FMaterialPixelParameters Mater
     float PixelDepthOffset = ApplyPixelDepthOffsetToMaterialParameters(MaterialParameters, PixelMaterialInputs, OutDepth);
 }
 
+EARLYDEPTHSTENCIL
 void Main( 
 	FVertexFactoryInterpolantsVSToPS Interpolants
 	, FMobileBasePassInterpolantsVSToPS BasePassInterpolants
@@ -814,4 +823,14 @@ void Main(
 		OutColor.rgb = sqrt(OutColor.rgb);
 	#endif
 #endif
+
+#if NUM_VIRTUALTEXTURE_SAMPLES || LIGHTMAP_VT_ENABLED
+	FinalizeVirtualTextureFeedback(
+		MaterialParameters.VirtualTextureFeedback,
+		MaterialParameters.SvPosition,
+		Opacity,
+		View.FrameNumber,
+		VTFeedbackBuffer
+	);
+#endif
 }
diff --git a/Engine/Shaders/Private/PageTableUpdate.usf b/Engine/Shaders/Private/PageTableUpdate.usf
--- a/Engine/Shaders/Private/PageTableUpdate.usf
+++ b/Engine/Shaders/Private/PageTableUpdate.usf
@@ -83,7 +83,24 @@ void PageTableUpdateVS(
 	OutColor = Page;
 }
 
-void PageTableUpdatePS(
+void PageTableUpdatePS_1(
+	nointerpolation uint InColor : TEXCOORD0,
+	out uint OutColor : SV_Target0
+	)
+{
+	OutColor = InColor; 
+}
+
+void PageTableUpdatePS_2(
+	nointerpolation uint InColor : TEXCOORD0,
+	out uint2 OutColor : SV_Target0
+	)
+{
+	// splat output to RGBA channels, color write mask will be configured to write to the correct channel
+	OutColor = InColor.xx;
+}
+
+void PageTableUpdatePS_4(
 	nointerpolation uint InColor : TEXCOORD0,
 	out uint4 OutColor : SV_Target0
 	)
diff --git a/Engine/Shaders/Private/RequiredTextureResolutionPixelShader.usf b/Engine/Shaders/Private/RequiredTextureResolutionPixelShader.usf
--- a/Engine/Shaders/Private/RequiredTextureResolutionPixelShader.usf
+++ b/Engine/Shaders/Private/RequiredTextureResolutionPixelShader.usf
@@ -71,7 +71,7 @@ MaterialFloat StoreAvailableVTLevel(in out FTexCoordScalesParams Params, float L
 #include "DebugViewModeCommon.ush"
 
 #if LIGHTMAP_VT_ENABLED || NUM_VIRTUALTEXTURE_SAMPLES
-	RWTexture2D<uint> VTFeedbackBuffer : register(VT_FEEDBACK_REGISTER);
+	RWBuffer<uint> VTFeedbackBuffer : register(VT_FEEDBACK_REGISTER);
 #endif
 
 #if LIGHTMAP_VT_ENABLED || NUM_VIRTUALTEXTURE_SAMPLES
diff --git a/Engine/Shaders/Private/VirtualTextureCommon.ush b/Engine/Shaders/Private/VirtualTextureCommon.ush
--- a/Engine/Shaders/Private/VirtualTextureCommon.ush
+++ b/Engine/Shaders/Private/VirtualTextureCommon.ush
@@ -3,10 +3,6 @@
 #ifndef VirtualTextureCommon_ush__
 #define VirtualTextureCommon_ush__
 
-#if (FEATURE_LEVEL != FEATURE_LEVEL_SM5)
-	#error Virtual textures require shader model 5.0 or higher!
-#endif
-
 #include "Common.ush"
 #include "Random.ush"
 
@@ -65,14 +61,16 @@ void StoreVirtualTextureFeedback(in out FVirtualTextureFeedbackParams Params, ui
 /*
 	This should be called at the end of the pixel shader to write out the gathered VT feedback info to the OutputBuffer.
 */
-void FinalizeVirtualTextureFeedback(in FVirtualTextureFeedbackParams Params, float4 SvPosition, float Opacity, uint FrameNumber, RWTexture2D<uint> OutputBuffer)
+void FinalizeVirtualTextureFeedback(in FVirtualTextureFeedbackParams Params, float4 SvPosition, float Opacity, uint FrameNumber, RWBuffer<uint> OutputBuffer)
 {
 	uint2 PixelPos = SvPosition.xy;	
+	uint FeedbackPos = 0;
 	// This code will only run every few pixels...
 	[branch] if (((PixelPos.x | PixelPos.y) & (VIRTUAL_TEXTURE_FEEDBACK_FACTOR-1)) == 0)
 	{
 		// TODO use append buffer ?
 		PixelPos /= VIRTUAL_TEXTURE_FEEDBACK_FACTOR;
+		FeedbackPos = PixelPos.y*View.VirtualTextureFeedbackStride + PixelPos.x;
 
 		// When using DBuffer, feedback buffer is blending using decal opacity for opaque objects
 		// This way we avoid overwriting values already written to feedback buffer while rendering dbuffer decals
@@ -97,10 +95,10 @@ void FinalizeVirtualTextureFeedback(in FVirtualTextureFeedbackParams Params, flo
         // that means more opaque surfaces get more pixels in the feedback buffer which is what we want
         [branch] if (Alpha > AlphaThreshold)				
         {
-            OutputBuffer[PixelPos] = Params.Request;
+            OutputBuffer[FeedbackPos] = Params.Request;
         }     
 #else	
-		OutputBuffer[PixelPos] = Params.Request;
+		OutputBuffer[FeedbackPos] = Params.Request;
 #endif
 	}	
 }
diff --git a/Engine/Shaders/Public/Platform.ush b/Engine/Shaders/Public/Platform.ush
--- a/Engine/Shaders/Public/Platform.ush
+++ b/Engine/Shaders/Public/Platform.ush
@@ -474,7 +474,7 @@
 
 #endif // SM5_PROFILE || COMPILER_SUPPORTS_ATTRIBUTES
 
-#if SM5_PROFILE || METAL_MRT_PROFILE || METAL_SM5_PROFILE || METAL_SM5_NOTESS_PROFILE
+#if SM5_PROFILE || METAL_MRT_PROFILE || METAL_SM5_PROFILE || METAL_SM5_NOTESS_PROFILE || ES3_1_PROFILE
 #define EARLYDEPTHSTENCIL [earlydepthstencil]
 #endif
 
diff --git a/Engine/Source/Developer/Android/AndroidTargetPlatform/Private/AndroidTargetPlatform.cpp b/Engine/Source/Developer/Android/AndroidTargetPlatform/Private/AndroidTargetPlatform.cpp
--- a/Engine/Source/Developer/Android/AndroidTargetPlatform/Private/AndroidTargetPlatform.cpp
+++ b/Engine/Source/Developer/Android/AndroidTargetPlatform/Private/AndroidTargetPlatform.cpp
@@ -216,6 +216,12 @@ FAndroidTargetDevicePtr FAndroidTargetPlatform::CreateTargetDevice(const ITarget
 	return MakeShareable(new FAndroidTargetDevice(InTargetPlatform, InSerialNumber, InAndroidVariant));
 }
 
+static bool UsesVirtualTextures()
+{
+	static auto* CVarMobileVirtualTextures = IConsoleManager::Get().FindTConsoleVariableDataInt(TEXT("r.Mobile.VirtualTextures"));
+	return CVarMobileVirtualTextures->GetValueOnAnyThread() != 0;
+}
+
 bool FAndroidTargetPlatform::SupportsES2() const
 {
 	// default to support ES2
@@ -354,6 +360,9 @@ bool FAndroidTargetPlatform::SupportsFeature( ETargetPlatformFeatures Feature )
 		case ETargetPlatformFeatures::SoftwareOcclusion:
 			return SupportsSoftwareOcclusion();
 
+		case ETargetPlatformFeatures::VirtualTextureStreaming:
+			return UsesVirtualTextures();
+
 		default:
 			break;
 	}
diff --git a/Engine/Source/Developer/Apple/MetalShaderFormat/Private/MetalBackend.cpp b/Engine/Source/Developer/Apple/MetalShaderFormat/Private/MetalBackend.cpp
--- a/Engine/Source/Developer/Apple/MetalShaderFormat/Private/MetalBackend.cpp
+++ b/Engine/Source/Developer/Apple/MetalShaderFormat/Private/MetalBackend.cpp
@@ -6201,6 +6201,7 @@ bool FMetalCodeBackend::GenerateMain(EHlslShaderFrequency Frequency, const char*
 	MainSig->wg_size_x = EntryPointSig->wg_size_x;
 	MainSig->wg_size_y = EntryPointSig->wg_size_y;
 	MainSig->wg_size_z = EntryPointSig->wg_size_z;
+	MainSig->is_early_depth_stencil = EntryPointSig->is_early_depth_stencil;
 	// NOTE: ParseState->tessellation has been modified since EntryPointSig->tessellation was used...
 	MainSig->tessellation = ParseState->tessellation;
 
diff --git a/Engine/Source/Developer/Apple/MetalShaderFormat/Private/MetalUtils.cpp b/Engine/Source/Developer/Apple/MetalShaderFormat/Private/MetalUtils.cpp
--- a/Engine/Source/Developer/Apple/MetalShaderFormat/Private/MetalUtils.cpp
+++ b/Engine/Source/Developer/Apple/MetalShaderFormat/Private/MetalUtils.cpp
@@ -249,6 +249,30 @@ namespace MetalUtils
 		{"SV_Target5", glsl_type::float_type, "FragColor5", ir_var_out, "[[ color(5) ]]"},
 		{"SV_Target6", glsl_type::float_type, "FragColor6", ir_var_out, "[[ color(6) ]]"},
 		{"SV_Target7", glsl_type::float_type, "FragColor7", ir_var_out, "[[ color(7) ]]"},
+		{"SV_Target0", glsl_type::uint_type, "FragColor0", ir_var_out, "[[ color(0) ]]"},
+		{"SV_Target1", glsl_type::uint_type, "FragColor1", ir_var_out, "[[ color(1) ]]"},
+		{"SV_Target2", glsl_type::uint_type, "FragColor2", ir_var_out, "[[ color(2) ]]"},
+		{"SV_Target3", glsl_type::uint_type, "FragColor3", ir_var_out, "[[ color(3) ]]"},
+		{"SV_Target4", glsl_type::uint_type, "FragColor4", ir_var_out, "[[ color(4) ]]"},
+		{"SV_Target5", glsl_type::uint_type, "FragColor5", ir_var_out, "[[ color(5) ]]"},
+		{"SV_Target6", glsl_type::uint_type, "FragColor6", ir_var_out, "[[ color(6) ]]"},
+		{"SV_Target7", glsl_type::uint_type, "FragColor7", ir_var_out, "[[ color(7) ]]"},
+		{"SV_Target0", glsl_type::uvec2_type, "FragColor0", ir_var_out, "[[ color(0) ]]"},
+		{"SV_Target1", glsl_type::uvec2_type, "FragColor1", ir_var_out, "[[ color(1) ]]"},
+		{"SV_Target2", glsl_type::uvec2_type, "FragColor2", ir_var_out, "[[ color(2) ]]"},
+		{"SV_Target3", glsl_type::uvec2_type, "FragColor3", ir_var_out, "[[ color(3) ]]"},
+		{"SV_Target4", glsl_type::uvec2_type, "FragColor4", ir_var_out, "[[ color(4) ]]"},
+		{"SV_Target5", glsl_type::uvec2_type, "FragColor5", ir_var_out, "[[ color(5) ]]"},
+		{"SV_Target6", glsl_type::uvec2_type, "FragColor6", ir_var_out, "[[ color(6) ]]"},
+		{"SV_Target7", glsl_type::uvec2_type, "FragColor7", ir_var_out, "[[ color(7) ]]"},
+		{"SV_Target0", glsl_type::uvec4_type, "FragColor0", ir_var_out, "[[ color(0) ]]"},
+		{"SV_Target1", glsl_type::uvec4_type, "FragColor1", ir_var_out, "[[ color(1) ]]"},
+		{"SV_Target2", glsl_type::uvec4_type, "FragColor2", ir_var_out, "[[ color(2) ]]"},
+		{"SV_Target3", glsl_type::uvec4_type, "FragColor3", ir_var_out, "[[ color(3) ]]"},
+		{"SV_Target4", glsl_type::uvec4_type, "FragColor4", ir_var_out, "[[ color(4) ]]"},
+		{"SV_Target5", glsl_type::uvec4_type, "FragColor5", ir_var_out, "[[ color(5) ]]"},
+		{"SV_Target6", glsl_type::uvec4_type, "FragColor6", ir_var_out, "[[ color(6) ]]"},
+		{"SV_Target7", glsl_type::uvec4_type, "FragColor7", ir_var_out, "[[ color(7) ]]"},
 		{NULL, NULL, NULL, ir_var_auto, nullptr}
 	};
 
diff --git a/Engine/Source/Developer/IOS/IOSTargetPlatform/Private/IOSTargetPlatform.cpp b/Engine/Source/Developer/IOS/IOSTargetPlatform/Private/IOSTargetPlatform.cpp
--- a/Engine/Source/Developer/IOS/IOSTargetPlatform/Private/IOSTargetPlatform.cpp
+++ b/Engine/Source/Developer/IOS/IOSTargetPlatform/Private/IOSTargetPlatform.cpp
@@ -447,6 +447,11 @@ bool FIOSTargetPlatform::HandleTicker(float DeltaTime)
 
 /* ITargetPlatform interface
  *****************************************************************************/
+static bool UsesVirtualTextures()
+{
+	static auto* CVarMobileVirtualTextures = IConsoleManager::Get().FindTConsoleVariableDataInt(TEXT("r.Mobile.VirtualTextures"));
+	return CVarMobileVirtualTextures->GetValueOnAnyThread() != 0;
+}
 
 static bool SupportsMetal()
 {
@@ -513,6 +518,9 @@ bool FIOSTargetPlatform::SupportsFeature( ETargetPlatformFeatures Feature ) cons
 		case ETargetPlatformFeatures::SoftwareOcclusion:
 			return SupportsSoftwareOcclusion();
 
+		case ETargetPlatformFeatures::VirtualTextureStreaming:
+			return UsesVirtualTextures();
+
 		default:
 			break;
 	}
diff --git a/Engine/Source/Developer/ShaderFormatOpenGL/Private/GlslBackend.cpp b/Engine/Source/Developer/ShaderFormatOpenGL/Private/GlslBackend.cpp
--- a/Engine/Source/Developer/ShaderFormatOpenGL/Private/GlslBackend.cpp
+++ b/Engine/Source/Developer/ShaderFormatOpenGL/Private/GlslBackend.cpp
@@ -851,7 +851,8 @@ class ir_gen_glsl_visitor : public ir_visitor
 		}
 		else if (type->is_integer())
 		{
-			return GLSL_PRECISION_HIGHP;
+			// integers use default precision which is always highp
+			return GLSL_PRECISION_DEFAULT;
 		}
 		return GLSL_PRECISION_DEFAULT;
 	}
@@ -3280,7 +3281,8 @@ public:
 			
 			const char* DefaultPrecision = bDefaultPrecisionIsHalf ? "mediump" : "highp";
 			ralloc_asprintf_append(&default_precision_buffer, "precision %s float;\n", DefaultPrecision);
-			ralloc_asprintf_append(&default_precision_buffer, "precision %s int;\n", DefaultPrecision);
+			// always use highp for integers as shaders use them as bit storage
+			ralloc_asprintf_append(&default_precision_buffer, "precision %s int;\n", "highp");
 
 			if (bIsES) // ES2 workarounds
 			{
@@ -3332,7 +3334,11 @@ bool compiler_internal_AdjustIsFrontFacing(bool isFrontFacing)
 #endif
 }
 )RawStrDelimiter";
-			ralloc_asprintf_append(buffer, func_clipControlAdjustments);
+			
+			if (ShaderTarget != compute_shader)
+			{
+				ralloc_asprintf_append(buffer, func_clipControlAdjustments);
+			}
 		}
 
 		// FramebufferFetchES2 'intrinsic'
@@ -3544,13 +3550,15 @@ struct FBreakPrecisionChangesVisitor : public ir_rvalue_visitor
 					case ir_unop_i2f:
 					case ir_unop_b2f:
 					case ir_unop_u2f:
-						bGenerateNewVar = bDefaultPrecisionIsHalf;
+						// integers always use highp
+						bGenerateNewVar = false;
 						break;
 
 					case ir_unop_i2h:
 					case ir_unop_b2h:
 					case ir_unop_u2h:
-						bGenerateNewVar = !bDefaultPrecisionIsHalf;
+						// integers always use highp
+						bGenerateNewVar = true;
 						break;
 
 					case ir_unop_h2f:
@@ -5242,6 +5250,11 @@ bool FGlslCodeBackend::GenerateMain(
 						);
 					break;
 				case ir_var_out:
+					if (Frequency == HSF_PixelShader && Variable->semantic && (strcmp(Variable->semantic, "SV_Depth") == 0))
+					{
+						bExplicitDepthWrites = true;
+					}
+
 					ArgVarDeref = GenShaderOutput(
 						Frequency,
 						ParseState,
@@ -5390,7 +5403,7 @@ bool FGlslCodeBackend::GenerateMain(
 		MainSig->body.push_tail(new(ParseState)ir_call(EntryPointSig, EntryPointReturn, &ArgInstructions));
 		MainSig->body.append_list(&PostCallInstructions);
 		MainSig->maxvertexcount = EntryPointSig->maxvertexcount;
-		MainSig->is_early_depth_stencil = EntryPointSig->is_early_depth_stencil;
+		MainSig->is_early_depth_stencil = (EntryPointSig->is_early_depth_stencil && !bExplicitDepthWrites);
 		MainSig->wg_size_x = EntryPointSig->wg_size_x;
 		MainSig->wg_size_y = EntryPointSig->wg_size_y;
 		MainSig->wg_size_z = EntryPointSig->wg_size_z;
diff --git a/Engine/Source/Developer/ShaderFormatOpenGL/Public/GlslBackend.h b/Engine/Source/Developer/ShaderFormatOpenGL/Public/GlslBackend.h
--- a/Engine/Source/Developer/ShaderFormatOpenGL/Public/GlslBackend.h
+++ b/Engine/Source/Developer/ShaderFormatOpenGL/Public/GlslBackend.h
@@ -57,7 +57,8 @@ struct SHADERFORMATOPENGL_API FGlslCodeBackend : public FCodeBackend
 {
 	FGlslCodeBackend(unsigned int InHlslCompileFlags, EHlslCompileTarget InTarget, bool bInIsWebGL) :
 		FCodeBackend(InHlslCompileFlags, InTarget),
-		bIsWebGL(bInIsWebGL)
+		bIsWebGL(bInIsWebGL),
+		bExplicitDepthWrites(false)
 	{
 	}
 
@@ -100,6 +101,7 @@ struct SHADERFORMATOPENGL_API FGlslCodeBackend : public FCodeBackend
 	}
 
 	bool bIsWebGL;
+	bool bExplicitDepthWrites;
 };
 
 
diff --git a/Engine/Source/Developer/TextureFormatIntelISPCTexComp/Private/TextureFormatIntelISPCTexComp.cpp b/Engine/Source/Developer/TextureFormatIntelISPCTexComp/Private/TextureFormatIntelISPCTexComp.cpp
--- a/Engine/Source/Developer/TextureFormatIntelISPCTexComp/Private/TextureFormatIntelISPCTexComp.cpp
+++ b/Engine/Source/Developer/TextureFormatIntelISPCTexComp/Private/TextureFormatIntelISPCTexComp.cpp
@@ -736,7 +736,18 @@ public:
 			bool bIsNormalMap = (BuildSettings.TextureFormatName == GTextureFormatNameASTC_NormalAG ||
 				BuildSettings.TextureFormatName == GTextureFormatNameASTC_NormalRG);
 
-			CompressedPixelFormat = GetQualityFormat( BlockWidth, BlockHeight, bIsNormalMap ? FORCED_NORMAL_MAP_COMPRESSION_SIZE_VALUE : BuildSettings.CompressionQuality );
+			if (BuildSettings.bVirtualStreamable)
+			{
+				// Always use 4x4 for streamable VT, to reduce texture format fragmentation
+				CompressedPixelFormat = PF_ASTC_4x4;
+				BlockWidth = 4;
+				BlockHeight = 4;
+			}
+			else
+			{
+				CompressedPixelFormat = GetQualityFormat( BlockWidth, BlockHeight, bIsNormalMap ? FORCED_NORMAL_MAP_COMPRESSION_SIZE_VALUE : BuildSettings.CompressionQuality );
+			}
+			check(CompressedPixelFormat == PF_ASTC_4x4 || !BuildSettings.bVirtualStreamable);
 
 			FASTCEncoderSettings EncoderSettings;
 			if (BuildSettings.TextureFormatName == GTextureFormatNameASTC_NormalAG)
@@ -744,24 +755,32 @@ public:
 				GetProfile_astc_alpha_fast(&EncoderSettings, BlockWidth, BlockHeight);
 				EncoderSettings.TextureFormatName = BuildSettings.TextureFormatName;
 				bCompressionSucceeded = true;
+				check(EncoderSettings.block_width!=0);
 			}
 			else if (BuildSettings.TextureFormatName == GTextureFormatNameASTC_NormalRG)
 			{
 				GetProfile_astc_fast(&EncoderSettings, BlockWidth, BlockHeight);
 				EncoderSettings.TextureFormatName = BuildSettings.TextureFormatName;
 				bCompressionSucceeded = true;
+				check(EncoderSettings.block_width!=0);
 			}
 			else if (bIsRGBColorASTC)
 			{
 				GetProfile_astc_fast(&EncoderSettings, BlockWidth, BlockHeight);
 				EncoderSettings.TextureFormatName = GTextureFormatNameASTC_RGB;
 				bCompressionSucceeded = true;
+				check(EncoderSettings.block_width!=0);
 			}
 			else if (bIsRGBAColorASTC)
 			{
 				GetProfile_astc_alpha_fast(&EncoderSettings, BlockWidth, BlockHeight);
 				EncoderSettings.TextureFormatName = GTextureFormatNameASTC_RGBA;
 				bCompressionSucceeded = true;
+				check(EncoderSettings.block_width!=0);
+			}
+			else
+			{
+				check(0);
 			}
 
 			if (bCompressionSucceeded)
diff --git a/Engine/Source/Developer/VulkanShaderFormat/Private/VulkanBackend.cpp b/Engine/Source/Developer/VulkanShaderFormat/Private/VulkanBackend.cpp
--- a/Engine/Source/Developer/VulkanShaderFormat/Private/VulkanBackend.cpp
+++ b/Engine/Source/Developer/VulkanShaderFormat/Private/VulkanBackend.cpp
@@ -3489,7 +3489,8 @@ public:
 
 			const char* DefaultPrecision = bDefaultPrecisionIsHalf ? "mediump" : "highp";
 			ralloc_asprintf_append(buffer, "precision %s float;\n", DefaultPrecision);
-			ralloc_asprintf_append(buffer, "precision %s int;\n", DefaultPrecision);
+			// always use highp for integers as shaders use them as bit storage
+			ralloc_asprintf_append(buffer, "precision %s int;\n", "highp");
 			//ralloc_asprintf_append(buffer, "\n#ifndef DONTEMITSAMPLERDEFAULTPRECISION\n");
 			ralloc_asprintf_append(buffer, "precision %s sampler;\n", DefaultPrecision);
 			ralloc_asprintf_append(buffer, "precision %s sampler2D;\n", DefaultPrecision);
@@ -5478,6 +5479,11 @@ bool FVulkanCodeBackend::GenerateMain(
 						);
 					break;
 				case ir_var_out:
+					if (Frequency == HSF_PixelShader && Variable->semantic && (strcmp(Variable->semantic, "SV_Depth") == 0))
+					{
+						bExplicitDepthWrites = true;
+					}
+					
 					ArgVarDeref = GenShaderOutput(
 						Frequency,
 						ParseState,
@@ -5626,7 +5632,7 @@ bool FVulkanCodeBackend::GenerateMain(
 		MainSig->body.push_tail(new(ParseState)ir_call(EntryPointSig, EntryPointReturn, &ArgInstructions));
 		MainSig->body.append_list(&PostCallInstructions);
 		MainSig->maxvertexcount = EntryPointSig->maxvertexcount;
-		MainSig->is_early_depth_stencil = EntryPointSig->is_early_depth_stencil;
+		MainSig->is_early_depth_stencil = (EntryPointSig->is_early_depth_stencil && !bExplicitDepthWrites);
 		MainSig->wg_size_x = EntryPointSig->wg_size_x;
 		MainSig->wg_size_y = EntryPointSig->wg_size_y;
 		MainSig->wg_size_z = EntryPointSig->wg_size_z;
diff --git a/Engine/Source/Developer/VulkanShaderFormat/Private/VulkanBackend.h b/Engine/Source/Developer/VulkanShaderFormat/Private/VulkanBackend.h
--- a/Engine/Source/Developer/VulkanShaderFormat/Private/VulkanBackend.h
+++ b/Engine/Source/Developer/VulkanShaderFormat/Private/VulkanBackend.h
@@ -126,7 +126,8 @@ struct FVulkanCodeBackend : public FCodeBackend
 						FVulkanBindingTable& InBindingTable,
 						EHlslCompileTarget InTarget) :
 		FCodeBackend(InHlslCompileFlags, InTarget),
-		BindingTable(InBindingTable)
+		BindingTable(InBindingTable),
+		bExplicitDepthWrites(false)
 	{
 	}
 
@@ -154,6 +155,7 @@ struct FVulkanCodeBackend : public FCodeBackend
 	ir_function_signature* FindPatchConstantFunction(exec_list* Instructions, _mesa_glsl_parse_state* ParseState);
 
 	FVulkanBindingTable& BindingTable;
+	bool bExplicitDepthWrites;
 };
 
 // Intrinsic name
diff --git a/Engine/Source/Runtime/Core/Public/Android/AndroidPlatform.h b/Engine/Source/Runtime/Core/Public/Android/AndroidPlatform.h
--- a/Engine/Source/Runtime/Core/Public/Android/AndroidPlatform.h
+++ b/Engine/Source/Runtime/Core/Public/Android/AndroidPlatform.h
@@ -60,6 +60,7 @@ typedef FAndroidTypes FPlatformTypes;
 #define PLATFORM_SUPPORTS_EARLY_MOVIE_PLAYBACK		1 // movies will start before engine is initalized
 #define PLATFORM_SUPPORTS_GEOMETRY_SHADERS			0
 #define PLATFORM_SUPPORTS_TESSELLATION_SHADERS		0
+#define PLATFORM_SUPPORTS_VIRTUAL_TEXTURE_STREAMING	1
 
 #define PLATFORM_CODE_SECTION(Name)					__attribute__((section(Name)))
 
diff --git a/Engine/Source/Runtime/Core/Public/IOS/IOSPlatform.h b/Engine/Source/Runtime/Core/Public/IOS/IOSPlatform.h
--- a/Engine/Source/Runtime/Core/Public/IOS/IOSPlatform.h
+++ b/Engine/Source/Runtime/Core/Public/IOS/IOSPlatform.h
@@ -73,6 +73,7 @@ typedef FIOSPlatformTypes FPlatformTypes;
 #define PLATFORM_NEEDS_RHIRESOURCELIST					0
 #define PLATFORM_SUPPORTS_GEOMETRY_SHADERS				0
 #define PLATFORM_SUPPORTS_TESSELLATION_SHADERS			0
+#define PLATFORM_SUPPORTS_VIRTUAL_TEXTURE_STREAMING		1
 
 #define PLATFORM_GLOBAL_LOG_CATEGORY					LogIOS
 
diff --git a/Engine/Source/Runtime/Engine/Classes/Engine/RendererSettings.h b/Engine/Source/Runtime/Engine/Classes/Engine/RendererSettings.h
--- a/Engine/Source/Runtime/Engine/Classes/Engine/RendererSettings.h
+++ b/Engine/Source/Runtime/Engine/Classes/Engine/RendererSettings.h
@@ -213,6 +213,13 @@ class ENGINE_API URendererSettings : public UDeveloperSettings
 		ToolTip="Whether to support 'Software Occlusion Culling' on mobile platforms. This will package occluder information and enable Software Occlusion Culling.",
 		ConfigRestartRequired=false))
 	uint32 bMobileAllowSoftwareOcclusionCulling:1;
+
+	UPROPERTY(config, EditAnywhere, Category = Mobile, meta = (
+		EditCondition = "bVirtualTextures",
+		ConsoleVariable = "r.Mobile.VirtualTextures", DisplayName = "Enable virtual texture support on Mobile",
+		ToolTip = "Whether to support Virtual Textures on mobile. Requires general Virtual Texturing option enabled as well. Changing this setting requires restarting the editor.",
+		ConfigRestartRequired = true))
+	uint32 bMobileVirtualTextures : 1;
 	
 	UPROPERTY(config, EditAnywhere, Category = Materials, meta = (
 		ConsoleVariable = "r.DiscardUnusedQuality", DisplayName = "Game Discards Unused Material Quality Levels",
diff --git a/Engine/Source/Runtime/Engine/Private/SceneView.cpp b/Engine/Source/Runtime/Engine/Private/SceneView.cpp
--- a/Engine/Source/Runtime/Engine/Private/SceneView.cpp
+++ b/Engine/Source/Runtime/Engine/Private/SceneView.cpp
@@ -2351,7 +2351,6 @@ void FSceneView::SetupCommonViewUniformBufferParameters(
 
 	ViewUniformShaderParameters.CameraCut = bCameraCut ? 1 : 0;
 
-	ViewUniformShaderParameters.VirtualTextureParams = FVector4(ForceInitToZero);
 	ViewUniformShaderParameters.MinRoughness = FMath::Clamp(CVarGlobalMinRoughnessOverride.GetValueOnRenderThread(), 0.02f, 1.0f);
 
 	//to tail call keep the order and number of parameters of the caller function
diff --git a/Engine/Source/Runtime/Engine/Private/Texture.cpp b/Engine/Source/Runtime/Engine/Private/Texture.cpp
--- a/Engine/Source/Runtime/Engine/Private/Texture.cpp
+++ b/Engine/Source/Runtime/Engine/Private/Texture.cpp
@@ -35,6 +35,12 @@ static TAutoConsoleVariable<int32> CVarVirtualTextures(
 	TEXT("Is virtual texture streaming enabled?"),
 	ECVF_RenderThreadSafe | ECVF_ReadOnly);
 
+static TAutoConsoleVariable<int32> CVarMobileVirtualTextures(
+	TEXT("r.Mobile.VirtualTextures"),
+	0,
+	TEXT("Whether virtual texture streaming is enabled on mobile platforms. Requires r.VirtualTextures enabled as well. \n"),
+	ECVF_RenderThreadSafe | ECVF_ReadOnly
+);
 
 DEFINE_LOG_CATEGORY(LogTexture);
 
diff --git a/Engine/Source/Runtime/Engine/Private/VT/RuntimeVirtualTexture.cpp b/Engine/Source/Runtime/Engine/Private/VT/RuntimeVirtualTexture.cpp
--- a/Engine/Source/Runtime/Engine/Private/VT/RuntimeVirtualTexture.cpp
+++ b/Engine/Source/Runtime/Engine/Private/VT/RuntimeVirtualTexture.cpp
@@ -282,6 +282,16 @@ int32 URuntimeVirtualTexture::GetLayerCount() const
 	return GetLayerCount(MaterialType);
 }
 
+static EPixelFormat PlatformCompressedRVTFormat(EPixelFormat Format)
+{
+	if (IsMobilePlatform(GMaxRHIShaderPlatform))
+	{
+		return PF_B8G8R8A8;
+	}
+	
+	return Format;
+}
+
 EPixelFormat URuntimeVirtualTexture::GetLayerFormat(int32 LayerIndex) const
 {
 	if (LayerIndex == 0)
@@ -289,11 +299,11 @@ EPixelFormat URuntimeVirtualTexture::GetLayerFormat(int32 LayerIndex) const
 		switch (MaterialType)
 		{
 		case ERuntimeVirtualTextureMaterialType::BaseColor:
-			return bCompressTextures ? PF_DXT1 : PF_B8G8R8A8;
+			return bCompressTextures ? PlatformCompressedRVTFormat(PF_DXT1) : PF_B8G8R8A8;
 		case ERuntimeVirtualTextureMaterialType::BaseColor_Normal_Specular:
 		case ERuntimeVirtualTextureMaterialType::BaseColor_Normal_Specular_YCoCg:
 		case ERuntimeVirtualTextureMaterialType::BaseColor_Normal_Specular_Mask_YCoCg:
-			return bCompressTextures ? PF_DXT5 : PF_B8G8R8A8;
+			return bCompressTextures ? PlatformCompressedRVTFormat(PF_DXT5) : PF_B8G8R8A8;
 		case ERuntimeVirtualTextureMaterialType::WorldHeight:
 			return PF_G16;
 		default:
@@ -305,10 +315,10 @@ EPixelFormat URuntimeVirtualTexture::GetLayerFormat(int32 LayerIndex) const
 		switch (MaterialType)
 		{
 		case ERuntimeVirtualTextureMaterialType::BaseColor_Normal_Specular:
-			return bCompressTextures ? PF_DXT5 : PF_B8G8R8A8;
+			return bCompressTextures ? PlatformCompressedRVTFormat(PF_DXT5) : PF_B8G8R8A8;
 		case ERuntimeVirtualTextureMaterialType::BaseColor_Normal_Specular_YCoCg:
 		case ERuntimeVirtualTextureMaterialType::BaseColor_Normal_Specular_Mask_YCoCg:
-			return bCompressTextures ? PF_BC5 : PF_B8G8R8A8;
+			return bCompressTextures ? PlatformCompressedRVTFormat(PF_BC5) : PF_B8G8R8A8;
 		default:
 			break;
 		}
@@ -318,9 +328,9 @@ EPixelFormat URuntimeVirtualTexture::GetLayerFormat(int32 LayerIndex) const
 		switch (MaterialType)
 		{
 		case ERuntimeVirtualTextureMaterialType::BaseColor_Normal_Specular_YCoCg:
-			return bCompressTextures ? PF_DXT1 : PF_B8G8R8A8;
+			return bCompressTextures ? PlatformCompressedRVTFormat(PF_DXT1) : PF_B8G8R8A8;
 		case ERuntimeVirtualTextureMaterialType::BaseColor_Normal_Specular_Mask_YCoCg:
-			return bCompressTextures ? PF_DXT5 : PF_B8G8R8A8;
+			return bCompressTextures ? PlatformCompressedRVTFormat(PF_DXT5) : PF_B8G8R8A8;
 		default:
 			break;
 		}
diff --git a/Engine/Source/Runtime/Engine/Private/VT/VirtualTextureDataBuilder.cpp b/Engine/Source/Runtime/Engine/Private/VT/VirtualTextureDataBuilder.cpp
--- a/Engine/Source/Runtime/Engine/Private/VT/VirtualTextureDataBuilder.cpp
+++ b/Engine/Source/Runtime/Engine/Private/VT/VirtualTextureDataBuilder.cpp
@@ -608,6 +608,7 @@ void FVirtualTextureDataBuilder::BuildTiles(const TArray<FVTSourceTileEntry>& Ti
 		TBSettings.bSRGB = BuildSettingsForLayer.bSRGB;
 		TBSettings.bUseLegacyGamma = BuildSettingsForLayer.bUseLegacyGamma;
 		TBSettings.MipGenSettings = TMGS_NoMipmaps;
+		TBSettings.bVirtualStreamable = true;
 
 		check(TBSettings.GetGammaSpace() == BuildSettingsForLayer.GetGammaSpace());
  
diff --git a/Engine/Source/Runtime/Engine/Public/SceneView.h b/Engine/Source/Runtime/Engine/Public/SceneView.h
--- a/Engine/Source/Runtime/Engine/Public/SceneView.h
+++ b/Engine/Source/Runtime/Engine/Public/SceneView.h
@@ -726,6 +726,7 @@ enum ETranslucencyVolumeCascade
 	VIEW_UNIFORM_BUFFER_MEMBER(FMatrix, WorldToVirtualTexture) \
 	VIEW_UNIFORM_BUFFER_MEMBER(FVector4, VirtualTextureParams) \
 	VIEW_UNIFORM_BUFFER_MEMBER_ARRAY(FVector4, XRPassthroughCameraUVs, [2]) \
+	VIEW_UNIFORM_BUFFER_MEMBER(uint32, VirtualTextureFeedbackStride) \
 	VIEW_UNIFORM_BUFFER_MEMBER(int32, FarShadowStaticMeshLODBias) \
 	VIEW_UNIFORM_BUFFER_MEMBER(float, MinRoughness) \
 	VIEW_UNIFORM_BUFFER_MEMBER(FVector4, HairRenderInfo) \
diff --git a/Engine/Source/Runtime/Landscape/Private/LandscapeRender.cpp b/Engine/Source/Runtime/Landscape/Private/LandscapeRender.cpp
--- a/Engine/Source/Runtime/Landscape/Private/LandscapeRender.cpp
+++ b/Engine/Source/Runtime/Landscape/Private/LandscapeRender.cpp
@@ -4711,7 +4711,9 @@ public:
 				}
 
 				static const FName LandscapeFixedGridVertexFactory = FName(TEXT("FLandscapeFixedGridVertexFactory"));
-				if (VertexFactoryType->GetFName() == LandscapeFixedGridVertexFactory)
+				static const FName LandscapeFixedGridVertexFactoryMobile = FName(TEXT("FLandscapeFixedGridVertexFactoryMobile"));
+				if (VertexFactoryType->GetFName() == LandscapeFixedGridVertexFactory ||
+					VertexFactoryType->GetFName() == LandscapeFixedGridVertexFactoryMobile)
 				{
 					return bIsRuntimeVirtualTextureShaderType && FMaterialResource::ShouldCache(Platform, ShaderType, VertexFactoryType);
 				}
diff --git a/Engine/Source/Runtime/Landscape/Private/LandscapeRenderMobile.cpp b/Engine/Source/Runtime/Landscape/Private/LandscapeRenderMobile.cpp
--- a/Engine/Source/Runtime/Landscape/Private/LandscapeRenderMobile.cpp
+++ b/Engine/Source/Runtime/Landscape/Private/LandscapeRenderMobile.cpp
@@ -261,6 +261,76 @@ FVertexFactoryShaderParameters* FLandscapeVertexFactoryMobile::ConstructShaderPa
 }
 
 IMPLEMENT_VERTEX_FACTORY_TYPE(FLandscapeVertexFactoryMobile, "/Engine/Private/LandscapeVertexFactory.ush", true, true, true, false, false);
+IMPLEMENT_VERTEX_FACTORY_TYPE_EX(FLandscapeFixedGridVertexFactoryMobile, "/Engine/Private/LandscapeVertexFactory.ush", true, true, true, false, false, true, false);
+
+
+/** 
+  * Shader parameters for use with FLandscapeFixedGridVertexFactory
+  * Simple grid rendering (without dynamic lod blend) needs a simpler fixed setup.
+  */
+class FLandscapeFixedGridVertexFactoryMobileVertexShaderParameters : public FLandscapeVertexFactoryMobileVertexShaderParameters
+{
+public:
+	virtual void GetElementShaderBindings(
+		const class FSceneInterface* Scene,
+		const FSceneView* InView,
+		const class FMeshMaterialShader* Shader,
+		const EVertexInputStreamType InputStreamType,
+		ERHIFeatureLevel::Type FeatureLevel,
+		const FVertexFactory* VertexFactory,
+		const FMeshBatchElement& BatchElement,
+		class FMeshDrawSingleShaderBindings& ShaderBindings,
+		FVertexInputStreamArray& VertexStreams
+	) const override
+	{
+		SCOPE_CYCLE_COUNTER(STAT_LandscapeVFDrawTimeVS);
+
+		const FLandscapeBatchElementParams* BatchElementParams = (const FLandscapeBatchElementParams*)BatchElement.UserData;
+		check(BatchElementParams);
+		const FLandscapeComponentSceneProxyMobile* SceneProxy = (const FLandscapeComponentSceneProxyMobile*)BatchElementParams->SceneProxy;
+		ShaderBindings.Add(Shader->GetUniformBufferParameter<FLandscapeUniformShaderParameters>(),*BatchElementParams->LandscapeUniformShaderParametersResource);
+
+		if (LodValuesParameter.IsBound())
+		{
+			ShaderBindings.Add(LodValuesParameter, SceneProxy->GetShaderLODValues(BatchElementParams->CurrentLOD));
+		}
+
+		if (LodBiasParameter.IsBound())
+		{
+			ShaderBindings.Add(LodBiasParameter, FVector4(ForceInitToZero));
+		}
+
+		if (ForcedLodParameter.IsBound())
+		{
+			ShaderBindings.Add(ForcedLodParameter, BatchElementParams->ForcedLOD);
+		}
+	}
+};
+
+void FLandscapeFixedGridVertexFactoryMobile::ModifyCompilationEnvironment(const FVertexFactoryType* Type, EShaderPlatform Platform, const FMaterial* Material, FShaderCompilerEnvironment& OutEnvironment)
+{
+	FLandscapeVertexFactoryMobile::ModifyCompilationEnvironment(Type, Platform, Material, OutEnvironment);
+	OutEnvironment.SetDefine(TEXT("FIXED_GRID"), TEXT("1"));
+}
+
+FVertexFactoryShaderParameters* FLandscapeFixedGridVertexFactoryMobile::ConstructShaderParameters(EShaderFrequency ShaderFrequency)
+{
+	switch (ShaderFrequency)
+	{
+	case SF_Vertex:
+		return new FLandscapeFixedGridVertexFactoryMobileVertexShaderParameters();
+	case SF_Pixel:
+		return new FLandscapeVertexFactoryMobilePixelShaderParameters();
+	}
+	return nullptr;
+}
+	
+bool FLandscapeFixedGridVertexFactoryMobile::ShouldCompilePermutation(EShaderPlatform Platform, const FMaterial* Material, const FShaderType* ShaderType)
+{
+	return GetMaxSupportedFeatureLevel(Platform) == ERHIFeatureLevel::ES3_1 &&
+		(Material->IsUsedWithLandscape() || Material->IsSpecialEngineMaterial());
+}
+
 
 /**
 * Initialize the RHI for this rendering resource
@@ -415,8 +485,25 @@ void FLandscapeComponentSceneProxyMobile::CreateRenderThreadResources()
 			GetScene().GetFeatureLevel(), false, NumOcclusionVertices);
 
 		FLandscapeComponentSceneProxy::SharedBuffersMap.Add(SharedBuffersKey, SharedBuffers);
+				
+		if (UseVirtualTexturing(FeatureLevel))
+		{
+			//todo[vt]: We will need a version of this to support XYOffsetmapTexture
+			FLandscapeFixedGridVertexFactoryMobile* LandscapeVertexFactory = new FLandscapeFixedGridVertexFactoryMobile(FeatureLevel);
+			LandscapeVertexFactory->MobileData.PositionComponent = FVertexStreamComponent(MobileRenderData->VertexBuffer, STRUCT_OFFSET(FLandscapeMobileVertex,Position), sizeof(FLandscapeMobileVertex), VET_UByte4N);
+			for( uint32 Index = 0; Index < LANDSCAPE_MAX_ES_LOD_COMP; ++Index )
+			{
+				LandscapeVertexFactory->MobileData.LODHeightsComponent.Add
+				(FVertexStreamComponent(MobileRenderData->VertexBuffer, STRUCT_OFFSET(FLandscapeMobileVertex,LODHeights) + sizeof(uint8) * 4 * Index, sizeof(FLandscapeMobileVertex), VET_UByte4N));
+			}
+			LandscapeVertexFactory->InitResource();
+			SharedBuffers->FixedGridVertexFactory = LandscapeVertexFactory;
+		}
 	}
 
+	//
+	FixedGridVertexFactory = SharedBuffers->FixedGridVertexFactory;
+	
 	SharedBuffers->AddRef();
 
 	// Init vertex buffer
diff --git a/Engine/Source/Runtime/Landscape/Public/LandscapeRender.h b/Engine/Source/Runtime/Landscape/Public/LandscapeRender.h
--- a/Engine/Source/Runtime/Landscape/Public/LandscapeRender.h
+++ b/Engine/Source/Runtime/Landscape/Public/LandscapeRender.h
@@ -998,7 +998,7 @@ protected:
 	
 	FORCEINLINE void GetShaderCurrentNeighborLOD(const FSceneView& InView, float InBatchElementCurrentLOD, int8 InSubSectionX, int8 InSubSectionY, int8 InCurrentSubSectionIndex, FVector4& OutShaderCurrentNeighborLOD) const;
 	FORCEINLINE FVector4 GetShaderLODBias() const;
-	FORCEINLINE FVector4 GetShaderLODValues(int8 BatchElementCurrentLOD) const;
+	FVector4 GetShaderLODValues(int8 BatchElementCurrentLOD) const;
 
 	bool GetMeshElement(bool UseSeperateBatchForShadow, bool ShadowOnly, bool HasTessellation, int8 InLODIndex, UMaterialInterface* InMaterialInterface, FMeshBatch& OutMeshBatch, TArray<FLandscapeBatchElementParams>& OutStaticBatchParamArray) const;
 	bool GetMeshElementForVirtualTexture(int32 InLodIndex, ERuntimeVirtualTextureMaterialType MaterialType, UMaterialInterface* InMaterialInterface, FMeshBatch& OutMeshBatch, TArray<FLandscapeBatchElementParams>& OutStaticBatchParamArray) const;
@@ -1037,6 +1037,7 @@ public:
 	friend class FLandscapeVertexFactoryMobileVertexShaderParameters;
 	friend class FLandscapeVertexFactoryMobilePixelShaderParameters;
 	friend class FLandscapeFixedGridVertexFactoryVertexShaderParameters;
+	friend class FLandscapeFixedGridVertexFactoryMobileVertexShaderParameters;
 
 	// FLandscapeComponentSceneProxy interface.
 	uint64 GetStaticBatchElementVisibility(const FSceneView& InView, const FMeshBatch* InBatch, const void* InViewCustomData) const;
diff --git a/Engine/Source/Runtime/Landscape/Public/LandscapeRenderMobile.h b/Engine/Source/Runtime/Landscape/Public/LandscapeRenderMobile.h
--- a/Engine/Source/Runtime/Landscape/Public/LandscapeRenderMobile.h
+++ b/Engine/Source/Runtime/Landscape/Public/LandscapeRenderMobile.h
@@ -134,3 +134,19 @@ public:
 
 	friend class FLandscapeVertexBufferMobile;
 };
+
+
+class FLandscapeFixedGridVertexFactoryMobile : public FLandscapeVertexFactoryMobile
+{
+	DECLARE_VERTEX_FACTORY_TYPE(FLandscapeFixedGridVertexFactoryMobile);
+
+public:
+	FLandscapeFixedGridVertexFactoryMobile(ERHIFeatureLevel::Type InFeatureLevel)
+		: FLandscapeVertexFactoryMobile(InFeatureLevel)
+	{
+	}
+
+	static void ModifyCompilationEnvironment(const FVertexFactoryType* Type, EShaderPlatform Platform, const FMaterial* Material, FShaderCompilerEnvironment& OutEnvironment);
+	static FVertexFactoryShaderParameters* ConstructShaderParameters(EShaderFrequency ShaderFrequency);
+	static bool ShouldCompilePermutation(EShaderPlatform Platform, const FMaterial* Material, const FShaderType* ShaderType);
+};
diff --git a/Engine/Source/Runtime/OpenGLDrv/Private/Android/AndroidOpenGL.cpp b/Engine/Source/Runtime/OpenGLDrv/Private/Android/AndroidOpenGL.cpp
--- a/Engine/Source/Runtime/OpenGLDrv/Private/Android/AndroidOpenGL.cpp
+++ b/Engine/Source/Runtime/OpenGLDrv/Private/Android/AndroidOpenGL.cpp
@@ -67,6 +67,7 @@ PFNGLDRAWARRAYSINSTANCEDPROC			glDrawArraysInstanced = NULL;
 PFNGLGENVERTEXARRAYSPROC 				glGenVertexArrays = NULL;
 PFNGLBINDVERTEXARRAYPROC 				glBindVertexArray = NULL;
 PFNGLMAPBUFFERRANGEPROC					glMapBufferRange = NULL;
+PFNGLUNMAPBUFFERPROC					glUnmapBuffer = NULL;
 PFNGLCOPYBUFFERSUBDATAPROC				glCopyBufferSubData = NULL;
 PFNGLDRAWARRAYSINDIRECTPROC				glDrawArraysIndirect = NULL;
 PFNGLDRAWELEMENTSINDIRECTPROC			glDrawElementsIndirect = NULL;
@@ -83,6 +84,7 @@ PFNGLCLEARBUFFERFIPROC					glClearBufferfi = NULL;
 PFNGLCLEARBUFFERFVPROC					glClearBufferfv = NULL;
 PFNGLCLEARBUFFERIVPROC					glClearBufferiv = NULL;
 PFNGLCLEARBUFFERUIVPROC					glClearBufferuiv = NULL;
+PFNGLREADBUFFERPROC						glReadBuffer = NULL;
 PFNGLDRAWBUFFERSPROC					glDrawBuffers = NULL;
 PFNGLTEXBUFFEREXTPROC					glTexBufferEXT = NULL;
 PFNGLCOPYIMAGESUBDATAPROC				glCopyImageSubData = nullptr;
@@ -108,6 +110,10 @@ PFNGLDISPATCHCOMPUTEPROC				glDispatchCompute = NULL;
 PFNGLDISPATCHCOMPUTEINDIRECTPROC		glDispatchComputeIndirect = NULL;
 PFNGLBINDIMAGETEXTUREPROC				glBindImageTexture = NULL;
 
+PFNGLDELETESYNCPROC						glDeleteSync = NULL;
+PFNGLFENCESYNCPROC						glFenceSync = NULL;
+PFNGLISSYNCPROC							glIsSync = NULL;
+PFNGLCLIENTWAITSYNCPROC					glClientWaitSync = NULL;
 
 PFNGLFRAMEBUFFERTEXTUREMULTIVIEWOVRPROC glFramebufferTextureMultiviewOVR = NULL;
 PFNGLFRAMEBUFFERTEXTUREMULTISAMPLEMULTIVIEWOVRPROC glFramebufferTextureMultisampleMultiviewOVR = NULL;
@@ -128,6 +134,11 @@ int32 FAndroidOpenGL::GLMinorVersion = 0;
 GLint FAndroidOpenGL::MaxComputeTextureImageUnits = -1;
 GLint FAndroidOpenGL::MaxComputeUniformComponents = -1;
 
+GLint FAndroidOpenGL::MaxComputeUAVUnits = -1;
+GLint FAndroidOpenGL::MaxPixelUAVUnits = -1;
+GLint FAndroidOpenGL::MaxCombinedUAVUnits = 0;
+
+
 static TAutoConsoleVariable<int32> CVarEnableAdrenoTilingHint(
 	TEXT("r.Android.EnableAdrenoTilingHint"),
 	1,
@@ -999,6 +1010,15 @@ void FAndroidOpenGL::ProcessExtensions(const FString& ExtensionsString)
 	{
 		GET_GL_INT(GL_MAX_COMPUTE_TEXTURE_IMAGE_UNITS, 0, MaxComputeTextureImageUnits);
 		GET_GL_INT(GL_MAX_COMPUTE_UNIFORM_COMPONENTS, 0, MaxComputeUniformComponents);
+		
+		LOG_AND_GET_GL_INT(GL_MAX_COMBINED_IMAGE_UNIFORMS, 0, MaxCombinedUAVUnits);
+		LOG_AND_GET_GL_INT(GL_MAX_COMPUTE_IMAGE_UNIFORMS, 0, MaxComputeUAVUnits);
+		LOG_AND_GET_GL_INT(GL_MAX_FRAGMENT_IMAGE_UNIFORMS, 0, MaxPixelUAVUnits);
+
+		// clamp UAV units to a sensible limit
+		MaxCombinedUAVUnits = FMath::Min(MaxCombinedUAVUnits, 8);
+		MaxComputeUAVUnits = FMath::Min(MaxComputeUAVUnits, MaxCombinedUAVUnits);
+		MaxPixelUAVUnits = FMath::Min(MaxPixelUAVUnits, MaxCombinedUAVUnits);
 	}
 
 	bSupportsETC2 = bES30Support;
@@ -1124,8 +1144,6 @@ void FAndroidOpenGL::ProcessExtensions(const FString& ExtensionsString)
 
 	if (bES30Support)
 	{
-		glMapBufferRange = (PFNGLMAPBUFFERRANGEPROC)((void*)eglGetProcAddress("glMapBufferRange"));
-		glCopyBufferSubData = (PFNGLCOPYBUFFERSUBDATAPROC)((void*)eglGetProcAddress("glCopyBufferSubData"));
 		glDrawElementsInstanced = (PFNGLDRAWELEMENTSINSTANCEDPROC)((void*)eglGetProcAddress("glDrawElementsInstanced"));
 		glDrawArraysInstanced = (PFNGLDRAWARRAYSINSTANCEDPROC)((void*)eglGetProcAddress("glDrawArraysInstanced"));
 		glVertexAttribDivisor = (PFNGLVERTEXATTRIBDIVISORPROC)((void*)eglGetProcAddress("glVertexAttribDivisor"));
@@ -1140,8 +1158,11 @@ void FAndroidOpenGL::ProcessExtensions(const FString& ExtensionsString)
 		glClearBufferiv = (PFNGLCLEARBUFFERIVPROC)((void*)eglGetProcAddress("glClearBufferiv"));
 		glClearBufferuiv = (PFNGLCLEARBUFFERUIVPROC)((void*)eglGetProcAddress("glClearBufferuiv"));
 		glDrawBuffers = (PFNGLDRAWBUFFERSPROC)((void*)eglGetProcAddress("glDrawBuffers"));
+		glReadBuffer = (PFNGLREADBUFFERPROC)((void*)eglGetProcAddress("glReadBuffer"));
 
-
+		glMapBufferRange = (PFNGLMAPBUFFERRANGEPROC)((void*)eglGetProcAddress("glMapBufferRange"));
+		glCopyBufferSubData = (PFNGLCOPYBUFFERSUBDATAPROC)((void*)eglGetProcAddress("glCopyBufferSubData"));
+		glUnmapBuffer = (PFNGLUNMAPBUFFERPROC)((void*)eglGetProcAddress("glUnmapBuffer"));
 		glBindBufferRange = (PFNGLBINDBUFFERRANGEPROC)((void*)eglGetProcAddress("glBindBufferRange"));
 		glBindBufferBase = (PFNGLBINDBUFFERBASEPROC)((void*)eglGetProcAddress("glBindBufferBase"));
 		glGetUniformBlockIndex = (PFNGLGETUNIFORMBLOCKINDEXPROC)((void*)eglGetProcAddress("glGetUniformBlockIndex"));
@@ -1157,6 +1178,11 @@ void FAndroidOpenGL::ProcessExtensions(const FString& ExtensionsString)
 
 		glTexStorage3D = (PFNGLTEXSTORAGE3DPROC)((void*)eglGetProcAddress("glTexStorage3D"));
 		
+		glDeleteSync = (PFNGLDELETESYNCPROC)((void*)eglGetProcAddress("glDeleteSync"));
+		glFenceSync = (PFNGLFENCESYNCPROC)((void*)eglGetProcAddress("glFenceSync"));
+		glIsSync = (PFNGLISSYNCPROC)((void*)eglGetProcAddress("glIsSync"));
+		glClientWaitSync = (PFNGLCLIENTWAITSYNCPROC)((void*)eglGetProcAddress("glClientWaitSync"));
+
 		// Required by the ES3 spec
 		bSupportsInstancing = true;
 		bSupportsTextureFloat = true;
@@ -1233,9 +1259,6 @@ void FAndroidOpenGL::ProcessExtensions(const FString& ExtensionsString)
 		glBlitFramebufferNV = (PFNBLITFRAMEBUFFERNVPROC)((void*)eglGetProcAddress("glBlitFramebufferNV"));
 	}
 
-	glMapBufferOESa = (PFNGLMAPBUFFEROESPROC)((void*)eglGetProcAddress("glMapBufferOES"));
-	glUnmapBufferOESa = (PFNGLUNMAPBUFFEROESPROC)((void*)eglGetProcAddress("glUnmapBufferOES"));
-
 	//On Android, there are problems compiling shaders with textureCubeLodEXT calls in the glsl code,
 	// so we set this to false to modify the glsl manually at compile-time.
 	bSupportsTextureCubeLodEXT = false;
@@ -1247,23 +1270,8 @@ void FAndroidOpenGL::ProcessExtensions(const FString& ExtensionsString)
 		bSupportsInstancing = false;
 	}
 
-	if (bSupportsBGRA8888)
-	{
-		// Check whether device supports BGRA as color attachment
-		GLuint FrameBuffer;
-		glGenFramebuffers(1, &FrameBuffer);
-		glBindFramebuffer(GL_FRAMEBUFFER, FrameBuffer);
-		GLuint BGRA8888Texture;
-		glGenTextures(1, &BGRA8888Texture);
-		glBindTexture(GL_TEXTURE_2D, BGRA8888Texture);
-		glTexImage2D(GL_TEXTURE_2D, 0, GL_BGRA_EXT, 256, 256, 0, GL_BGRA_EXT, GL_UNSIGNED_BYTE, NULL);
-		glFramebufferTexture2D(GL_FRAMEBUFFER, GL_COLOR_ATTACHMENT0, GL_TEXTURE_2D, BGRA8888Texture, 0);
-
-		bSupportsBGRA8888RenderTarget = (glCheckFramebufferStatus(GL_FRAMEBUFFER) == GL_FRAMEBUFFER_COMPLETE);
-
-		glDeleteTextures(1, &BGRA8888Texture);
-		glDeleteFramebuffers(1, &FrameBuffer);
-	}
+	// disable swizzled render targets on Android
+	bSupportsBGRA8888RenderTarget = false;
 
 	if (IsES31Usable())
 	{
diff --git a/Engine/Source/Runtime/OpenGLDrv/Private/Android/AndroidOpenGL.h b/Engine/Source/Runtime/OpenGLDrv/Private/Android/AndroidOpenGL.h
--- a/Engine/Source/Runtime/OpenGLDrv/Private/Android/AndroidOpenGL.h
+++ b/Engine/Source/Runtime/OpenGLDrv/Private/Android/AndroidOpenGL.h
@@ -18,7 +18,7 @@
 #include <EGL/eglplatform.h>
 #include <GLES2/gl2ext.h>
 	
-typedef EGLSyncKHR UGLsync;
+typedef GLsync			UGLsync;
 #define GLdouble		GLfloat
 typedef khronos_int64_t GLint64;
 typedef khronos_uint64_t GLuint64;
@@ -36,8 +36,6 @@ typedef khronos_uint64_t GLuint64;
 
 #define GL_BGRA			GL_BGRA_EXT 
 #define GL_UNSIGNED_INT_8_8_8_8_REV	GL_UNSIGNED_BYTE
-#define glMapBuffer		glMapBufferOESa
-#define glUnmapBuffer	glUnmapBufferOESa
 
 #ifndef GL_HALF_FLOAT
 #define GL_HALF_FLOAT	GL_HALF_FLOAT_OES
@@ -147,6 +145,7 @@ extern PFNGLVERTEXATTRIBDIVISORPROC		glVertexAttribDivisor;
 extern PFNGLGENVERTEXARRAYSPROC 		glGenVertexArrays;
 extern PFNGLBINDVERTEXARRAYPROC 		glBindVertexArray;
 extern PFNGLMAPBUFFERRANGEPROC			glMapBufferRange;
+extern PFNGLUNMAPBUFFERPROC				glUnmapBuffer;
 extern PFNGLCOPYBUFFERSUBDATAPROC		glCopyBufferSubData;
 extern PFNGLDRAWARRAYSINDIRECTPROC		glDrawArraysIndirect;
 extern PFNGLDRAWELEMENTSINDIRECTPROC	glDrawElementsIndirect;
@@ -157,6 +156,7 @@ extern PFNGLCLEARBUFFERFIPROC			glClearBufferfi;
 extern PFNGLCLEARBUFFERFVPROC			glClearBufferfv;
 extern PFNGLCLEARBUFFERIVPROC			glClearBufferiv;
 extern PFNGLCLEARBUFFERUIVPROC			glClearBufferuiv;
+extern PFNGLREADBUFFERPROC				glReadBuffer;
 extern PFNGLDRAWBUFFERSPROC				glDrawBuffers;
 extern PFNGLTEXIMAGE3DPROC				glTexImage3D;
 extern PFNGLTEXSUBIMAGE3DPROC			glTexSubImage3D;
@@ -190,6 +190,10 @@ extern PFNGLDISPATCHCOMPUTEPROC			glDispatchCompute;
 extern PFNGLDISPATCHCOMPUTEINDIRECTPROC	glDispatchComputeIndirect;
 extern PFNGLBINDIMAGETEXTUREPROC		glBindImageTexture;
 
+extern PFNGLDELETESYNCPROC				glDeleteSync;
+extern PFNGLFENCESYNCPROC				glFenceSync;
+extern PFNGLISSYNCPROC					glIsSync;
+extern PFNGLCLIENTWAITSYNCPROC			glClientWaitSync;
 
 #include "OpenGLES2.h"
 
@@ -292,52 +296,38 @@ struct FAndroidOpenGL : public FOpenGLES2
 	{
 		if (GUseThreadedRendering)
 		{
-			//handle error here
-			EGLBoolean Result = eglDestroySyncKHR_p( AndroidEGL::GetInstance()->GetDisplay(), Sync );
-			if(Result == EGL_FALSE)
-			{
-				//handle error here
-			}
+			glDeleteSync( Sync );
 		}
 	}
 
 	static FORCEINLINE UGLsync FenceSync(GLenum Condition, GLbitfield Flags)
 	{
-		check(Condition == GL_SYNC_GPU_COMMANDS_COMPLETE && Flags == 0);
-		return GUseThreadedRendering ? eglCreateSyncKHR_p( AndroidEGL::GetInstance()->GetDisplay(), EGL_SYNC_FENCE_KHR, NULL ) : 0;
+		return GUseThreadedRendering ? glFenceSync( Condition, Flags ) : 0;
 	}
-	
+
 	static FORCEINLINE bool IsSync(UGLsync Sync)
 	{
-		if(GUseThreadedRendering)
-		{
-			return (Sync != EGL_NO_SYNC_KHR) ? true : false;
-		}
-		else
+		if (GUseThreadedRendering)
 		{
-			return true;
+			return (glIsSync( Sync ) == GL_TRUE) ? true : false;
 		}
+		return true;
 	}
 
 	static FORCEINLINE EFenceResult ClientWaitSync(UGLsync Sync, GLbitfield Flags, GLuint64 Timeout)
 	{
 		if (GUseThreadedRendering)
 		{
-			QUICK_SCOPE_CYCLE_COUNTER(STAT_eglClientWaitSyncKHR_p);
-			// check( Flags == GL_SYNC_FLUSH_COMMANDS_BIT );
-			GLenum Result = eglClientWaitSyncKHR_p( AndroidEGL::GetInstance()->GetDisplay(), Sync, EGL_SYNC_FLUSH_COMMANDS_BIT_KHR, Timeout );
+			GLenum Result = glClientWaitSync( Sync, Flags, Timeout );
 			switch (Result)
 			{
-			case EGL_TIMEOUT_EXPIRED_KHR:		return FR_TimeoutExpired;
-			case EGL_CONDITION_SATISFIED_KHR:	return FR_ConditionSatisfied;
+				case GL_ALREADY_SIGNALED:		return FR_AlreadySignaled;
+				case GL_TIMEOUT_EXPIRED:		return FR_TimeoutExpired;
+				case GL_CONDITION_SATISFIED:	return FR_ConditionSatisfied;
 			}
 			return FR_WaitFailed;
 		}
-		else
-		{
-			return FR_ConditionSatisfied;
-		}
-		return FR_WaitFailed;
+		return FR_ConditionSatisfied;
 	}
 
 	static FORCEINLINE void FramebufferTexture2D(GLenum Target, GLenum Attachment, GLenum TexTarget, GLuint Texture, GLint Level)
@@ -554,6 +544,11 @@ struct FAndroidOpenGL : public FOpenGLES2
 		glDrawBuffers(NumBuffers, Buffers);
 	}
 	
+	static FORCEINLINE void ReadBuffer(GLenum Mode)
+	{
+		glReadBuffer( Mode );
+	}
+
 	static FORCEINLINE void ColorMaskIndexed(GLuint Index, GLboolean Red, GLboolean Green, GLboolean Blue, GLboolean Alpha)
 	{
 		check(Index == 0 || SupportsMultipleRenderTargets());
@@ -677,6 +672,8 @@ struct FAndroidOpenGL : public FOpenGLES2
 		glBindImageTexture(Unit, Texture, Level, Layered, Layer, Access, Format);
 	}
 	
+	static FORCEINLINE bool SupportsPixelBuffers() { return false; }
+	
 	// Adreno doesn't support HALF_FLOAT
 	static FORCEINLINE int32 GetReadHalfFloatPixelsEnum()				{ return GL_FLOAT; }
 
@@ -740,7 +737,12 @@ struct FAndroidOpenGL : public FOpenGLES2
 
 	static FORCEINLINE GLint GetMaxComputeTextureImageUnits() { check(MaxComputeTextureImageUnits != -1); return MaxComputeTextureImageUnits; }
 	static FORCEINLINE GLint GetMaxComputeUniformComponents() { check(MaxComputeUniformComponents != -1); return MaxComputeUniformComponents; }
-
+	static FORCEINLINE GLint GetFirstComputeUAVUnit()			{ return 0; }
+	static FORCEINLINE GLint GetMaxComputeUAVUnits()			{ check(MaxComputeUAVUnits != -1); return MaxComputeUAVUnits; }
+	static FORCEINLINE GLint GetFirstPixelUAVUnit()				{ return 0; }
+	static FORCEINLINE GLint GetMaxPixelUAVUnits()				{ check(MaxPixelUAVUnits != -1); return MaxPixelUAVUnits; }
+	static FORCEINLINE GLint GetMaxCombinedUAVUnits()			{ return MaxCombinedUAVUnits; }
+		
 	static FORCEINLINE void FrameBufferFetchBarrier()
 	{
 		if (glFramebufferFetchBarrierQCOM)
@@ -808,6 +810,10 @@ struct FAndroidOpenGL : public FOpenGLES2
 
 	static GLint MaxComputeTextureImageUnits;
 	static GLint MaxComputeUniformComponents;
+
+	static GLint MaxCombinedUAVUnits;
+	static GLint MaxComputeUAVUnits;
+	static GLint MaxPixelUAVUnits;
 };
 
 typedef FAndroidOpenGL FOpenGL;
diff --git a/Engine/Source/Runtime/OpenGLDrv/Private/OpenGL4.cpp b/Engine/Source/Runtime/OpenGLDrv/Private/OpenGL4.cpp
--- a/Engine/Source/Runtime/OpenGLDrv/Private/OpenGL4.cpp
+++ b/Engine/Source/Runtime/OpenGLDrv/Private/OpenGL4.cpp
@@ -12,6 +12,10 @@
 
 GLint FOpenGL4::MaxComputeTextureImageUnits = -1;
 GLint FOpenGL4::MaxComputeUniformComponents = -1;
+GLint FOpenGL4::MaxCombinedUAVUnits = 0;
+GLint FOpenGL4::MaxComputeUAVUnits = -1;
+GLint FOpenGL4::MaxPixelUAVUnits = -1;
+
 bool FOpenGL4::bSupportsComputeShaders = true;
 bool FOpenGL4::bSupportsGPUMemoryInfo = false;
 bool FOpenGL4::bSupportsVertexAttribBinding = true;
@@ -19,12 +23,20 @@ bool FOpenGL4::bSupportsTextureView = true;
 
 void FOpenGL4::ProcessQueryGLInt()
 {
+	GET_GL_INT(GL_MAX_COMBINED_IMAGE_UNIFORMS, 0, MaxCombinedUAVUnits);
+	GET_GL_INT(GL_MAX_FRAGMENT_IMAGE_UNIFORMS, 0, MaxPixelUAVUnits);
+	
 	if (bSupportsComputeShaders)
 	{
 		GET_GL_INT(GL_MAX_COMPUTE_TEXTURE_IMAGE_UNITS, 0, MaxComputeTextureImageUnits);
 		GET_GL_INT(GL_MAX_COMPUTE_UNIFORM_COMPONENTS, 0, MaxComputeUniformComponents);
+		GET_GL_INT(GL_MAX_COMPUTE_IMAGE_UNIFORMS, 0, MaxComputeUAVUnits);
 	}
-
+	
+	// clamp UAV units to a sensible limit
+	MaxCombinedUAVUnits = FMath::Min(MaxCombinedUAVUnits, 8);
+	MaxComputeUAVUnits = FMath::Min(MaxComputeUAVUnits, MaxCombinedUAVUnits);
+	MaxPixelUAVUnits = FMath::Min(MaxPixelUAVUnits, MaxCombinedUAVUnits);
 }
 
 void FOpenGL4::ProcessExtensions( const FString& ExtensionsString )
diff --git a/Engine/Source/Runtime/OpenGLDrv/Private/OpenGLCommands.cpp b/Engine/Source/Runtime/OpenGLDrv/Private/OpenGLCommands.cpp
--- a/Engine/Source/Runtime/OpenGLDrv/Private/OpenGLCommands.cpp
+++ b/Engine/Source/Runtime/OpenGLDrv/Private/OpenGLCommands.cpp
@@ -450,11 +450,11 @@ void FOpenGLDynamicRHI::RHISetUAVParameter(FRHIComputeShader* ComputeShaderRHI,u
 		FOpenGLUnorderedAccessView* UnorderedAccessView = ResourceCast(UnorderedAccessViewRHI);
 		bool bLayered = UnorderedAccessView->IsLayered();
 		GLint Layer = UnorderedAccessView->GetLayer();
-		InternalSetShaderUAV(FOpenGL::GetFirstComputeUAVUnit() + UAVIndex, UnorderedAccessView->Format, UnorderedAccessView->Resource, bLayered, Layer);
+		InternalSetShaderUAV(FOpenGL::GetFirstComputeUAVUnit() + UAVIndex, UnorderedAccessView->Format, UnorderedAccessView->Resource, bLayered, Layer, GL_READ_WRITE);
 	}
 	else
 	{
-		InternalSetShaderUAV(FOpenGL::GetFirstComputeUAVUnit() + UAVIndex, GL_R32F, 0, false, 0);
+		InternalSetShaderUAV(FOpenGL::GetFirstComputeUAVUnit() + UAVIndex, GL_R32F, 0, false, 0, GL_READ_WRITE);
 	}
 }
 
@@ -466,7 +466,6 @@ void FOpenGLDynamicRHI::RHISetUAVParameter(FRHIComputeShader* ComputeShaderRHI,u
 
 void FOpenGLDynamicRHI::InternalSetShaderTexture(FOpenGLTextureBase* Texture, FOpenGLShaderResourceView* SRV, GLint TextureIndex, GLenum Target, GLuint Resource, int NumMips, int LimitMip)
 {
-	VERIFY_GL_SCOPE();
 	auto& PendingTextureState = PendingState.Textures[TextureIndex];
 	PendingTextureState.Texture = Texture;
 	PendingTextureState.SRV = SRV;
@@ -479,7 +478,6 @@ void FOpenGLDynamicRHI::InternalSetShaderTexture(FOpenGLTextureBase* Texture, FO
 
 void FOpenGLDynamicRHI::InternalSetSamplerStates(GLint TextureIndex, FOpenGLSamplerState* SamplerState)
 {
-	VERIFY_GL_SCOPE();
 	PendingState.SamplerStates[TextureIndex] = SamplerState;
 }
 
@@ -747,52 +745,92 @@ void FOpenGLDynamicRHI::SetupTexturesForDraw( FOpenGLContextState& ContextState
 	SetupTexturesForDraw(ContextState, PendingState.BoundShaderState, FOpenGL::GetMaxCombinedTextureImageUnits());
 }
 
-void FOpenGLDynamicRHI::InternalSetShaderUAV(GLint UAVIndex, GLenum Format, GLuint Resource, bool bLayered, GLint Layer)
+void FOpenGLDynamicRHI::InternalSetShaderUAV(GLint UAVIndex, GLenum Format, GLuint Resource, bool bLayered, GLint Layer, GLenum Access)
 {
 	VERIFY_GL_SCOPE();
 	PendingState.UAVs[UAVIndex].Format = Format;
 	PendingState.UAVs[UAVIndex].Resource = Resource;
 	PendingState.UAVs[UAVIndex].Layer = Layer;
 	PendingState.UAVs[UAVIndex].bLayered = bLayered;
+	PendingState.UAVs[UAVIndex].Access = Access;
 }
 
-void FOpenGLDynamicRHI::SetupUAVsForDraw( FOpenGLContextState& ContextState, const FOpenGLComputeShader* ComputeShader, int32 MaxUAVsNeeded )
+void FOpenGLDynamicRHI::SetupUAVsForDraw(FOpenGLContextState& ContextState)
 {
-	check(FOpenGL::SupportsComputeShaders());
+	int32 MaxUAVUnitUsed = 0;
+	const TBitArray<>& NeededBits = PendingState.BoundShaderState->GetUAVNeeds(MaxUAVUnitUsed);
+	SetupUAVsForProgram(ContextState, NeededBits, MaxUAVUnitUsed, FOpenGL::GetMaxPixelUAVUnits());
+}
+
+void FOpenGLDynamicRHI::SetupUAVsForCompute(FOpenGLContextState& ContextState, const FOpenGLComputeShader* ComputeShader)
+{
+	int32 MaxUAVUnitUsed = 0;
+	const TBitArray<>& NeededBits = ComputeShader->GetUAVNeeds(MaxUAVUnitUsed);
+	SetupUAVsForProgram(ContextState, NeededBits, MaxUAVUnitUsed, FOpenGL::GetMaxComputeUAVUnits());
+}
+
+void FOpenGLDynamicRHI::SetupUAVsForProgram(FOpenGLContextState& ContextState, const TBitArray<>& NeededBits, int32 MaxUAVUnitUsed, int32 MaxUAVUnits)
+{
+	if (MaxUAVUnitUsed < 0 && ContextState.ActiveUAVMask == 0)
+	{
+		// Quit early if program does not use UAVs and context has no active UAV units
+		return;
+	}
 
-	for( int32 UAVStageIndex = 0; UAVStageIndex < MaxUAVsNeeded; ++UAVStageIndex )
+	for (int32 UAVStageIndex = 0; UAVStageIndex <= MaxUAVUnitUsed; ++UAVStageIndex)
 	{
-		if (!ComputeShader->NeedsUAVStage(UAVStageIndex))
+		if (!NeededBits[UAVStageIndex])
 		{
-			CachedSetupUAVStage(ContextState, UAVStageIndex, GL_R32F, 0, false, 0);
+			CachedSetupUAVStage(ContextState, UAVStageIndex, GL_R32F, 0, false, 0, GL_READ_WRITE);
 		}
 		else
 		{
 			const FUAVStage& UAVStage = PendingState.UAVs[UAVStageIndex];
-			CachedSetupUAVStage(ContextState, UAVStageIndex, UAVStage.Format, UAVStage.Resource, UAVStage.bLayered, UAVStage.Layer);
+			CachedSetupUAVStage(ContextState, UAVStageIndex, UAVStage.Format, UAVStage.Resource, UAVStage.bLayered, UAVStage.Layer, UAVStage.Access);
 		}
 	}
-	
+
+	// clear rest of the units
+	for (int32 UAVStageIndex = (MaxUAVUnitUsed+1); UAVStageIndex < MaxUAVUnits; ++UAVStageIndex)
+	{
+		CachedSetupUAVStage(ContextState, UAVStageIndex, GL_R32F, 0, false, 0, GL_READ_WRITE);
+	}
 }
 
 
-void FOpenGLDynamicRHI::CachedSetupUAVStage( FOpenGLContextState& ContextState, GLint UAVIndex, GLenum Format, GLuint Resource, bool bLayered, GLint Layer)
+void FOpenGLDynamicRHI::CachedSetupUAVStage( FOpenGLContextState& ContextState, GLint UAVIndex, GLenum Format, GLuint Resource, bool bLayered, GLint Layer, GLenum Access)
 {
-	if (ContextState.UAVs[UAVIndex].Format == Format && 
-		ContextState.UAVs[UAVIndex].Resource == Resource &&
-		ContextState.UAVs[UAVIndex].Layer == Layer &&
-		ContextState.UAVs[UAVIndex].bLayered == bLayered)
+	FUAVStage& UAVStage = ContextState.UAVs[UAVIndex];	
+
+	if (UAVStage.Format == Format && 
+		UAVStage.Resource == Resource &&
+		UAVStage.Access == Access &&
+		UAVStage.Layer == Layer &&
+		UAVStage.bLayered == bLayered)
 	{
 		// Nothing's changed, no need to update
 		return;
 	}
-
-	FOpenGL::BindImageTexture(UAVIndex, Resource, 0, bLayered ? GL_TRUE : GL_FALSE, Layer, GL_READ_WRITE, Format);
+		
+	VERIFY_GL_SCOPE();
 	
-	ContextState.UAVs[UAVIndex].Format = Format;
-	ContextState.UAVs[UAVIndex].Resource = Resource;
-	ContextState.UAVs[UAVIndex].Layer = Layer;
-	ContextState.UAVs[UAVIndex].bLayered = bLayered;
+	FOpenGL::BindImageTexture(UAVIndex, Resource, 0, bLayered ? GL_TRUE : GL_FALSE, Layer, Access, Format);
+	
+	UAVStage.Format = Format;
+	UAVStage.Resource = Resource;
+	UAVStage.Access = Access;
+	UAVStage.Layer = Layer;
+	UAVStage.bLayered = bLayered;
+
+	uint32 UAVBit = 1 << UAVIndex;
+	if (Resource != 0)
+	{
+		ContextState.ActiveUAVMask|= UAVBit;
+	}
+	else
+	{
+		ContextState.ActiveUAVMask&= ~UAVBit;
+	}
 }
 
 void FOpenGLDynamicRHI::UpdateSRV(FOpenGLShaderResourceView* SRV)
@@ -2461,7 +2499,17 @@ template <> FORCEINLINE uint32 GetNumTextureUnits<SF_Geometry>() { return FOpenG
 template <> FORCEINLINE uint32 GetNumTextureUnits<SF_Compute>() { return FOpenGL::GetMaxComputeTextureImageUnits(); }
 
 template <EShaderFrequency Frequency>
-FORCEINLINE void SetResource(FOpenGLDynamicRHI* RESTRICT OpenGLRHI, uint32 BindIndex, FRHITexture* RESTRICT TextureRHI)
+uint32 GetFirstUAVUnit() { return 0; }
+template <> FORCEINLINE uint32 GetFirstUAVUnit<SF_Pixel>() { return FOpenGL::GetFirstPixelUAVUnit(); }
+template <> FORCEINLINE uint32 GetFirstUAVUnit<SF_Compute>() { return FOpenGL::GetFirstComputeUAVUnit(); }
+
+template <EShaderFrequency Frequency>
+uint32 GetNumUAVUnits() { return 0; }
+template <> FORCEINLINE uint32 GetNumUAVUnits<SF_Compute>()	{ return FOpenGL::GetMaxComputeUAVUnits(); }
+template <> FORCEINLINE uint32 GetNumUAVUnits<SF_Pixel>()	{ return FOpenGL::GetMaxPixelUAVUnits(); }
+
+template <EShaderFrequency Frequency>
+void SetResource(FOpenGLDynamicRHI* RESTRICT OpenGLRHI, uint32 BindIndex, FRHITexture* RESTRICT TextureRHI)
 {
 	FOpenGLTextureBase* Texture = GetOpenGLTextureFromRHITexture(TextureRHI);
 #if !UE_BUILD_TEST && !UE_BUILD_SHIPPING
@@ -2476,7 +2524,6 @@ FORCEINLINE void SetResource(FOpenGLDynamicRHI* RESTRICT OpenGLRHI, uint32 BindI
 	{
 		OpenGLRHI->InternalSetShaderTexture(Texture, nullptr, GetFirstTextureUnit<Frequency>() + BindIndex, 0, 0, 0, -1);
 	}
-
 	// clear any previous sampler state
 	OpenGLRHI->InternalSetSamplerStates(GetFirstTextureUnit<Frequency>() + BindIndex, nullptr);
 }
@@ -2498,6 +2545,20 @@ FORCEINLINE void SetResource(FOpenGLDynamicRHI* RESTRICT OpenGLRHI, uint32 BindI
 	SetResource<Frequency>(OpenGLRHI,BindIndex,OpenGLRHI->GetPointSamplerState());
 }
 
+template <EShaderFrequency Frequency>
+FORCEINLINE void SetResource(FOpenGLDynamicRHI* RESTRICT OpenGLRHI, uint32 BindIndex, FRHIUnorderedAccessView* RESTRICT RHIUAV)
+{
+#if !UE_BUILD_TEST && !UE_BUILD_SHIPPING
+	ensureMsgf(BindIndex < GetNumUAVUnits<Frequency>(), TEXT("Using more %s image units (%d) than allowed (%d) on a shader unit!"), GetShaderFrequencyString(Frequency, false), BindIndex, GetNumUAVUnits<Frequency>());	VERIFY_GL_SCOPE();
+#endif
+	auto UAV = FOpenGLDynamicRHI::ResourceCast(RHIUAV);
+	GLenum Access = (Frequency == SF_Compute) ? GL_READ_WRITE : GL_WRITE_ONLY;
+	// TODO: This must be true for 3D textures
+	bool bLayered = false;
+	GLint Layer = 0;
+	OpenGLRHI->InternalSetShaderUAV(GetFirstUAVUnit<Frequency>() + BindIndex, UAV->Format, UAV->Resource, bLayered, Layer, Access);
+}
+
 template <class GLResourceType, EShaderFrequency ShaderFrequency>
 FORCEINLINE void SetShaderResourcesFromBuffer(FOpenGLDynamicRHI* RESTRICT OpenGLRHI, FOpenGLUniformBuffer* RESTRICT Buffer, const uint32* RESTRICT ResourceMap, int32 BufferIndex)
 {
@@ -2544,6 +2605,7 @@ FORCEINLINE void FOpenGLDynamicRHI::SetResourcesFromTables(const ShaderType* RES
 			SetShaderResourcesFromBuffer<FRHITexture, (EShaderFrequency)ShaderType::StaticFrequency>(this, Buffer, SRT->TextureMap.GetData(), BufferIndex);
 			SetShaderResourcesFromBuffer<FRHIShaderResourceView, (EShaderFrequency)ShaderType::StaticFrequency>(this, Buffer, SRT->ShaderResourceViewMap.GetData(), BufferIndex);
 			SetShaderResourcesFromBuffer<FOpenGLSamplerState, (EShaderFrequency)ShaderType::StaticFrequency>(this, Buffer, SRT->SamplerMap.GetData(), BufferIndex);
+			SetShaderResourcesFromBuffer<FRHIUnorderedAccessView, (EShaderFrequency)ShaderType::StaticFrequency>(this, Buffer, SRT->UnorderedAccessViewMap.GetData(), BufferIndex);
 		}
 		BufferIndex++;
 		DirtyBits >>= 1;
@@ -2663,6 +2725,7 @@ void FOpenGLDynamicRHI::RHIDrawPrimitive(uint32 BaseVertexIndex,uint32 NumPrimit
 	BindPendingShaderState(ContextState);
 	CommitGraphicsResourceTables();
 	SetupTexturesForDraw(ContextState);
+	SetupUAVsForDraw(ContextState);
 	CommitNonComputeShaderConstants();
 	CachedBindElementArrayBuffer(ContextState,0);
 	uint32 VertexCount = GetVertexCountForPrimitiveCount(NumPrimitives,PrimitiveType);
@@ -2718,6 +2781,7 @@ void FOpenGLDynamicRHI::RHIDrawPrimitiveIndirect(FRHIVertexBuffer* ArgumentBuffe
 		UpdateDepthStencilStateInOpenGLContext(ContextState);
 		BindPendingShaderState(ContextState);
 		SetupTexturesForDraw(ContextState);
+		SetupUAVsForDraw(ContextState);
 		CommitNonComputeShaderConstants();
 		CachedBindElementArrayBuffer(ContextState,0);
 
@@ -2774,6 +2838,7 @@ void FOpenGLDynamicRHI::RHIDrawIndexedIndirect(FRHIIndexBuffer* IndexBufferRHI,
 		UpdateDepthStencilStateInOpenGLContext(ContextState);
 		BindPendingShaderState(ContextState);
 		SetupTexturesForDraw(ContextState);
+		SetupUAVsForDraw(ContextState);
 		CommitNonComputeShaderConstants();
 		CachedBindElementArrayBuffer(ContextState,IndexBuffer->Resource);
 
@@ -2856,6 +2921,10 @@ void FOpenGLDynamicRHI::RHIDrawIndexedPrimitive(FRHIIndexBuffer* IndexBufferRHI,
 		DETAILED_QUICK_SCOPE_CYCLE_COUNTER(STAT_SetupTexturesForDraw);
 		SetupTexturesForDraw(ContextState);
 	}
+	{
+		DETAILED_QUICK_SCOPE_CYCLE_COUNTER(STAT_SetupUAVsForDraw);
+		SetupUAVsForDraw(ContextState);
+	}
 	{
 		DETAILED_QUICK_SCOPE_CYCLE_COUNTER(STAT_CommitNonComputeShaderConstants);
 		CommitNonComputeShaderConstants();
@@ -2932,6 +3001,7 @@ void FOpenGLDynamicRHI::RHIDrawIndexedPrimitiveIndirect(FRHIIndexBuffer* IndexBu
 		UpdateDepthStencilStateInOpenGLContext(ContextState);
 		BindPendingShaderState(ContextState);
 		SetupTexturesForDraw(ContextState);
+		SetupUAVsForDraw(ContextState);
 		CommitNonComputeShaderConstants();
 		CachedBindElementArrayBuffer(ContextState,IndexBuffer->Resource);
 
@@ -2992,6 +3062,8 @@ static inline void ClearCurrentDepthStencilWithCurrentScissor( int8 ClearType, f
 
 void FOpenGLDynamicRHI::ClearCurrentFramebufferWithCurrentScissor(FOpenGLContextState& ContextState, int8 ClearType, int32 NumClearColors, const FLinearColor* ClearColorArray, float Depth, uint32 Stencil)
 {
+	VERIFY_GL_SCOPE();
+		
 	if ( FOpenGL::SupportsMultipleRenderTargets() )
 	{
 		// Clear color buffers
@@ -3238,7 +3310,7 @@ void FOpenGLDynamicRHI::RHIDispatchComputeShader(uint32 ThreadGroupCountX, uint3
 		BindPendingComputeShaderState(ContextState, ComputeShader);
 		CommitComputeResourceTables(ComputeShader);
 		SetupTexturesForDraw(ContextState, ComputeShader, FOpenGL::GetMaxComputeTextureImageUnits());
-		SetupUAVsForDraw(ContextState, ComputeShader, OGL_MAX_COMPUTE_STAGE_UAV_UNITS);
+		SetupUAVsForCompute(ContextState, ComputeShader);
 		CommitComputeShaderConstants(ComputeShader);
 	
 		FOpenGL::MemoryBarrier(GL_ALL_BARRIER_BITS);
@@ -3276,7 +3348,7 @@ void FOpenGLDynamicRHI::RHIDispatchIndirectComputeShader(FRHIVertexBuffer* Argum
 
 		SetupTexturesForDraw(ContextState, ComputeShader, FOpenGL::GetMaxComputeTextureImageUnits());
 
-		SetupUAVsForDraw(ContextState, ComputeShader, OGL_MAX_COMPUTE_STAGE_UAV_UNITS);
+		SetupUAVsForCompute(ContextState, ComputeShader);
 	
 		CommitComputeShaderConstants(ComputeShader);
 	
@@ -3350,13 +3422,12 @@ void FOpenGLDynamicRHI::RHIInvalidateCachedState()
 	RenderingContextState = FOpenGLContextState();
 	SharedContextState = FOpenGLContextState();
 
-	RenderingContextState.InitializeResources(FOpenGL::GetMaxCombinedTextureImageUnits(), OGL_MAX_COMPUTE_STAGE_UAV_UNITS);
-	SharedContextState.InitializeResources(FOpenGL::GetMaxCombinedTextureImageUnits(), OGL_MAX_COMPUTE_STAGE_UAV_UNITS);
+	RenderingContextState.InitializeResources(FOpenGL::GetMaxCombinedTextureImageUnits(), FOpenGL::GetMaxCombinedUAVUnits());
+	SharedContextState.InitializeResources(FOpenGL::GetMaxCombinedTextureImageUnits(), FOpenGL::GetMaxCombinedUAVUnits());
 }
 
 void FOpenGLDynamicRHI::RHICopyToStagingBuffer(FRHIVertexBuffer* SourceBufferRHI, FRHIStagingBuffer* DestinationStagingBufferRHI, uint32 InOffset, uint32 InNumBytes)
 {
-#if OPENGL_GL3 || USE_ANDROID_OPENGL
 	VERIFY_GL_SCOPE();
 	FOpenGLVertexBuffer* SourceBuffer = ResourceCast(SourceBufferRHI);
 	FOpenGLStagingBuffer* DestinationBuffer = ResourceCast(DestinationStagingBufferRHI);
@@ -3374,9 +3445,6 @@ void FOpenGLDynamicRHI::RHICopyToStagingBuffer(FRHIVertexBuffer* SourceBufferRHI
 	FOpenGL::CopyBufferSubData(GL_COPY_READ_BUFFER, GL_COPY_WRITE_BUFFER, InOffset, 0, InNumBytes);
 	glBindBuffer(GL_COPY_READ_BUFFER, 0);
 	glBindBuffer(GL_COPY_WRITE_BUFFER, 0);
-#else
-	UE_LOG(LogRHI, Fatal, TEXT("Only OpenGL3 or higher supports CopyToStagingBuffer"));
-#endif
 }
 
 void FOpenGLDynamicRHI::RHIWriteGPUFence(FRHIGPUFence* FenceRHI)
diff --git a/Engine/Source/Runtime/OpenGLDrv/Private/OpenGLDevice.cpp b/Engine/Source/Runtime/OpenGLDrv/Private/OpenGLDevice.cpp
--- a/Engine/Source/Runtime/OpenGLDrv/Private/OpenGLDevice.cpp
+++ b/Engine/Source/Runtime/OpenGLDrv/Private/OpenGLDevice.cpp
@@ -917,13 +917,16 @@ static void InitRHICapabilitiesForGL()
 	//@todo: ES2 requires GL_OES_depth_texture extension to support depth textures of any kind.
 	SetupTextureFormat( PF_ShadowDepth,			FOpenGLTextureFormat( ShadowDepthFormat,		ShadowDepthFormat,		GL_DEPTH_COMPONENT,	GL_UNSIGNED_INT,				false,			false));
 	SetupTextureFormat( PF_D24,					FOpenGLTextureFormat( DepthFormat,				DepthFormat,			GL_DEPTH_COMPONENT,	GL_UNSIGNED_INT,				false,			false));
-	SetupTextureFormat( PF_A16B16G16R16,		FOpenGLTextureFormat( GL_RGBA16,				GL_RGBA16,				GL_RGBA,			GL_UNSIGNED_SHORT,				false,			false));
+	SetupTextureFormat( PF_A16B16G16R16,		FOpenGLTextureFormat( GL_RGBA16F,				GL_RGBA16F,				GL_RGBA,			GL_HALF_FLOAT,					false,			false));
 	SetupTextureFormat( PF_A1,					FOpenGLTextureFormat( ));
 	SetupTextureFormat( PF_R16G16B16A16_UINT,	FOpenGLTextureFormat( GL_RGBA16UI,				GL_RGBA16UI,			GL_RGBA_INTEGER,	GL_UNSIGNED_SHORT,				false,			false));
 	SetupTextureFormat( PF_R16G16B16A16_SINT,	FOpenGLTextureFormat( GL_RGBA16I,				GL_RGBA16I,				GL_RGBA_INTEGER,	GL_SHORT,						false,			false));
 	SetupTextureFormat( PF_R32G32B32A32_UINT,	FOpenGLTextureFormat( GL_RGBA32UI,				GL_RGBA32UI,			GL_RGBA_INTEGER,	GL_UNSIGNED_INT,				false,			false));
 	SetupTextureFormat( PF_R16G16B16A16_UNORM,	FOpenGLTextureFormat( GL_RGBA16,				GL_RGBA16,				GL_RGBA,			GL_UNSIGNED_SHORT,				false,			false));
 	SetupTextureFormat( PF_R16G16B16A16_SNORM,	FOpenGLTextureFormat( GL_RGBA16,				GL_RGBA16,				GL_RGBA,			GL_SHORT,						false,			false));
+	
+	SetupTextureFormat( PF_R16G16_UINT,			FOpenGLTextureFormat( GL_RG16UI,				GL_RG16UI,				GL_RG_INTEGER,		GL_UNSIGNED_SHORT,				false,			false));
+	
 
 	SetupTextureFormat( PF_R5G6B5_UNORM,		FOpenGLTextureFormat( ));
 #if PLATFORM_DESKTOP || PLATFORM_ANDROIDESDEFERRED || PLATFORM_LUMINGL4
@@ -994,6 +997,7 @@ static void InitRHICapabilitiesForGL()
 			SetupTextureFormat(PF_R8G8B8A8_SNORM, FOpenGLTextureFormat(GL_RGBA8, GL_RGBA8, GL_RGBA, GL_UNSIGNED_BYTE, false, false));
 
 			SetupTextureFormat(PF_G16R16, FOpenGLTextureFormat(GL_RG16, GL_RG16, GL_RG, GL_UNSIGNED_SHORT, false, false));
+			SetupTextureFormat(PF_R32G32_UINT, FOpenGLTextureFormat(GL_RG32UI, GL_RG32UI, GL_RG, GL_UNSIGNED_INT, false, false));
 		}
 		if (FOpenGL::SupportsPackedDepthStencil())
 		{
@@ -1009,26 +1013,12 @@ static void InitRHICapabilitiesForGL()
 #endif // PLATFORM_DESKTOP || PLATFORM_ANDROIDESDEFERRED || PLATFORM_LUMINGL4
 	{
 #if !PLATFORM_DESKTOP && !PLATFORM_LUMINGL4
-		// ES2-based cases
-		GLuint BGRA8888 = (FOpenGL::SupportsBGRA8888() && !FOpenGL::SupportsSRGB()) ? GL_BGRA_EXT : GL_RGBA;
-		GLuint SizedBGRA8888 = (FOpenGL::SupportsBGRA8888() && !FOpenGL::SupportsSRGB()) ? GL_BGRA_EXT : GL_RGBA8;
-		const bool bNeedsBGRASwizzle = (BGRA8888 == GL_RGBA);
-		GLuint RGBA8 = FOpenGL::SupportsRGBA8() ? GL_RGBA8_OES : GL_RGBA;
-
-	#if PLATFORM_ANDROID
-		SetupTextureFormat(PF_B8G8R8A8, FOpenGLTextureFormat(BGRA8888, GL_SRGB8_ALPHA8, SizedBGRA8888, GL_SRGB8_ALPHA8, BGRA8888, GL_UNSIGNED_BYTE, false, bNeedsBGRASwizzle));
-		SetupTextureFormat(PF_R8G8B8A8_UINT, FOpenGLTextureFormat(GL_RGBA8, GL_RGBA8, GL_RGBA, GL_UNSIGNED_INT_8_8_8_8_REV, false, false));
-	#else
-		SetupTextureFormat(PF_B8G8R8A8, FOpenGLTextureFormat(GL_RGBA, GL_SRGB_ALPHA_EXT, GL_BGRA8_EXT, GL_SRGB8_ALPHA8_EXT, BGRA8888, GL_UNSIGNED_BYTE, false, false));
-	#endif
-		SetupTextureFormat(PF_R8G8B8A8, FOpenGLTextureFormat(RGBA8, GL_SRGB8_ALPHA8, GL_RGBA8, GL_SRGB8_ALPHA8, GL_RGBA, GL_UNSIGNED_BYTE, false, false));
-	#if PLATFORM_IOS
-		SetupTextureFormat(PF_G8, FOpenGLTextureFormat(GL_LUMINANCE, GL_LUMINANCE, GL_LUMINANCE8_EXT, GL_LUMINANCE8_EXT, GL_LUMINANCE, GL_UNSIGNED_BYTE, false, false));
-		SetupTextureFormat(PF_A8, FOpenGLTextureFormat(GL_ALPHA, GL_ALPHA, GL_ALPHA8_EXT, GL_ALPHA8_EXT, GL_ALPHA, GL_UNSIGNED_BYTE, false, false));
-	#else
+		SetupTextureFormat(PF_B8G8R8A8, FOpenGLTextureFormat(GL_RGBA8, GL_SRGB8_ALPHA8, GL_RGBA, GL_UNSIGNED_BYTE, false, true));
+		SetupTextureFormat(PF_R8G8B8A8, FOpenGLTextureFormat(GL_RGBA8, GL_SRGB8_ALPHA8, GL_RGBA, GL_UNSIGNED_BYTE, false, false));
+		SetupTextureFormat(PF_R8G8B8A8_UINT, FOpenGLTextureFormat(GL_RGBA8UI, GL_RGBA8UI, GL_RGBA_INTEGER, GL_UNSIGNED_BYTE, false, false));
 		SetupTextureFormat(PF_G8, FOpenGLTextureFormat(GL_LUMINANCE, GL_LUMINANCE, GL_LUMINANCE, GL_LUMINANCE, GL_LUMINANCE, GL_UNSIGNED_BYTE, false, false));
 		SetupTextureFormat(PF_A8, FOpenGLTextureFormat(GL_ALPHA, GL_ALPHA, GL_ALPHA, GL_ALPHA, GL_ALPHA, GL_UNSIGNED_BYTE, false, false));
-	#endif
+
 		FOpenGL::PE_SetupTextureFormat(&SetupTextureFormat); // platform extension
 
 #if PLATFORM_ANDROID
@@ -1088,6 +1078,7 @@ static void InitRHICapabilitiesForGL()
 			SetupTextureFormat( PF_R32_SINT,		FOpenGLTextureFormat( GL_R32I,					GL_R32I,				GL_RED_INTEGER,	GL_INT,								false,	false));
 			SetupTextureFormat( PF_R16_UINT,		FOpenGLTextureFormat( GL_R16UI,					GL_R16UI,				GL_RED_INTEGER,	GL_UNSIGNED_SHORT,					false,	false));
 			SetupTextureFormat( PF_R16_SINT,		FOpenGLTextureFormat( GL_R16I,					GL_R16I,				GL_RED_INTEGER,	GL_SHORT,							false,	false));
+			SetupTextureFormat( PF_R32G32_UINT,		FOpenGLTextureFormat( GL_RG32UI,				GL_RG32UI,				GL_RG_INTEGER,	GL_UNSIGNED_INT,					false, false));
 		}
 #endif
 	
@@ -1099,7 +1090,7 @@ static void InitRHICapabilitiesForGL()
 		{
 			// @todo android: This is cheating by not setting a stencil anywhere, need that! And Shield is still rendering black scene
 			SetupTextureFormat(PF_DepthStencil, FOpenGLTextureFormat(DepthFormat, GL_NONE, GL_DEPTH_COMPONENT, GL_UNSIGNED_INT, false, false));
-	}
+		}
 #endif // !PLATFORM_DESKTOP
 	}
 
diff --git a/Engine/Source/Runtime/OpenGLDrv/Private/OpenGLDrv.cpp b/Engine/Source/Runtime/OpenGLDrv/Private/OpenGLDrv.cpp
--- a/Engine/Source/Runtime/OpenGLDrv/Private/OpenGLDrv.cpp
+++ b/Engine/Source/Runtime/OpenGLDrv/Private/OpenGLDrv.cpp
@@ -393,9 +393,9 @@ float FOpenGLEventNode::GetTiming()
 void FOpenGLDynamicRHI::InitializeStateResources()
 {
 	VERIFY_GL_SCOPE();
-	SharedContextState.InitializeResources(FOpenGL::GetMaxCombinedTextureImageUnits(), OGL_MAX_COMPUTE_STAGE_UAV_UNITS);
-	RenderingContextState.InitializeResources(FOpenGL::GetMaxCombinedTextureImageUnits(), OGL_MAX_COMPUTE_STAGE_UAV_UNITS);
-	PendingState.InitializeResources(FOpenGL::GetMaxCombinedTextureImageUnits(), OGL_MAX_COMPUTE_STAGE_UAV_UNITS);
+	SharedContextState.InitializeResources(FOpenGL::GetMaxCombinedTextureImageUnits(), FOpenGL::GetMaxCombinedUAVUnits());
+	RenderingContextState.InitializeResources(FOpenGL::GetMaxCombinedTextureImageUnits(), FOpenGL::GetMaxCombinedUAVUnits());
+	PendingState.InitializeResources(FOpenGL::GetMaxCombinedTextureImageUnits(), FOpenGL::GetMaxCombinedUAVUnits());
 }
 
 GLint FOpenGLBase::MaxTextureImageUnits = -1;
diff --git a/Engine/Source/Runtime/OpenGLDrv/Private/OpenGLQuery.cpp b/Engine/Source/Runtime/OpenGLDrv/Private/OpenGLQuery.cpp
--- a/Engine/Source/Runtime/OpenGLDrv/Private/OpenGLQuery.cpp
+++ b/Engine/Source/Runtime/OpenGLDrv/Private/OpenGLQuery.cpp
@@ -1146,17 +1146,11 @@ void FOpenGLDisjointTimeStampQuery::ReleaseResources()
 
 FGPUFenceRHIRef FOpenGLDynamicRHI::RHICreateGPUFence(const FName &Name)
 {
-#if OPENGL_GL3 || USE_ANDROID_OPENGL
 	return new FOpenGLGPUFence(Name);
-#else
-	UE_LOG(LogRHI, Fatal, TEXT("Fences are only available in OpenGL3 or later"));
-	return nullptr;
-#endif
 }
 
 FOpenGLGPUFence::~FOpenGLGPUFence()
 {
-#if OPENGL_GL3 || USE_ANDROID_OPENGL
 	if (bValidSync)
 	{
 		FRHICommandListImmediate& RHICmdList = FRHICommandListExecutor::GetImmediateCommandList();
@@ -1165,15 +1159,10 @@ FOpenGLGPUFence::~FOpenGLGPUFence()
 		FOpenGL::DeleteSync(Fence);
 		RHITHREAD_GLCOMMAND_EPILOGUE_NORETURN();
 	}
-#else
-	UE_LOG(LogRHI, Fatal, TEXT("Fences are only available in OpenGL3 or later"));
-#endif
 }
 
 void FOpenGLGPUFence::Clear()
 {
-#if OPENGL_GL3 || USE_ANDROID_OPENGL
-
 	if (bValidSync)
 	{
 		FRHICommandListImmediate& RHICmdList = FRHICommandListExecutor::GetImmediateCommandList();
@@ -1183,13 +1172,9 @@ void FOpenGLGPUFence::Clear()
 		bValidSync = false;
 		RHITHREAD_GLCOMMAND_EPILOGUE();
 	}
-#else
-	UE_LOG(LogRHI, Fatal, TEXT("Fences are only available in OpenGL3 or later"));
-#endif
 }
 bool FOpenGLGPUFence::Poll() const
 {
-#if OPENGL_GL3 || USE_ANDROID_OPENGL
 	if (!bValidSync)
 	{
 		return false;
@@ -1203,15 +1188,10 @@ bool FOpenGLGPUFence::Poll() const
 	return (Result == FOpenGLBase::FR_AlreadySignaled || Result == FOpenGLBase::FR_ConditionSatisfied);
 
 	RHITHREAD_GLCOMMAND_EPILOGUE_RETURN(bool);
-#else
-	UE_LOG(LogRHI, Fatal, TEXT("Fences are only available in OpenGL3 or later"));
-	return false;
-#endif
 }
 
 void FOpenGLGPUFence::WriteInternal()
 {
-#if OPENGL_GL3 || USE_ANDROID_OPENGL
 	FRHICommandListImmediate& RHICmdList = FRHICommandListExecutor::GetImmediateCommandList();
 	RHITHREAD_GLCOMMAND_PROLOGUE();
 	VERIFY_GL_SCOPE();
@@ -1225,7 +1205,4 @@ void FOpenGLGPUFence::WriteInternal()
 	Fence = FOpenGL::FenceSync(GL_SYNC_GPU_COMMANDS_COMPLETE, 0);
 	bValidSync = true;
 	RHITHREAD_GLCOMMAND_EPILOGUE();
-#else
-	UE_LOG(LogRHI, Fatal, TEXT("Fences are only available in OpenGL3 or later"));
-#endif
 }
diff --git a/Engine/Source/Runtime/OpenGLDrv/Private/OpenGLShaders.cpp b/Engine/Source/Runtime/OpenGLDrv/Private/OpenGLShaders.cpp
--- a/Engine/Source/Runtime/OpenGLDrv/Private/OpenGLShaders.cpp
+++ b/Engine/Source/Runtime/OpenGLDrv/Private/OpenGLShaders.cpp
@@ -1383,13 +1383,16 @@ void OPENGLDRV_API GLSLToDeviceCompatibleGLSL(FAnsiCharArray& GlslCodeOriginal,
 		PRAGMA_ENABLE_DEPRECATION_WARNINGS
 	}
 
-	if (FOpenGL::SupportsClipControl())
+	if (TypeEnum != GL_COMPUTE_SHADER)
 	{
-		AppendCString(GlslCode, "#define HLSLCC_DX11ClipSpace 0 \n");
-	}
-	else
-	{
-		AppendCString(GlslCode, "#define HLSLCC_DX11ClipSpace 1 \n");
+		if (FOpenGL::SupportsClipControl())
+		{
+			AppendCString(GlslCode, "#define HLSLCC_DX11ClipSpace 0 \n");
+		}
+		else
+		{
+			AppendCString(GlslCode, "#define HLSLCC_DX11ClipSpace 1 \n");
+		}
 	}
 
 	// Append the possibly edited shader to the one we will compile.
@@ -1693,6 +1696,7 @@ public:
 
 	int32		MaxTextureStage;
 	TBitArray<>	TextureStageNeeds;
+	int32		MaxUAVUnitUsed;
 	TBitArray<>	UAVStageNeeds;
 
 	TArray<FOpenGLBindlessSamplerInfo> Samplers;
@@ -1713,10 +1717,10 @@ public:
 
 private:
 	FOpenGLLinkedProgram()
-	: Program(0), bUsingTessellation(false), bDrawn(false), bConfigIsInitalized(false), MaxTextureStage(-1)
+	: Program(0), bUsingTessellation(false), bDrawn(false), bConfigIsInitalized(false), MaxTextureStage(-1), MaxUAVUnitUsed(-1)
 	{
 		TextureStageNeeds.Init( false, FOpenGL::GetMaxCombinedTextureImageUnits() );
-		UAVStageNeeds.Init( false, OGL_MAX_COMPUTE_STAGE_UAV_UNITS );
+		UAVStageNeeds.Init( false, FOpenGL::GetMaxCombinedUAVUnits() );
 	}
 
 public:
@@ -2606,15 +2610,15 @@ void FOpenGLLinkedProgram::ConfigureShaderStage( int Stage, uint32 FirstUniformB
 	static const GLint FirstUAVUnit[CrossCompiler::NUM_SHADER_STAGES] =
 	{
 		OGL_UAV_NOT_SUPPORTED_FOR_GRAPHICS_UNIT,
-		OGL_UAV_NOT_SUPPORTED_FOR_GRAPHICS_UNIT,
+		FOpenGL::GetFirstPixelUAVUnit(),
 		OGL_UAV_NOT_SUPPORTED_FOR_GRAPHICS_UNIT,
 		OGL_UAV_NOT_SUPPORTED_FOR_GRAPHICS_UNIT,
 		OGL_UAV_NOT_SUPPORTED_FOR_GRAPHICS_UNIT,
 		FOpenGL::GetFirstComputeUAVUnit()
 	};
 	
-	// verify that only CS uses UAVs
-	check((Stage != CrossCompiler::SHADER_STAGE_COMPUTE) ? (CountSetBits(UAVStageNeeds) == 0) : true);
+	// verify that only CS and PS uses UAVs
+	check(!(Stage == CrossCompiler::SHADER_STAGE_COMPUTE || Stage == CrossCompiler::SHADER_STAGE_PIXEL) ? (CountSetBits(UAVStageNeeds) == 0) : true);
 
 	SCOPE_CYCLE_COUNTER(STAT_OpenGLShaderBindParameterTime);
 	VERIFY_GL_SCOPE();
@@ -2767,6 +2771,7 @@ void FOpenGLLinkedProgram::ConfigureShaderStage( int Stage, uint32 FirstUniformB
 			// glUniform1i(Location, FirstUAVUnit[Stage] + UAVIndex);
 			
 			UAVStageNeeds[ FirstUAVUnit[Stage] + UAVIndex ] = true;
+			MaxUAVUnitUsed = FMath::Max(MaxUAVUnitUsed, FirstUAVUnit[Stage] + UAVIndex);
 		}
 	}
 
@@ -4052,6 +4057,12 @@ const TBitArray<>& FOpenGLBoundShaderState::GetTextureNeeds(int32& OutMaxTexture
 	return LinkedProgram->TextureStageNeeds;
 }
 
+const TBitArray<>& FOpenGLBoundShaderState::GetUAVNeeds(int32& OutMaxUAVUnitUsed) const
+{
+	OutMaxUAVUnitUsed = LinkedProgram->MaxUAVUnitUsed;
+	return LinkedProgram->UAVStageNeeds;
+}
+
 void FOpenGLBoundShaderState::GetNumUniformBuffers(int32 NumUniformBuffers[SF_Compute])
 {
 	if (IsRunningRHIInSeparateThread())
@@ -4102,6 +4113,12 @@ const TBitArray<>& FOpenGLComputeShader::GetTextureNeeds(int32& OutMaxTextureSta
 	return LinkedProgram->TextureStageNeeds;
 }
 
+const TBitArray<>& FOpenGLComputeShader::GetUAVNeeds(int32& OutMaxUAVUnitUsed) const
+{
+	OutMaxUAVUnitUsed = LinkedProgram->MaxUAVUnitUsed;
+	return LinkedProgram->UAVStageNeeds;
+}
+
 bool FOpenGLComputeShader::NeedsUAVStage(int32 UAVStageIndex) const
 {
 	return LinkedProgram->UAVStageNeeds[UAVStageIndex];
diff --git a/Engine/Source/Runtime/OpenGLDrv/Private/OpenGLTexture.cpp b/Engine/Source/Runtime/OpenGLDrv/Private/OpenGLTexture.cpp
--- a/Engine/Source/Runtime/OpenGLDrv/Private/OpenGLTexture.cpp
+++ b/Engine/Source/Runtime/OpenGLDrv/Private/OpenGLTexture.cpp
@@ -429,7 +429,7 @@ void FOpenGLDynamicRHI::InitializeGLTexture(FRHITexture* Texture, uint32 SizeX,
 			}
 		}
 		glTexParameteri(Target, GL_TEXTURE_MAG_FILTER, GL_NEAREST);
-		glTexParameteri(Target, GL_TEXTURE_MIN_FILTER, GL_NEAREST);
+		glTexParameteri(Target, GL_TEXTURE_MIN_FILTER, NumMips > 1 ? GL_NEAREST_MIPMAP_NEAREST : GL_NEAREST);
 		if( FOpenGL::SupportsTextureFilterAnisotropic() )
 		{
 			glTexParameteri(Target, GL_TEXTURE_MAX_ANISOTROPY_EXT, 1);
@@ -2434,7 +2434,7 @@ void FOpenGLDynamicRHI::InvalidateTextureResourceInCache(GLuint Resource)
 
 void FOpenGLDynamicRHI::InvalidateUAVResourceInCache(GLuint Resource)
 {
-	for (int32 UAVIndex = 0; UAVIndex < OGL_MAX_COMPUTE_STAGE_UAV_UNITS; ++UAVIndex)
+	for (int32 UAVIndex = 0; UAVIndex < FOpenGL::GetMaxCombinedUAVUnits(); ++UAVIndex)
 	{
 		if (SharedContextState.UAVs[UAVIndex].Resource == Resource)
 		{
@@ -2574,7 +2574,7 @@ void FOpenGLDynamicRHI::RHICopySubTextureRegion(FRHITexture2D* SourceTextureRHI,
 	ContextState.Framebuffer = (GLuint)-1;
 }
 
-/*
+
 void FOpenGLDynamicRHI::RHICopyTexture(FRHITexture* SourceTextureRHI, FRHITexture* DestTextureRHI, const FRHICopyTextureInfo& CopyInfo)
 {
 	VERIFY_GL_SCOPE();
@@ -2720,7 +2720,7 @@ void FOpenGLDynamicRHI::RHICopyTexture(FRHITexture* SourceTextureRHI, FRHITextur
 
 	ContextState.Framebuffer = (GLuint)-1;
 }
-*/
+
 
 FTexture2DRHIRef FOpenGLDynamicRHI::RHICreateTexture2DFromResource(EPixelFormat Format, uint32 SizeX, uint32 SizeY, uint32 NumMips, uint32 NumSamples, uint32 NumSamplesTileMem, const FClearValueBinding& ClearValueBinding, GLuint Resource, uint32 TexCreateFlags)
 {
diff --git a/Engine/Source/Runtime/OpenGLDrv/Private/OpenGLVertexBuffer.cpp b/Engine/Source/Runtime/OpenGLDrv/Private/OpenGLVertexBuffer.cpp
--- a/Engine/Source/Runtime/OpenGLDrv/Private/OpenGLVertexBuffer.cpp
+++ b/Engine/Source/Runtime/OpenGLDrv/Private/OpenGLVertexBuffer.cpp
@@ -253,12 +253,7 @@ void FOpenGLDynamicRHI::RHICopyVertexBuffer(FRHIVertexBuffer* SourceBufferRHI, F
 
 FStagingBufferRHIRef FOpenGLDynamicRHI::RHICreateStagingBuffer()
 {
-#if OPENGL_GL3 || USE_ANDROID_OPENGL
 	return new FOpenGLStagingBuffer();
-#else
-	UE_LOG(LogRHI, Fatal, TEXT("Staging Buffers are only available in OpenGL3 or later"));
-	return nullptr;
-#endif
 }
 
 void FOpenGLStagingBuffer::Initialize()
@@ -285,7 +280,6 @@ FOpenGLStagingBuffer::~FOpenGLStagingBuffer()
 // I don't see a way to do this without stalling the RHI thread.
 void* FOpenGLStagingBuffer::Lock(uint32 Offset, uint32 NumBytes)
 {
-#if OPENGL_GL3 || USE_ANDROID_OPENGL
 	FRHICommandListImmediate& RHICmdList = FRHICommandListExecutor::GetImmediateCommandList();
 	RHITHREAD_GLCOMMAND_PROLOGUE();
 	VERIFY_GL_SCOPE();
@@ -297,25 +291,17 @@ void* FOpenGLStagingBuffer::Lock(uint32 Offset, uint32 NumBytes)
 	return reinterpret_cast<uint8*>(Mapping) + Offset;
 
 	RHITHREAD_GLCOMMAND_EPILOGUE_RETURN(void*);
-#else
-	UE_LOG(LogRHI, Fatal, TEXT("Staging Buffers are only available in OpenGL3 or later"));
-	return nullptr;
-#endif
 }
 
 // Unfortunately I think we have to stall the RHI thread here as well to play nice with OpenGL.
 // Since this will probably be close to a call to lock we've probably paid most of the cost already.
 void FOpenGLStagingBuffer::Unlock()
 {
-#if OPENGL_GL3 || USE_ANDROID_OPENGL
 	FRHICommandListImmediate& RHICmdList = FRHICommandListExecutor::GetImmediateCommandList();
 	RHITHREAD_GLCOMMAND_PROLOGUE();
 	FOpenGL::UnmapBuffer(GL_COPY_WRITE_BUFFER);
 	glBindBuffer(GL_COPY_WRITE_BUFFER, 0);
 	RHITHREAD_GLCOMMAND_EPILOGUE();
-#else
-	UE_LOG(LogRHI, Fatal, TEXT("Staging Buffers are only available in OpenGL3 or later"));
-#endif
 }
 
 void* FOpenGLDynamicRHI::RHILockStagingBuffer(FRHIStagingBuffer* StagingBuffer, uint32 Offset, uint32 SizeRHI)
diff --git a/Engine/Source/Runtime/OpenGLDrv/Public/OpenGL.h b/Engine/Source/Runtime/OpenGLDrv/Public/OpenGL.h
--- a/Engine/Source/Runtime/OpenGLDrv/Public/OpenGL.h
+++ b/Engine/Source/Runtime/OpenGLDrv/Public/OpenGL.h
@@ -205,8 +205,14 @@ public:
 	static FORCEINLINE GLint GetFirstDomainTextureUnit()		{ return GetFirstHullTextureUnit() + GetMaxHullTextureImageUnits(); }
 
 	static FORCEINLINE GLint GetFirstComputeTextureUnit()		{ return 0; }
+	
+	// Image load/store units
 	static FORCEINLINE GLint GetFirstComputeUAVUnit()			{ return 0; }
-
+	static FORCEINLINE GLint GetMaxComputeUAVUnits()			{ return 0; }
+	static FORCEINLINE GLint GetFirstPixelUAVUnit()				{ return 0; }
+	static FORCEINLINE GLint GetMaxPixelUAVUnits()				{ return 0; }
+	static FORCEINLINE GLint GetMaxCombinedUAVUnits()			{ return 0; }
+	
 	static FORCEINLINE GLint GetMaxVaryingVectors()				{ check(MaxVaryingVectors != -1); return MaxVaryingVectors; }
 	static FORCEINLINE GLint GetMaxPixelUniformComponents()		{ check(MaxPixelUniformComponents != -1); return MaxPixelUniformComponents; }
 	static FORCEINLINE GLint GetMaxVertexUniformComponents()	{ check(MaxVertexUniformComponents != -1); return MaxVertexUniformComponents; }
diff --git a/Engine/Source/Runtime/OpenGLDrv/Public/OpenGL4.h b/Engine/Source/Runtime/OpenGLDrv/Public/OpenGL4.h
--- a/Engine/Source/Runtime/OpenGLDrv/Public/OpenGL4.h
+++ b/Engine/Source/Runtime/OpenGLDrv/Public/OpenGL4.h
@@ -102,6 +102,12 @@ struct FOpenGL4 : public FOpenGL3
 	static FORCEINLINE GLint GetMaxComputeTextureImageUnits() { check(MaxComputeTextureImageUnits != -1); return MaxComputeTextureImageUnits; }
 	static FORCEINLINE GLint GetMaxComputeUniformComponents() { check(MaxComputeUniformComponents != -1); return MaxComputeUniformComponents; }
 
+	static FORCEINLINE GLint GetFirstComputeUAVUnit()			{ return 0; }
+	static FORCEINLINE GLint GetMaxComputeUAVUnits()			{ check(MaxComputeUAVUnits != -1); return MaxComputeUAVUnits; }
+	static FORCEINLINE GLint GetFirstPixelUAVUnit()				{ return 0; }
+	static FORCEINLINE GLint GetMaxPixelUAVUnits()				{ check(MaxPixelUAVUnits != -1); return MaxPixelUAVUnits; }
+	static FORCEINLINE GLint GetMaxCombinedUAVUnits()			{ return MaxCombinedUAVUnits; }
+		
 	static uint64 GetVideoMemorySize();
 
 	static FORCEINLINE int32 GetReadHalfFloatPixelsEnum() { return GL_HALF_FLOAT; }
@@ -110,6 +116,10 @@ protected:
 	static GLint MaxComputeTextureImageUnits;
 	static GLint MaxComputeUniformComponents;
 
+	static GLint MaxCombinedUAVUnits;
+	static GLint MaxComputeUAVUnits;
+	static GLint MaxPixelUAVUnits;
+
 	static bool bSupportsComputeShaders;
 	static bool bSupportsGPUMemoryInfo;
 	static bool bSupportsVertexAttribBinding;
diff --git a/Engine/Source/Runtime/OpenGLDrv/Public/OpenGLDrv.h b/Engine/Source/Runtime/OpenGLDrv/Public/OpenGLDrv.h
--- a/Engine/Source/Runtime/OpenGLDrv/Public/OpenGLDrv.h
+++ b/Engine/Source/Runtime/OpenGLDrv/Public/OpenGLDrv.h
@@ -521,7 +521,7 @@ public:
 
 	virtual void RHICopySubTextureRegion(FRHITexture2D* SourceTexture, FRHITexture2D* DestinationTexture, FBox2D SourceBox, FBox2D DestinationBox) final override;
 	// FIXME: Broken on Android for cubemaps
-	//virtual void RHICopyTexture(FRHITexture* SourceTexture, FRHITexture* DestTexture, const FRHICopyTextureInfo& CopyInfo) final override;
+	virtual void RHICopyTexture(FRHITexture* SourceTexture, FRHITexture* DestTexture, const FRHICopyTextureInfo& CopyInfo) final override;
 
 	// FOpenGLDynamicRHI interface.
 	virtual FTexture2DRHIRef RHICreateTexture2DFromResource(EPixelFormat Format, uint32 SizeX, uint32 SizeY, uint32 NumMips, uint32 NumSamples, uint32 NumSamplesTileMem, const FClearValueBinding& ClearValueBinding, GLuint Resource, uint32 Flags);
@@ -568,7 +568,7 @@ public:
 	}
 
 	void CachedSetupTextureStageInner(FOpenGLContextState& ContextState, GLint TextureIndex, GLenum Target, GLuint Resource, GLint BaseMip, GLint NumMips);
-	void CachedSetupUAVStage(FOpenGLContextState& ContextState, GLint UAVIndex, GLenum Format, GLuint Resource, bool bLayered, GLint Layer);
+	void CachedSetupUAVStage(FOpenGLContextState& ContextState, GLint UAVIndex, GLenum Format, GLuint Resource, bool bLayered, GLint Layer, GLenum Access);
 	void UpdateSRV(FOpenGLShaderResourceView* SRV);
 	FOpenGLContextState& GetContextStateForCurrentContext(bool bAssertIfInvalid = true);
 
@@ -1145,14 +1145,16 @@ private:
 	template <typename StateType>
 	void SetupTexturesForDraw( FOpenGLContextState& ContextState, const StateType& ShaderState, int32 MaxTexturesNeeded);
 
-	void SetupUAVsForDraw(FOpenGLContextState& ContextState, const FOpenGLComputeShader* ComputeShader, int32 MaxUAVsNeeded);
+	void SetupUAVsForDraw(FOpenGLContextState& ContextState);
+	void SetupUAVsForCompute(FOpenGLContextState& ContextState, const FOpenGLComputeShader* ComputeShader);
+	void SetupUAVsForProgram(FOpenGLContextState& ContextState, const TBitArray<>& NeededBits, int32 MaxUAVUnitUsed, int32 MaxUAVUnits);
 
 	void RHIClearMRT(bool bClearColor, int32 NumClearColors, const FLinearColor* ColorArray, bool bClearDepth, float Depth, bool bClearStencil, uint32 Stencil);
 
 public:
 	/** Remember what RHI user wants set on a specific OpenGL texture stage, translating from Stage and TextureIndex for stage pair. */
 	void InternalSetShaderTexture(FOpenGLTextureBase* Texture, FOpenGLShaderResourceView* SRV, GLint TextureIndex, GLenum Target, GLuint Resource, int NumMips, int LimitMip);
-	void InternalSetShaderUAV(GLint UAVIndex, GLenum Format, GLuint Resource, bool bLayered, GLint Layer);
+	void InternalSetShaderUAV(GLint UAVIndex, GLenum Format, GLuint Resource, bool bLayered, GLint Layer, GLenum Access);
 	void InternalSetSamplerStates(GLint TextureIndex, FOpenGLSamplerState* SamplerState);
 
 private:
diff --git a/Engine/Source/Runtime/OpenGLDrv/Public/OpenGLResources.h b/Engine/Source/Runtime/OpenGLDrv/Public/OpenGLResources.h
--- a/Engine/Source/Runtime/OpenGLDrv/Public/OpenGLResources.h
+++ b/Engine/Source/Runtime/OpenGLDrv/Public/OpenGLResources.h
@@ -1215,6 +1215,7 @@ public:
 		);
 
 	const TBitArray<>& GetTextureNeeds(int32& OutMaxTextureStageUsed);
+	const TBitArray<>& GetUAVNeeds(int32& OutMaxUAVUnitUsed) const;
 	void GetNumUniformBuffers(int32 NumVertexUniformBuffers[SF_Compute]);
 
 	bool NeedsTextureStage(int32 TextureStageIndex);
diff --git a/Engine/Source/Runtime/OpenGLDrv/Public/OpenGLShaderResources.h b/Engine/Source/Runtime/OpenGLDrv/Public/OpenGLShaderResources.h
--- a/Engine/Source/Runtime/OpenGLDrv/Public/OpenGLShaderResources.h
+++ b/Engine/Source/Runtime/OpenGLDrv/Public/OpenGLShaderResources.h
@@ -30,8 +30,7 @@ enum
 {
 	OGL_MAX_UNIFORM_BUFFER_BINDINGS = 12,	// @todo-mobile: Remove me
 	OGL_FIRST_UNIFORM_BUFFER = 0,			// @todo-mobile: Remove me
-	OGL_MAX_COMPUTE_STAGE_UAV_UNITS = 8,	// @todo-mobile: Remove me
-	OGL_UAV_NOT_SUPPORTED_FOR_GRAPHICS_UNIT = -1, // for now, only CS supports UAVs/ images
+	OGL_UAV_NOT_SUPPORTED_FOR_GRAPHICS_UNIT = -1, // for now, only CS and PS supports UAVs/ images
 };
 
 struct FOpenGLShaderResourceTable : public FBaseShaderResourceTable
@@ -325,6 +324,7 @@ public:
 	bool NeedsTextureStage(int32 TextureStageIndex);
 	int32 MaxTextureStageUsed();
 	const TBitArray<>& GetTextureNeeds(int32& OutMaxTextureStageUsed);
+	const TBitArray<>& GetUAVNeeds(int32& OutMaxUAVUnitUsed) const;
 	bool NeedsUAVStage(int32 UAVStageIndex) const;
 
 	FOpenGLLinkedProgram* LinkedProgram;
diff --git a/Engine/Source/Runtime/OpenGLDrv/Public/OpenGLState.h b/Engine/Source/Runtime/OpenGLDrv/Public/OpenGLState.h
--- a/Engine/Source/Runtime/OpenGLDrv/Public/OpenGLState.h
+++ b/Engine/Source/Runtime/OpenGLDrv/Public/OpenGLState.h
@@ -220,12 +220,14 @@ struct FUAVStage
 {
 	GLenum Format;
 	GLuint Resource;
+	GLenum Access;
 	GLint Layer;
 	bool bLayered;
 	
 	FUAVStage()
 	:	Format(GL_NONE)
 	,	Resource(0)
+	,	Access(GL_READ_WRITE)
 	,	Layer(0)
 	,	bLayered(false)
 	{
@@ -290,16 +292,16 @@ struct FOpenGLCommonState
 	}
 
 	// NumCombinedTextures must be greater than or equal to FOpenGL::GetMaxCombinedTextureImageUnits()
-	// NumComputeUAVUnits must be greater than or equal to OGL_MAX_COMPUTE_STAGE_UAV_UNITS
-	virtual void InitializeResources(int32 NumCombinedTextures, int32 NumComputeUAVUnits)
+	// NumCombinedUAVUnits must be greater than or equal to FOpenGL::GetMaxCombinedUAVUnits()
+	virtual void InitializeResources(int32 NumCombinedTextures, int32 NumCombinedUAVUnits)
 	{
 		check(NumCombinedTextures >= FOpenGL::GetMaxCombinedTextureImageUnits());
-		check(NumComputeUAVUnits >= OGL_MAX_COMPUTE_STAGE_UAV_UNITS);
+		check(NumCombinedUAVUnits >= FOpenGL::GetMaxCombinedUAVUnits());
 		check(!Textures && !SamplerStates && !UAVs);
 		Textures = new FTextureStage[NumCombinedTextures];
 		SamplerStates = new FOpenGLSamplerState*[NumCombinedTextures];
 		FMemory::Memset( SamplerStates, 0, NumCombinedTextures * sizeof(*SamplerStates) );
-		UAVs = new FUAVStage[NumComputeUAVUnits];
+		UAVs = new FUAVStage[NumCombinedUAVUnits];
 	}
 
 	virtual void CleanupResources()
@@ -376,6 +378,8 @@ struct FOpenGLContextState : public FOpenGLCommonState
 	uint32 ActiveStreamMask;
 	uint32 MaxActiveAttrib;
 
+	uint32 ActiveUAVMask;
+
 	FOpenGLContextState()
 	:	StencilRef(0)
 	,	Framebuffer(0)
@@ -404,6 +408,7 @@ struct FOpenGLContextState : public FOpenGLCommonState
 	, ActiveAttribMask(0)
 	, ActiveStreamMask(0)
 	, MaxActiveAttrib(0)
+	, ActiveUAVMask(0)
 	{
 		Scissor.Min.X = Scissor.Min.Y = Scissor.Max.X = Scissor.Max.Y = 0;
 		Viewport.Min.X = Viewport.Min.Y = Viewport.Max.X = Viewport.Max.Y = 0;
@@ -412,11 +417,13 @@ struct FOpenGLContextState : public FOpenGLCommonState
 		FMemory::Memzero(DrawFramebuffers, sizeof(DrawFramebuffers));
 	}
 
-	virtual void InitializeResources(int32 NumCombinedTextures, int32 NumComputeUAVUnits) override
+	virtual void InitializeResources(int32 NumCombinedTextures, int32 NumCombinedUAVUnits) override
 	{
-		FOpenGLCommonState::InitializeResources(NumCombinedTextures, NumComputeUAVUnits);
+		FOpenGLCommonState::InitializeResources(NumCombinedTextures, NumCombinedUAVUnits);
 		CachedSamplerStates.Empty(NumCombinedTextures);
 		CachedSamplerStates.AddZeroed(NumCombinedTextures);
+
+		checkf(NumCombinedUAVUnits <= sizeof(ActiveUAVMask) * 8, TEXT("Not enough bits in ActiveUAVMask to store %d UAV units"), NumCombinedUAVUnits);
 	}
 
 	virtual void CleanupResources() override
diff --git a/Engine/Source/Runtime/RHI/Public/RHIDefinitions.h b/Engine/Source/Runtime/RHI/Public/RHIDefinitions.h
--- a/Engine/Source/Runtime/RHI/Public/RHIDefinitions.h
+++ b/Engine/Source/Runtime/RHI/Public/RHIDefinitions.h
@@ -1373,6 +1373,7 @@ inline bool IsShaderParameterTypeForUniformBufferLayout(EUniformBufferBaseType B
 		BaseType == UBMT_TEXTURE ||
 		BaseType == UBMT_SRV ||
 		BaseType == UBMT_SAMPLER ||
+		BaseType == UBMT_UAV ||
 
 		// RHI is able to access RHI resources from RDG.
 		IsRDGResourceReferenceShaderParameterType(BaseType) ||
diff --git a/Engine/Source/Runtime/RenderCore/Private/RenderUtils.cpp b/Engine/Source/Runtime/RenderCore/Private/RenderUtils.cpp
--- a/Engine/Source/Runtime/RenderCore/Private/RenderUtils.cpp
+++ b/Engine/Source/Runtime/RenderCore/Private/RenderUtils.cpp
@@ -1250,12 +1250,6 @@ RENDERCORE_API bool UseVirtualTexturing(ERHIFeatureLevel::Type InFeatureLevel, c
 	else
 #endif
 	{
-		// only at SM5 
-		if (InFeatureLevel < ERHIFeatureLevel::SM5)
-		{
-			return false;
-		}
-
 		// does the platform supports it.
 #if WITH_EDITOR
 		if (GIsEditor && TargetPlatform == nullptr)
@@ -1273,13 +1267,6 @@ RENDERCORE_API bool UseVirtualTexturing(ERHIFeatureLevel::Type InFeatureLevel, c
 		}
 #endif
 
-		// Remove this when UE-80097 is implemented
-		static bool bIsVulkanRHI = FCString::Strcmp(GDynamicRHI->GetName(), TEXT("Vulkan")) == 0;
-		if (bIsVulkanRHI)
-		{
-			return false;
-		}
-
 		// does the project has it enabled ?
 		static const auto CVarVirtualTexture = IConsoleManager::Get().FindTConsoleVariableDataInt(TEXT("r.VirtualTextures"));
 		check(CVarVirtualTexture);
@@ -1288,6 +1275,13 @@ RENDERCORE_API bool UseVirtualTexturing(ERHIFeatureLevel::Type InFeatureLevel, c
 			return false;
 		}		
 
+		// mobile needs an additional switch to enable VT		
+		static const auto CVarMobileVirtualTexture = IConsoleManager::Get().FindTConsoleVariableDataInt(TEXT("r.Mobile.VirtualTextures"));
+		if (InFeatureLevel == ERHIFeatureLevel::ES3_1 && CVarMobileVirtualTexture->GetValueOnAnyThread() == 0)
+		{
+			return false;
+		}
+
 		return true;
 	}
 }
diff --git a/Engine/Source/Runtime/RenderCore/Private/Shader.cpp b/Engine/Source/Runtime/RenderCore/Private/Shader.cpp
--- a/Engine/Source/Runtime/RenderCore/Private/Shader.cpp
+++ b/Engine/Source/Runtime/RenderCore/Private/Shader.cpp
@@ -2533,8 +2533,14 @@ void ShaderMapAppendKeyString(EShaderPlatform Platform, FString& KeyString)
 		const bool VTLightmaps = CVarVirtualTextureLightmaps && CVarVirtualTextureLightmaps->GetValueOnAnyThread() != 0;
 
 		static const auto CVarVirtualTexture = IConsoleManager::Get().FindTConsoleVariableDataInt(TEXT("r.VirtualTextures"));
-		const bool VTTextures = CVarVirtualTexture && CVarVirtualTexture->GetValueOnAnyThread() != 0;
+		bool VTTextures = CVarVirtualTexture && CVarVirtualTexture->GetValueOnAnyThread() != 0;
 
+		static const auto CVarMobileVirtualTexture = IConsoleManager::Get().FindTConsoleVariableDataInt(TEXT("r.Mobile.VirtualTextures"));
+		if (IsMobilePlatform(Platform) && VTTextures)
+		{
+			VTTextures = (CVarMobileVirtualTexture->GetValueOnAnyThread() != 0);
+		}
+		
 		static const auto CVarVTFactor = IConsoleManager::Get().FindConsoleVariable(TEXT("r.vt.FeedbackFactor")); check(CVarVTFactor);
 		const int32 VTFeedbackFactor = CVarVTFactor->GetInt(); 
 
diff --git a/Engine/Source/Runtime/RenderCore/Public/ClearReplacementShaders.h b/Engine/Source/Runtime/RenderCore/Public/ClearReplacementShaders.h
--- a/Engine/Source/Runtime/RenderCore/Public/ClearReplacementShaders.h
+++ b/Engine/Source/Runtime/RenderCore/Public/ClearReplacementShaders.h
@@ -113,7 +113,8 @@ public:
 
 	static bool ShouldCompilePermutation(const FGlobalShaderPermutationParameters& Parameters)
 	{
-		return IsFeatureLevelSupported(Parameters.Platform, ERHIFeatureLevel::SM5);
+		return IsFeatureLevelSupported(Parameters.Platform, ERHIFeatureLevel::SM5)
+				|| GetMaxSupportedFeatureLevel(Parameters.Platform) == ERHIFeatureLevel::ES3_1;
 	}
 
 	static inline void ModifyCompilationEnvironment(const FGlobalShaderPermutationParameters& Parameters, FShaderCompilerEnvironment& OutEnvironment)
@@ -173,7 +174,8 @@ public:
 
 	static bool ShouldCompilePermutation(const FGlobalShaderPermutationParameters& Parameters)
 	{
-		return IsFeatureLevelSupported(Parameters.Platform, ERHIFeatureLevel::SM5);
+		return IsFeatureLevelSupported(Parameters.Platform, ERHIFeatureLevel::SM5)
+				|| GetMaxSupportedFeatureLevel(Parameters.Platform) == ERHIFeatureLevel::ES3_1;
 	}
 
 	static inline void ModifyCompilationEnvironment(const FGlobalShaderPermutationParameters& Parameters, FShaderCompilerEnvironment& OutEnvironment)
@@ -220,7 +222,8 @@ public:
 
 	static bool ShouldCompilePermutation(const FGlobalShaderPermutationParameters& Parameters)
 	{
-		return IsFeatureLevelSupported(Parameters.Platform, ERHIFeatureLevel::SM5);
+		return IsFeatureLevelSupported(Parameters.Platform, ERHIFeatureLevel::SM5)
+				|| GetMaxSupportedFeatureLevel(Parameters.Platform) == ERHIFeatureLevel::ES3_1;
 	}
 
 	static inline void ModifyCompilationEnvironment(const FGlobalShaderPermutationParameters& Parameters, FShaderCompilerEnvironment& OutEnvironment)
diff --git a/Engine/Source/Runtime/Renderer/Private/MobileShadingRenderer.cpp b/Engine/Source/Runtime/Renderer/Private/MobileShadingRenderer.cpp
--- a/Engine/Source/Runtime/Renderer/Private/MobileShadingRenderer.cpp
+++ b/Engine/Source/Runtime/Renderer/Private/MobileShadingRenderer.cpp
@@ -48,6 +48,7 @@
 #include "DebugViewModeRendering.h"
 #include "SkyAtmosphereRendering.h"
 #include "VisualizeTexture.h"
+#include "VT/VirtualTextureSystem.h"
 
 uint32 GetShadowQuality();
 
@@ -405,15 +406,29 @@ void FMobileSceneRenderer::Render(FRHICommandListImmediate& RHICmdList)
 	FSceneRenderTargets& SceneContext = FSceneRenderTargets::Get(RHICmdList);
 
 	// Allocate the maximum scene render target space for the current view family.
+	SceneContext.ReleaseSceneColor();
 	SceneContext.SetKeepDepthContent(bKeepDepthContent);
 	SceneContext.Allocate(RHICmdList, this);
 
+	const bool bUseVirtualTexturing = UseVirtualTexturing(ViewFeatureLevel);
+	if (bUseVirtualTexturing)
+	{
+		// AllocateResources needs to be called before RHIBeginScene
+		FVirtualTextureSystem::Get().AllocateResources(RHICmdList, ViewFeatureLevel);
+		FVirtualTextureSystem::Get().CallPendingCallbacks();
+	}
+
 	//make sure all the targets we're going to use will be safely writable.
 	GRenderTargetPool.TransitionTargetsWritable(RHICmdList);
 
 	// Find the visible primitives.
 	InitViews(RHICmdList);
 
+	if (bUseVirtualTexturing)
+	{
+		FVirtualTextureSystem::Get().Update(RHICmdList, ViewFeatureLevel, Scene);
+	}
+	
 	if (GRHINeedsExtraDeletionLatency || !GRHICommandList.Bypass())
 	{
 		QUICK_SCOPE_CYCLE_COUNTER(STAT_FMobileSceneRenderer_PostInitViewsFlushDel);
@@ -550,6 +565,13 @@ void FMobileSceneRenderer::Render(FRHICommandListImmediate& RHICmdList)
 	SceneColorRenderPassInfo.NumOcclusionQueries = ComputeNumOcclusionQueriesToBatch();
 	SceneColorRenderPassInfo.bOcclusionQueries = SceneColorRenderPassInfo.NumOcclusionQueries != 0;
 	SceneColorRenderPassInfo.bMultiviewPass = View.bIsMobileMultiViewEnabled;
+
+	// TODO: required only for DX11 ?
+	if (UseVirtualTexturing(ViewFeatureLevel) && !IsHlslccShaderPlatform(View.GetShaderPlatform()))
+	{
+		SceneContext.BindVirtualTextureFeedbackUAV(SceneColorRenderPassInfo);
+	}
+	
 	RHICmdList.BeginRenderPass(SceneColorRenderPassInfo, TEXT("SceneColorRendering"));
 
 	if (GIsEditor && !View.bIsSceneCapture)
@@ -721,6 +743,19 @@ void FMobileSceneRenderer::Render(FRHICommandListImmediate& RHICmdList)
 		CopyMobileMultiViewSceneColor(RHICmdList);
 	}
 
+	if (bUseVirtualTexturing)
+	{	
+		// No pass after this can make VT page requests
+		TArray<FIntRect, TInlineAllocator<FVirtualTextureFeedback::MaxRectPerTarget>> ViewRects;
+		ViewRects.AddUninitialized(Views.Num());
+		for (int32 ViewIndex = 0; ViewIndex < Views.Num(); ++ViewIndex)
+		{
+			ViewRects[ViewIndex] = Views[ViewIndex].ViewRect;
+		}
+		
+		SceneContext.VirtualTextureFeedback.TransferGPUToCPU(RHICmdList, ViewRects);
+	}
+	
 	if (ViewFamily.bResolveScene)
 	{
 		if (!bGammaSpace)
diff --git a/Engine/Source/Runtime/Renderer/Private/PostProcess/SceneRenderTargets.cpp b/Engine/Source/Runtime/Renderer/Private/PostProcess/SceneRenderTargets.cpp
--- a/Engine/Source/Runtime/Renderer/Private/PostProcess/SceneRenderTargets.cpp
+++ b/Engine/Source/Runtime/Renderer/Private/PostProcess/SceneRenderTargets.cpp
@@ -318,7 +318,6 @@ FSceneRenderTargets::FSceneRenderTargets(const FViewInfo& View, const FSceneRend
 	SnapshotArray(TranslucencyLightingVolumeAmbient, SnapshotSource.TranslucencyLightingVolumeAmbient);
 	SnapshotArray(TranslucencyLightingVolumeDirectional, SnapshotSource.TranslucencyLightingVolumeDirectional);
 	SnapshotArray(OptionalShadowDepthColor, SnapshotSource.OptionalShadowDepthColor);
-	VirtualTextureFeedback.MakeSnapshot(SnapshotSource.VirtualTextureFeedback);
 }
 
 inline const TCHAR* GetSceneColorTargetName(EShadingPath ShadingPath)
@@ -893,10 +892,10 @@ void FSceneRenderTargets::BindVirtualTextureFeedbackUAV(FRHIRenderPassInfo& RPIn
 	// must match VT_FEEDBACK_REGISTER in VirtualTextureCommon.ush
 	static const int32 VT_FEEDBACK_REGISTER = 7;
 
-	checkf(VirtualTextureFeedback.FeedbackTextureGPU, TEXT("Invalid attempt to render VT feedback after it has already been transferred to CPU, need to adjust location of TransferGPUToCPU"));
+	checkf(VirtualTextureFeedback.FeedbackBufferUAV, TEXT("Invalid attempt to render VT feedback after it has already been transferred to CPU, need to adjust location of TransferGPUToCPU"));
 	check(RPInfo.GetNumColorRenderTargets() <= VT_FEEDBACK_REGISTER);
 	RPInfo.UAVIndex = VT_FEEDBACK_REGISTER;
-	RPInfo.UAVs[RPInfo.NumUAVs++] = VirtualTextureFeedback.FeedbackTextureGPU->GetRenderTargetItem().UAV;
+	RPInfo.UAVs[RPInfo.NumUAVs++] = VirtualTextureFeedback.FeedbackBufferUAV;
 }
 
 void FSceneRenderTargets::BeginRenderingGBuffer(FRHICommandList& RHICmdList, ERenderTargetLoadAction ColorLoadAction, ERenderTargetLoadAction DepthLoadAction, FExclusiveDepthStencil DepthStencilAccess, bool bBindQuadOverdrawBuffers, bool bClearQuadOverdrawBuffers, const FLinearColor& ClearColor/*=(0,0,0,1)*/, bool bIsWireframe)
@@ -2062,7 +2061,7 @@ void FSceneRenderTargets::SetSeparateTranslucencyBufferSize(bool bAnyViewWantsDo
 	SeparateTranslucencyScale = EffectiveScale;
 }
 
-void FSceneRenderTargets::AllocateMobileRenderTargets(FRHICommandList& RHICmdList)
+void FSceneRenderTargets::AllocateMobileRenderTargets(FRHICommandListImmediate& RHICmdList)
 {
 	// on ES2 we don't do on demand allocation of SceneColor yet (in non ES2 it's released in the Tonemapper Process())
 	AllocSceneColor(RHICmdList);
@@ -2104,6 +2103,12 @@ void FSceneRenderTargets::AllocateMobileRenderTargets(FRHICommandList& RHICmdLis
 	{
 		SceneAlphaCopy = GSystemTextures.MaxFP16Depth;
 	}
+	
+	if (UseVirtualTexturing(CurrentFeatureLevel))
+	{
+		FIntPoint FeedbackSize = FIntPoint::DivideAndRoundUp(BufferSize, FMath::Max(GVirtualTextureFeedbackFactor, 1));
+		VirtualTextureFeedback.CreateResourceGPU(RHICmdList, FeedbackSize);
+	}
 }
 
 // This is a helper class. It generates and provides N names with
@@ -3234,6 +3239,15 @@ void SetupMobileSceneTextureUniformParameters(
 
 	SceneTextureParameters.MobileCustomStencilTexture = MobileCustomStencil;
 	SceneTextureParameters.MobileCustomStencilTextureSampler = TStaticSamplerState<SF_Point, AM_Clamp, AM_Clamp, AM_Clamp>::GetRHI();
+	
+	if (SceneContext.VirtualTextureFeedback.FeedbackBufferUAV.IsValid())
+	{
+		SceneTextureParameters.VirtualTextureFeedbackUAV = SceneContext.VirtualTextureFeedback.FeedbackBufferUAV;
+	}
+	else
+	{
+		SceneTextureParameters.VirtualTextureFeedbackUAV = GVirtualTextureFeedbackDummyResource.UAV;
+	}
 }
 
 template< typename TRHICmdList >
diff --git a/Engine/Source/Runtime/Renderer/Private/PostProcess/SceneRenderTargets.h b/Engine/Source/Runtime/Renderer/Private/PostProcess/SceneRenderTargets.h
--- a/Engine/Source/Runtime/Renderer/Private/PostProcess/SceneRenderTargets.h
+++ b/Engine/Source/Runtime/Renderer/Private/PostProcess/SceneRenderTargets.h
@@ -632,7 +632,7 @@ private:
 	void InitEditorPrimitivesDepth(FRHICommandList& RHICmdList);
 
 	/** Allocates render targets for use with the mobile path. */
-	void AllocateMobileRenderTargets(FRHICommandList& RHICmdList);
+	void AllocateMobileRenderTargets(FRHICommandListImmediate& RHICmdList);
 
 public:
 	/** Allocates render targets for use with the deferred shading path. */
diff --git a/Engine/Source/Runtime/Renderer/Private/SceneRendering.cpp b/Engine/Source/Runtime/Renderer/Private/SceneRendering.cpp
--- a/Engine/Source/Runtime/Renderer/Private/SceneRendering.cpp
+++ b/Engine/Source/Runtime/Renderer/Private/SceneRendering.cpp
@@ -1643,6 +1643,9 @@ void FViewInfo::SetupUniformBufferParameters(
 
 	ViewUniformShaderParameters.PreIntegratedBRDF = GEngine->PreIntegratedSkinBRDFTexture->Resource->TextureRHI;
 
+	ViewUniformShaderParameters.VirtualTextureParams = FVector4(ForceInitToZero);
+	ViewUniformShaderParameters.VirtualTextureFeedbackStride = SceneContext.VirtualTextureFeedback.GetFeedbackStride();
+	
 	if (UseGPUScene(GMaxRHIShaderPlatform, RHIFeatureLevel))
 	{
 		if (PrimitiveSceneDataOverrideSRV)
diff --git a/Engine/Source/Runtime/Renderer/Private/VT/RuntimeVirtualTextureRender.cpp b/Engine/Source/Runtime/Renderer/Private/VT/RuntimeVirtualTextureRender.cpp
--- a/Engine/Source/Runtime/Renderer/Private/VT/RuntimeVirtualTextureRender.cpp
+++ b/Engine/Source/Runtime/Renderer/Private/VT/RuntimeVirtualTextureRender.cpp
@@ -28,7 +28,7 @@ namespace RuntimeVirtualTexture
 
 		static bool ShouldCompilePermutation(const FMeshMaterialShaderPermutationParameters& Parameters)
 		{
-			return IsFeatureLevelSupported(Parameters.Platform, ERHIFeatureLevel::SM5) && 
+			return IsFeatureLevelSupported(Parameters.Platform, ERHIFeatureLevel::ES3_1) && 
 				(Parameters.Material->GetMaterialDomain() == MD_RuntimeVirtualTexture || Parameters.Material->HasRuntimeVirtualTextureOutput());
 		}
 
@@ -271,6 +271,7 @@ namespace RuntimeVirtualTexture
 	}
 
 	FRegisterPassProcessorCreateFunction RegisterVirtualTexturePass(&CreateRuntimeVirtualTexturePassProcessor, EShadingPath::Deferred, EMeshPass::VirtualTexture, EMeshPassFlags::CachedMeshCommands);
+	FRegisterPassProcessorCreateFunction RegisterVirtualTexturePassMobile(&CreateRuntimeVirtualTexturePassProcessor, EShadingPath::Mobile, EMeshPass::VirtualTexture, EMeshPassFlags::CachedMeshCommands);
 
 
 	/** Collect meshes and draw. */
@@ -427,7 +428,7 @@ namespace RuntimeVirtualTexture
 
 		static bool ShouldCompilePermutation(FGlobalShaderPermutationParameters const& Parameters)
 		{
-			return IsFeatureLevelSupported(Parameters.Platform, ERHIFeatureLevel::SM5);
+			return RHISupportsComputeShaders(Parameters.Platform);
 		}
 
 		FShader_VirtualTextureCompress()
@@ -515,7 +516,7 @@ namespace RuntimeVirtualTexture
 
 		static bool ShouldCompilePermutation(FGlobalShaderPermutationParameters const& Parameters)
 		{
-			return IsFeatureLevelSupported(Parameters.Platform, ERHIFeatureLevel::SM5);
+			return IsFeatureLevelSupported(Parameters.Platform, ERHIFeatureLevel::ES3_1);
 		}
 
 		FShader_VirtualTextureCopy()
@@ -642,13 +643,15 @@ namespace RuntimeVirtualTexture
 			bCopyThumbnailPass = bRenderPass && bIsThumbnails;
 			bCompressPass = bRenderPass && !bCopyThumbnailPass && (OutputTexture0->GetFormat() == PF_DXT1 || OutputTexture0->GetFormat() == PF_DXT5 || OutputTexture0->GetFormat() == PF_BC5);
 			bCopyPass = bRenderPass && !bCopyThumbnailPass && !bCompressPass && (MaterialType == ERuntimeVirtualTextureMaterialType::BaseColor_Normal_Specular || MaterialType == ERuntimeVirtualTextureMaterialType::BaseColor_Normal_Specular_YCoCg || MaterialType == ERuntimeVirtualTextureMaterialType::BaseColor_Normal_Specular_Mask_YCoCg);
+			// Not all mobile RHIs support sRGB texture views/aliasing, use only linear targets on mobile
+			uint32 VT_SRGB = GMaxRHIFeatureLevel > ERHIFeatureLevel::ES3_1 ? TexCreate_SRGB : TexCreate_None;
 		
 			switch (MaterialType)
 			{
 			case ERuntimeVirtualTextureMaterialType::BaseColor:
 				if (bRenderPass)
 				{
-					OutputAlias0 = RenderTexture0 = GraphBuilder.CreateTexture(FPooledRenderTargetDesc::Create2DDesc(TextureSize, PF_B8G8R8A8, FClearValueBinding::Black, TexCreate_SRGB, TexCreate_RenderTargetable | TexCreate_ShaderResource, false), TEXT("RenderTexture0"));
+					OutputAlias0 = RenderTexture0 = GraphBuilder.CreateTexture(FPooledRenderTargetDesc::Create2DDesc(TextureSize, PF_B8G8R8A8, FClearValueBinding::Black, VT_SRGB, TexCreate_RenderTargetable | TexCreate_ShaderResource, false), TEXT("RenderTexture0"));
 				}
 				if (bCompressPass)
 				{
@@ -656,13 +659,13 @@ namespace RuntimeVirtualTexture
 				}
 				if (bCopyThumbnailPass)
 				{
-					OutputAlias0 = CopyTexture0 = GraphBuilder.CreateTexture(FPooledRenderTargetDesc::Create2DDesc(TextureSize, PF_B8G8R8A8, FClearValueBinding::Black, TexCreate_SRGB, TexCreate_RenderTargetable | TexCreate_ShaderResource, false), TEXT("CopyTexture0"));
+					OutputAlias0 = CopyTexture0 = GraphBuilder.CreateTexture(FPooledRenderTargetDesc::Create2DDesc(TextureSize, PF_B8G8R8A8, FClearValueBinding::Black, VT_SRGB, TexCreate_RenderTargetable | TexCreate_ShaderResource, false), TEXT("CopyTexture0"));
 				}
 				break;
 			case ERuntimeVirtualTextureMaterialType::BaseColor_Normal_Specular:
 				if (bRenderPass)
 				{
-					OutputAlias0 = RenderTexture0 = GraphBuilder.CreateTexture(FPooledRenderTargetDesc::Create2DDesc(TextureSize, PF_B8G8R8A8, FClearValueBinding::Black, TexCreate_SRGB, TexCreate_RenderTargetable | TexCreate_ShaderResource, false), TEXT("RenderTexture0"));
+					OutputAlias0 = RenderTexture0 = GraphBuilder.CreateTexture(FPooledRenderTargetDesc::Create2DDesc(TextureSize, PF_B8G8R8A8, FClearValueBinding::Black, VT_SRGB, TexCreate_RenderTargetable | TexCreate_ShaderResource, false), TEXT("RenderTexture0"));
 					RenderTexture1 = GraphBuilder.CreateTexture(FPooledRenderTargetDesc::Create2DDesc(TextureSize, PF_B8G8R8A8, FClearValueBinding::DefaultNormal8Bit, TexCreate_None, TexCreate_RenderTargetable | TexCreate_ShaderResource, false), TEXT("RenderTexture1"));
 					RenderTexture2 = GraphBuilder.CreateTexture(FPooledRenderTargetDesc::Create2DDesc(TextureSize, PF_B8G8R8A8, FClearValueBinding::DefaultNormal8Bit, TexCreate_None, TexCreate_RenderTargetable | TexCreate_ShaderResource, false), TEXT("RenderTexture2"));
 				}
@@ -673,18 +676,18 @@ namespace RuntimeVirtualTexture
 				}
 				if (bCopyPass)
 				{
-					OutputAlias0 = CopyTexture0 = GraphBuilder.CreateTexture(FPooledRenderTargetDesc::Create2DDesc(TextureSize, PF_B8G8R8A8, FClearValueBinding::Black, TexCreate_SRGB, TexCreate_RenderTargetable | TexCreate_ShaderResource, false), TEXT("CopyTexture0"));
+					OutputAlias0 = CopyTexture0 = GraphBuilder.CreateTexture(FPooledRenderTargetDesc::Create2DDesc(TextureSize, PF_B8G8R8A8, FClearValueBinding::Black, VT_SRGB, TexCreate_RenderTargetable | TexCreate_ShaderResource, false), TEXT("CopyTexture0"));
 					OutputAlias1 = CopyTexture1 = GraphBuilder.CreateTexture(FPooledRenderTargetDesc::Create2DDesc(TextureSize, PF_B8G8R8A8, FClearValueBinding::Black, TexCreate_None, TexCreate_RenderTargetable | TexCreate_ShaderResource, false), TEXT("CopyTexture1"));
 				}
 				if (bCopyThumbnailPass)
 				{
-					OutputAlias0 = CopyTexture0 = GraphBuilder.CreateTexture(FPooledRenderTargetDesc::Create2DDesc(TextureSize, PF_B8G8R8A8, FClearValueBinding::Black, TexCreate_SRGB, TexCreate_RenderTargetable | TexCreate_ShaderResource, false), TEXT("CopyTexture0"));
+					OutputAlias0 = CopyTexture0 = GraphBuilder.CreateTexture(FPooledRenderTargetDesc::Create2DDesc(TextureSize, PF_B8G8R8A8, FClearValueBinding::Black, VT_SRGB, TexCreate_RenderTargetable | TexCreate_ShaderResource, false), TEXT("CopyTexture0"));
 				}
 				break;
 			case ERuntimeVirtualTextureMaterialType::BaseColor_Normal_Specular_YCoCg:
 				if (bRenderPass)
 				{
-					OutputAlias0 = RenderTexture0 = GraphBuilder.CreateTexture(FPooledRenderTargetDesc::Create2DDesc(TextureSize, PF_B8G8R8A8, FClearValueBinding::Black, TexCreate_SRGB, TexCreate_RenderTargetable | TexCreate_ShaderResource, false), TEXT("RenderTexture0"));
+					OutputAlias0 = RenderTexture0 = GraphBuilder.CreateTexture(FPooledRenderTargetDesc::Create2DDesc(TextureSize, PF_B8G8R8A8, FClearValueBinding::Black, VT_SRGB, TexCreate_RenderTargetable | TexCreate_ShaderResource, false), TEXT("RenderTexture0"));
 					RenderTexture1 = GraphBuilder.CreateTexture(FPooledRenderTargetDesc::Create2DDesc(TextureSize, PF_B8G8R8A8, FClearValueBinding::DefaultNormal8Bit, TexCreate_None, TexCreate_RenderTargetable | TexCreate_ShaderResource, false), TEXT("RenderTexture1"));
 					RenderTexture2 = GraphBuilder.CreateTexture(FPooledRenderTargetDesc::Create2DDesc(TextureSize, PF_B8G8R8A8, FClearValueBinding::DefaultNormal8Bit, TexCreate_None, TexCreate_RenderTargetable | TexCreate_ShaderResource, false), TEXT("RenderTexture2"));
 				}
@@ -702,13 +705,13 @@ namespace RuntimeVirtualTexture
 				}
 				if (bCopyThumbnailPass)
 				{
-					OutputAlias0 = CopyTexture0 = GraphBuilder.CreateTexture(FPooledRenderTargetDesc::Create2DDesc(TextureSize, PF_B8G8R8A8, FClearValueBinding::Black, TexCreate_SRGB, TexCreate_RenderTargetable | TexCreate_ShaderResource, false), TEXT("CopyTexture0"));
+					OutputAlias0 = CopyTexture0 = GraphBuilder.CreateTexture(FPooledRenderTargetDesc::Create2DDesc(TextureSize, PF_B8G8R8A8, FClearValueBinding::Black, VT_SRGB, TexCreate_RenderTargetable | TexCreate_ShaderResource, false), TEXT("CopyTexture0"));
 				}
 				break;
 			case ERuntimeVirtualTextureMaterialType::BaseColor_Normal_Specular_Mask_YCoCg:
 				if (bRenderPass)
 				{
-					OutputAlias0 = RenderTexture0 = GraphBuilder.CreateTexture(FPooledRenderTargetDesc::Create2DDesc(TextureSize, PF_B8G8R8A8, FClearValueBinding::Black, TexCreate_SRGB, TexCreate_RenderTargetable | TexCreate_ShaderResource, false), TEXT("RenderTexture0"));
+					OutputAlias0 = RenderTexture0 = GraphBuilder.CreateTexture(FPooledRenderTargetDesc::Create2DDesc(TextureSize, PF_B8G8R8A8, FClearValueBinding::Black, VT_SRGB, TexCreate_RenderTargetable | TexCreate_ShaderResource, false), TEXT("RenderTexture0"));
 					RenderTexture1 = GraphBuilder.CreateTexture(FPooledRenderTargetDesc::Create2DDesc(TextureSize, PF_B8G8R8A8, FClearValueBinding::DefaultNormal8Bit, TexCreate_None, TexCreate_RenderTargetable | TexCreate_ShaderResource, false), TEXT("RenderTexture1"));
 					RenderTexture2 = GraphBuilder.CreateTexture(FPooledRenderTargetDesc::Create2DDesc(TextureSize, PF_B8G8R8A8, FClearValueBinding::DefaultNormal8Bit, TexCreate_None, TexCreate_RenderTargetable | TexCreate_ShaderResource, false), TEXT("RenderTexture2"));
 				}
@@ -726,7 +729,7 @@ namespace RuntimeVirtualTexture
 				}
 				if (bCopyThumbnailPass)
 				{
-					OutputAlias0 = CopyTexture0 = GraphBuilder.CreateTexture(FPooledRenderTargetDesc::Create2DDesc(TextureSize, PF_B8G8R8A8, FClearValueBinding::Black, TexCreate_SRGB, TexCreate_RenderTargetable | TexCreate_ShaderResource, false), TEXT("CopyTexture0"));
+					OutputAlias0 = CopyTexture0 = GraphBuilder.CreateTexture(FPooledRenderTargetDesc::Create2DDesc(TextureSize, PF_B8G8R8A8, FClearValueBinding::Black, VT_SRGB, TexCreate_RenderTargetable | TexCreate_ShaderResource, false), TEXT("CopyTexture0"));
 				}
 				break;
 			case ERuntimeVirtualTextureMaterialType::WorldHeight:
diff --git a/Engine/Source/Runtime/Renderer/Private/VT/VirtualTextureFeedback.cpp b/Engine/Source/Runtime/Renderer/Private/VT/VirtualTextureFeedback.cpp
--- a/Engine/Source/Runtime/Renderer/Private/VT/VirtualTextureFeedback.cpp
+++ b/Engine/Source/Runtime/Renderer/Private/VT/VirtualTextureFeedback.cpp
@@ -127,14 +127,23 @@ public:
 };
 #endif
 
+#include "RHIGPUReadback.h"
+
+TGlobalResource<FVirtualTextureFeedbackDummyResource> GVirtualTextureFeedbackDummyResource;
 
 FVirtualTextureFeedback::FVirtualTextureFeedback()
-	: Size( 0, 0 )
+	: Size(0, 0)
+	, NumBytes(0)
 	, GPUWriteIndex(0)
 	, CPUReadIndex(0)
 	, PendingTargetCount(0)
 {
 	FeedBackFences = new FFeedbackFences;
+
+	for (int i = 0; i < TargetCapacity; ++i)
+	{
+		FeedbackCPU[i].ReadbackBuffer = nullptr;
+	}
 }
 
 FVirtualTextureFeedback::~FVirtualTextureFeedback()
@@ -144,15 +153,15 @@ FVirtualTextureFeedback::~FVirtualTextureFeedback()
 
 void FVirtualTextureFeedback::ReleaseResources()
 {
-	GRenderTargetPool.FreeUnusedResource( FeedbackTextureGPU );
-	FeedbackTextureGPU.SafeRelease();
+	FeedbackBufferUAV.SafeRelease();
+	FeedbackBuffer.SafeRelease();
 	
 	for (int i = 0; i < TargetCapacity; ++i)
 	{
-		if (FeedbackTextureCPU[i].TextureCPU.IsValid())
+		if (FeedbackCPU[i].ReadbackBuffer)
 		{
-			GRenderTargetPool.FreeUnusedResource(FeedbackTextureCPU[i].TextureCPU);
-			FeedbackTextureCPU[i].TextureCPU.SafeRelease();
+			delete FeedbackCPU[i].ReadbackBuffer;
+			FeedbackCPU[i].ReadbackBuffer = nullptr;
 		}
 	}
 
@@ -165,46 +174,56 @@ void FVirtualTextureFeedback::ReleaseResources()
 
 void FVirtualTextureFeedback::CreateResourceGPU( FRHICommandListImmediate& RHICmdList, FIntPoint InSize)
 {
-	Size = InSize;
-
-	FPooledRenderTargetDesc Desc( FPooledRenderTargetDesc::Create2DDesc( Size, PF_R32_UINT, FClearValueBinding::None, TexCreate_None, TexCreate_ShaderResource | TexCreate_UAV, false ) );
-	GRenderTargetPool.FindFreeElement( RHICmdList, Desc, FeedbackTextureGPU, TEXT("VTFeedbackGPU") );
+	if (Size != InSize || !FeedbackBuffer.IsValid())
+	{
+		Size = InSize;
+		NumBytes = Size.X * Size.Y * sizeof(uint32);
 
+		FRHIResourceCreateInfo CreateInfo(TEXT("VTFeedbackGPU"));
+		FeedbackBuffer = RHICreateVertexBuffer(NumBytes, BUF_Static | BUF_ShaderResource | BUF_UnorderedAccess, CreateInfo);
+		FeedbackBufferUAV = RHICreateUnorderedAccessView(FeedbackBuffer, /*Format=*/ PF_R32_UINT);
+	}
+	
 	// Clear to default value
 	const uint32 ClearValue[4] = { ~0u, ~0u, ~0u, ~0u };
-	ClearUAV( RHICmdList, FeedbackTextureGPU->GetRenderTargetItem(), ClearValue );
-	RHICmdList.TransitionResource( EResourceTransitionAccess::ERWNoBarrier, EResourceTransitionPipeline::EGfxToGfx, FeedbackTextureGPU->GetRenderTargetItem().UAV );
+	ClearUAV(RHICmdList, FeedbackBufferUAV.GetReference(), NumBytes, ~0u);
+	RHICmdList.TransitionResource(EResourceTransitionAccess::ERWNoBarrier, EResourceTransitionPipeline::EGfxToGfx, FeedbackBufferUAV);
 
 	FeedBackFences->Init(RHICmdList);
 }
 
-void FVirtualTextureFeedback::MakeSnapshot(const FVirtualTextureFeedback& SnapshotSource)
-{
-	Size = SnapshotSource.Size;
-	FeedbackTextureGPU = GRenderTargetPool.MakeSnapshot(SnapshotSource.FeedbackTextureGPU);
-	for (int i = 0; i < TargetCapacity; ++i)
-	{
-		FeedbackTextureCPU[i].TextureCPU = GRenderTargetPool.MakeSnapshot(SnapshotSource.FeedbackTextureCPU[i].TextureCPU);
-	}
-}
+const static FName VTReadbackRequestNames[FVirtualTextureFeedback::TargetCapacity] = 
+{ 
+	FName("VTFeedbackCPU_0"), 
+	FName("VTFeedbackCPU_1"), 
+	FName("VTFeedbackCPU_2"), 
+	FName("VTFeedbackCPU_3") 
+};
 
-void FVirtualTextureFeedback::TransferGPUToCPU( FRHICommandListImmediate& RHICmdList, TArrayView<FIntRect> const& ViewRects )
+void FVirtualTextureFeedback::TransferGPUToCPU( FRHICommandListImmediate& RHICmdList, TArrayView<FIntRect> const& ViewRects)
 {
-	RHICmdList.TransitionResource( EResourceTransitionAccess::EReadable, EResourceTransitionPipeline::EGfxToGfx, FeedbackTextureGPU->GetRenderTargetItem().UAV );
+	QUICK_SCOPE_CYCLE_COUNTER(STAT_VirtualTextureFeedback_TransferGPUToCPU);
+
+	RHICmdList.TransitionResource( EResourceTransitionAccess::EReadable, EResourceTransitionPipeline::EGfxToGfx, FeedbackBufferUAV );
 
-	GVisualizeTexture.SetCheckPoint( RHICmdList, FeedbackTextureGPU );
-	
 	if (PendingTargetCount >= TargetCapacity)
 	{
 		// If we have too many pending transfers, start throwing away the oldest
 		// We will need to allocate a new fence, since the previous fence will still be set on the old CopyToResolveTarget command (which we will now ignore/discard)
 		FeedBackFences->Release(CPUReadIndex);
-
 		--PendingTargetCount;
 		CPUReadIndex = (CPUReadIndex + 1) % TargetCapacity;
 	}
 
-	FFeedBackItem& FeedbackEntryCPU = FeedbackTextureCPU[GPUWriteIndex];
+	FFeedBackItem& FeedbackEntryCPU = FeedbackCPU[GPUWriteIndex];
+
+	///////////////////////////////
+	if (FeedbackEntryCPU.ReadbackBuffer == nullptr)
+	{
+		FName RequestName = VTReadbackRequestNames[GPUWriteIndex];
+		FeedbackEntryCPU.ReadbackBuffer = new FRHIGPUBufferReadback(RequestName);
+	}
+
 	FeedbackEntryCPU.NumRects = FMath::Min((int32)MaxRectPerTarget, ViewRects.Num());
 	for (int32 RectIndex = 0; RectIndex < FeedbackEntryCPU.NumRects; ++RectIndex)
 	{
@@ -213,25 +232,14 @@ void FVirtualTextureFeedback::TransferGPUToCPU( FRHICommandListImmediate& RHICmd
 		FeedbackEntryCPU.Rects[RectIndex].Min = FIntPoint(Rect.Min.X / 16, Rect.Min.Y / 16);
 		FeedbackEntryCPU.Rects[RectIndex].Max = FIntPoint((Rect.Max.X + 15) / 16, (Rect.Max.Y + 15) / 16);
 	}
-
-	const FPooledRenderTargetDesc Desc(FPooledRenderTargetDesc::Create2DDesc(Size, PF_R32_UINT, FClearValueBinding::None, TexCreate_CPUReadback | TexCreate_HideInVisualizeTexture, TexCreate_None, false));
-
-	const TCHAR* DebugNames[TargetCapacity] = { TEXT("VTFeedbackCPU_0") , TEXT("VTFeedbackCPU_1") , TEXT("VTFeedbackCPU_2") , TEXT("VTFeedbackCPU_3") };
-	GRenderTargetPool.FindFreeElement(RHICmdList, Desc, FeedbackEntryCPU.TextureCPU, DebugNames[GPUWriteIndex]);
 	
-	FeedBackFences->Allocate(RHICmdList, GPUWriteIndex);
-
 	// We only need to transfer 1 copy of the data, so restrict mask to the first active GPU.
 	FeedbackEntryCPU.GPUMask = FRHIGPUMask::FromIndex(RHICmdList.GetGPUMask().GetFirstIndex());
 	SCOPED_GPU_MASK(RHICmdList, FeedbackEntryCPU.GPUMask);
 
-	// Transfer memory GPU -> CPU
-	RHICmdList.CopyToResolveTarget(
-		FeedbackTextureGPU->GetRenderTargetItem().TargetableTexture,
-		FeedbackEntryCPU.TextureCPU->GetRenderTargetItem().ShaderResourceTexture,
-		FResolveParams());
-	
-	FeedBackFences->Write(RHICmdList, GPUWriteIndex);
+	// copy whole feedback buffer
+	FeedbackEntryCPU.ReadbackBuffer->EnqueueCopy(RHICmdList, FeedbackBuffer, NumBytes);
+	///////////////////////////////
 
 	GPUWriteIndex = (GPUWriteIndex + 1) % TargetCapacity;
 	++PendingTargetCount;
@@ -239,48 +247,40 @@ void FVirtualTextureFeedback::TransferGPUToCPU( FRHICommandListImmediate& RHICmd
 
 bool FVirtualTextureFeedback::CanMap(FRHICommandListImmediate& RHICmdList)
 {
-	const FFeedBackItem& FeedbackEntryCPU = FeedbackTextureCPU[CPUReadIndex];
-	return (
-		PendingTargetCount > 0u &&
-		FeedbackEntryCPU.TextureCPU.IsValid() &&
-		FeedBackFences->Poll(RHICmdList, CPUReadIndex));
+	const FFeedBackItem& FeedbackEntryCPU = FeedbackCPU[CPUReadIndex];
+	return (PendingTargetCount > 0u && FeedbackEntryCPU.ReadbackBuffer && FeedbackEntryCPU.ReadbackBuffer->IsReady());
 }
 
 bool FVirtualTextureFeedback::Map(FRHICommandListImmediate& RHICmdList, MapResult& OutResult)
 {
-	const FFeedBackItem& FeedbackEntryCPU = FeedbackTextureCPU[CPUReadIndex];
+	QUICK_SCOPE_CYCLE_COUNTER(STAT_VirtualTextureFeedback_Map);
+		
+	FFeedBackItem& FeedbackEntryCPU = FeedbackCPU[CPUReadIndex];
+		
 	if (PendingTargetCount > 0u &&
-		FeedbackEntryCPU.TextureCPU.IsValid() &&
-		FeedBackFences->Poll(RHICmdList, CPUReadIndex))
+		FeedbackEntryCPU.ReadbackBuffer &&
+		FeedbackEntryCPU.ReadbackBuffer->IsReady())
 	{
 		SCOPED_GPU_MASK(RHICmdList, FeedbackEntryCPU.GPUMask);
-
-		OutResult.MapHandle = CPUReadIndex;
-		
+	
 		OutResult.NumRects = FeedbackEntryCPU.NumRects;
 		for (int32 i = 0; i < FeedbackEntryCPU.NumRects; ++i)
 		{
 			OutResult.Rects[i] = FeedbackEntryCPU.Rects[i];
 		}
-
-		int32 LockHeight = 0;
-		RHICmdList.MapStagingSurface(FeedbackEntryCPU.TextureCPU->GetRenderTargetItem().ShaderResourceTexture, FeedBackFences->GetMapFence(CPUReadIndex), *(void**)&OutResult.Buffer, OutResult.Pitch, LockHeight);
-
+		OutResult.Pitch = Size.X;
+		OutResult.Buffer.SetNumUninitialized(NumBytes/sizeof(uint32));
+		
+		void* MappedMem = FeedbackEntryCPU.ReadbackBuffer->Lock(NumBytes);
+		FMemory::Memcpy(OutResult.Buffer.GetData(), MappedMem, NumBytes);
+		FeedbackEntryCPU.ReadbackBuffer->Unlock();
+		
 		check(PendingTargetCount > 0u);
 		--PendingTargetCount;
 		CPUReadIndex = (CPUReadIndex + 1) % TargetCapacity;
-
 		return true;
 	}
 
 	return false;
 }
 
-void FVirtualTextureFeedback::Unmap( FRHICommandListImmediate& RHICmdList, int32 MapHandle )
-{
-	check(FeedbackTextureCPU[MapHandle].TextureCPU.IsValid());
-
-	SCOPED_GPU_MASK(RHICmdList, FeedbackTextureCPU[MapHandle].GPUMask);
-	RHICmdList.UnmapStagingSurface(FeedbackTextureCPU[MapHandle].TextureCPU->GetRenderTargetItem().ShaderResourceTexture);
-	GRenderTargetPool.FreeUnusedResource(FeedbackTextureCPU[MapHandle].TextureCPU);
-}
diff --git a/Engine/Source/Runtime/Renderer/Private/VT/VirtualTextureFeedback.h b/Engine/Source/Runtime/Renderer/Private/VT/VirtualTextureFeedback.h
--- a/Engine/Source/Runtime/Renderer/Private/VT/VirtualTextureFeedback.h
+++ b/Engine/Source/Runtime/Renderer/Private/VT/VirtualTextureFeedback.h
@@ -5,6 +5,8 @@
 #include "CoreMinimal.h"
 #include "RenderCore/Public/RenderTargetPool.h"
 
+class FRHIGPUBufferReadback;
+
 /*
 ====================================
 	Manages GPU and CPU buffers for VT feedback.
@@ -23,44 +25,69 @@ public:
 	static const uint32 TargetCapacity = 4u;
 	static const uint32 MaxRectPerTarget = 4u;
 
-	TRefCountPtr< IPooledRenderTarget >	FeedbackTextureGPU;
-
 	struct MapResult
 	{
-		int32 MapHandle;
-		uint32* RESTRICT Buffer;
+		TArray<uint32> Buffer;
 		int32 Pitch;
 		int32 NumRects;
 		FIntRect Rects[MaxRectPerTarget];
 	};
 
+	FVertexBufferRHIRef FeedbackBuffer;
+	FUnorderedAccessViewRHIRef FeedbackBufferUAV;
+
 	void			CreateResourceGPU(FRHICommandListImmediate& RHICmdList, FIntPoint InSize);
 	void			ReleaseResources();
-	void			MakeSnapshot(const FVirtualTextureFeedback& SnapshotSource);
 
 	void			TransferGPUToCPU(FRHICommandListImmediate& RHICmdList, TArrayView<FIntRect> const& ViewRects);
 
 	uint32			GetPendingTargetCount() const { return PendingTargetCount; }
+	uint32			GetFeedbackStride() const { return (uint32)Size.X; }
 
 	bool			CanMap(FRHICommandListImmediate& RHICmdList);
 	bool			Map(FRHICommandListImmediate& RHICmdList, MapResult& OutResult);
-	void			Unmap(FRHICommandListImmediate& RHICmdList, int32 MapHandle);
 
 private:
 	struct FFeedBackItem
 	{
+		FRHIGPUBufferReadback* ReadbackBuffer;
 		int32 NumRects;
 		FIntRect Rects[MaxRectPerTarget];
-		TRefCountPtr< IPooledRenderTarget > TextureCPU;
 		FRHIGPUMask GPUMask;
 	};
 
-	FFeedBackItem FeedbackTextureCPU[TargetCapacity];
+	FFeedBackItem FeedbackCPU[TargetCapacity];
 
 	class FFeedbackFences* FeedBackFences;
 
 	FIntPoint Size;
+	uint32 NumBytes;
+
 	uint32 GPUWriteIndex;
 	uint32 CPUReadIndex;
 	uint32 PendingTargetCount;
 };
+
+
+/* Dummy resource for VT feedback, used to bypass UB validation errors when VT is disabled */
+class FVirtualTextureFeedbackDummyResource : public FRenderResource
+{
+public:
+	FVertexBufferRHIRef			VertexBufferRHI;
+	FUnorderedAccessViewRHIRef	UAV;
+
+	virtual void InitRHI() override
+	{
+		FRHIResourceCreateInfo CreateInfo(TEXT("VTFeedbackDummy"));
+		VertexBufferRHI = RHICreateVertexBuffer(4, BUF_Static | BUF_ShaderResource | BUF_UnorderedAccess, CreateInfo);
+		UAV = RHICreateUnorderedAccessView(VertexBufferRHI, /*Format=*/ PF_R32_UINT);
+	}
+	virtual void ReleaseRHI() override
+	{
+		UAV.SafeRelease();
+		VertexBufferRHI.SafeRelease();
+	}
+};
+
+extern TGlobalResource<FVirtualTextureFeedbackDummyResource> GVirtualTextureFeedbackDummyResource;
+
diff --git a/Engine/Source/Runtime/Renderer/Private/VT/VirtualTexturePhysicalSpace.cpp b/Engine/Source/Runtime/Renderer/Private/VT/VirtualTexturePhysicalSpace.cpp
--- a/Engine/Source/Runtime/Renderer/Private/VT/VirtualTexturePhysicalSpace.cpp
+++ b/Engine/Source/Runtime/Renderer/Private/VT/VirtualTexturePhysicalSpace.cpp
@@ -72,13 +72,26 @@ EPixelFormat GetUnorderedAccessViewFormat(EPixelFormat InFormat)
 	return InFormat;
 }
 
+EPixelFormat RemapVirtualTexturePhysicalSpaceFormat(EPixelFormat InFormat)
+{
+	if (InFormat == PF_B8G8R8A8 && IsOpenGLPlatform(GMaxRHIShaderPlatform) && IsMobilePlatform(GMaxRHIShaderPlatform))
+	{
+		// FIXME: Mobile/Android OpenGL can't copy data between swizzled formats
+		// Always use RGBA format for both VT intermediate render targets and VT physical texture
+		// This will also make uncompressed streaming VT to have a R and B channel swapped 
+		return PF_R8G8B8A8;
+	}
+
+	return InFormat;
+}
+
 void FVirtualTexturePhysicalSpace::InitRHI()
 {
 	FRHICommandListImmediate& RHICmdList = FRHICommandListExecutor::GetImmediateCommandList();
 
 	for (int32 Layer = 0; Layer < Description.NumLayers; ++Layer)
 	{
-		const EPixelFormat FormatSRV = Description.Format[Layer];
+		const EPixelFormat FormatSRV = RemapVirtualTexturePhysicalSpaceFormat(Description.Format[Layer]);
 		const EPixelFormat FormatUAV = GetUnorderedAccessViewFormat(FormatSRV);
 		const bool bCreateAliasedUAV = (FormatUAV != PF_Unknown) && (FormatUAV != FormatSRV);
 
@@ -100,8 +113,16 @@ void FVirtualTexturePhysicalSpace::InitRHI()
 		SRVCreateInfo.Format = FormatSRV;
 		TextureSRV[Layer] = RHICreateShaderResourceView(TextureRHI, SRVCreateInfo);
 
-		SRVCreateInfo.SRGBOverride = SRGBO_ForceEnable;
-		TextureSRV_SRGB[Layer] = RHICreateShaderResourceView(TextureRHI, SRVCreateInfo);
+		if (GMaxRHIFeatureLevel > ERHIFeatureLevel::ES3_1)
+		{
+			SRVCreateInfo.SRGBOverride = SRGBO_ForceEnable;
+			TextureSRV_SRGB[Layer] = RHICreateShaderResourceView(TextureRHI, SRVCreateInfo);
+		}
+		else
+		{
+			// Not all mobile RHIs support sRGB views
+			TextureSRV_SRGB[Layer] = TextureSRV[Layer];
+		}
 
 		if (bCreateAliasedUAV)
 		{
diff --git a/Engine/Source/Runtime/Renderer/Private/VT/VirtualTextureSpace.cpp b/Engine/Source/Runtime/Renderer/Private/VT/VirtualTextureSpace.cpp
--- a/Engine/Source/Runtime/Renderer/Private/VT/VirtualTextureSpace.cpp
+++ b/Engine/Source/Runtime/Renderer/Private/VT/VirtualTextureSpace.cpp
@@ -184,7 +184,8 @@ public:
 
 	static bool ShouldCompilePermutation(const FGlobalShaderPermutationParameters& Parameters)
 	{
-		return IsFeatureLevelSupported(Parameters.Platform, ERHIFeatureLevel::SM5);
+		return IsFeatureLevelSupported(Parameters.Platform, ERHIFeatureLevel::SM5) 
+			|| GetMaxSupportedFeatureLevel(Parameters.Platform) == ERHIFeatureLevel::ES3_1;
 	}
 
 	virtual bool Serialize( FArchive& Ar ) override
@@ -211,7 +212,8 @@ public:
 	
 	static bool ShouldCompilePermutation(const FGlobalShaderPermutationParameters& Parameters)
 	{
-		return IsFeatureLevelSupported(Parameters.Platform, ERHIFeatureLevel::SM5);
+		return IsFeatureLevelSupported(Parameters.Platform, ERHIFeatureLevel::SM5)
+			|| GetMaxSupportedFeatureLevel(Parameters.Platform) == ERHIFeatureLevel::ES3_1;
 	}
 };
 
@@ -255,12 +257,12 @@ public:
 
 IMPLEMENT_SHADER_TYPE(template<>, TPageTableUpdateVS<false>, TEXT("/Engine/Private/PageTableUpdate.usf"), TEXT("PageTableUpdateVS"), SF_Vertex);
 IMPLEMENT_SHADER_TYPE(template<>, TPageTableUpdateVS<true>, TEXT("/Engine/Private/PageTableUpdate.usf"), TEXT("PageTableUpdateVS"), SF_Vertex);
-IMPLEMENT_SHADER_TYPE(template<>, TPageTableUpdatePS<PF_R16_UINT>, TEXT("/Engine/Private/PageTableUpdate.usf"), TEXT("PageTableUpdatePS"), SF_Pixel);
-IMPLEMENT_SHADER_TYPE(template<>, TPageTableUpdatePS<PF_R16G16_UINT>, TEXT("/Engine/Private/PageTableUpdate.usf"), TEXT("PageTableUpdatePS"), SF_Pixel);
-IMPLEMENT_SHADER_TYPE(template<>, TPageTableUpdatePS<PF_R16G16B16A16_UINT>, TEXT("/Engine/Private/PageTableUpdate.usf"), TEXT("PageTableUpdatePS"), SF_Pixel);
-IMPLEMENT_SHADER_TYPE(template<>, TPageTableUpdatePS<PF_R32_UINT>, TEXT("/Engine/Private/PageTableUpdate.usf"), TEXT("PageTableUpdatePS"), SF_Pixel);
-IMPLEMENT_SHADER_TYPE(template<>, TPageTableUpdatePS<PF_R32G32_UINT>, TEXT("/Engine/Private/PageTableUpdate.usf"), TEXT("PageTableUpdatePS"), SF_Pixel);
-IMPLEMENT_SHADER_TYPE(template<>, TPageTableUpdatePS<PF_R32G32B32A32_UINT>, TEXT("/Engine/Private/PageTableUpdate.usf"), TEXT("PageTableUpdatePS"), SF_Pixel);
+IMPLEMENT_SHADER_TYPE(template<>, TPageTableUpdatePS<PF_R16_UINT>, TEXT("/Engine/Private/PageTableUpdate.usf"), TEXT("PageTableUpdatePS_1"), SF_Pixel);
+IMPLEMENT_SHADER_TYPE(template<>, TPageTableUpdatePS<PF_R16G16_UINT>, TEXT("/Engine/Private/PageTableUpdate.usf"), TEXT("PageTableUpdatePS_2"), SF_Pixel);
+IMPLEMENT_SHADER_TYPE(template<>, TPageTableUpdatePS<PF_R16G16B16A16_UINT>, TEXT("/Engine/Private/PageTableUpdate.usf"), TEXT("PageTableUpdatePS_4"), SF_Pixel);
+IMPLEMENT_SHADER_TYPE(template<>, TPageTableUpdatePS<PF_R32_UINT>, TEXT("/Engine/Private/PageTableUpdate.usf"), TEXT("PageTableUpdatePS_1"), SF_Pixel);
+IMPLEMENT_SHADER_TYPE(template<>, TPageTableUpdatePS<PF_R32G32_UINT>, TEXT("/Engine/Private/PageTableUpdate.usf"), TEXT("PageTableUpdatePS_2"), SF_Pixel);
+IMPLEMENT_SHADER_TYPE(template<>, TPageTableUpdatePS<PF_R32G32B32A32_UINT>, TEXT("/Engine/Private/PageTableUpdate.usf"), TEXT("PageTableUpdatePS_4"), SF_Pixel);
 
 void FVirtualTextureSpace::QueueUpdateEntirePageTable()
 {
diff --git a/Engine/Source/Runtime/Renderer/Private/VT/VirtualTextureSystem.cpp b/Engine/Source/Runtime/Renderer/Private/VT/VirtualTextureSystem.cpp
--- a/Engine/Source/Runtime/Renderer/Private/VT/VirtualTextureSystem.cpp
+++ b/Engine/Source/Runtime/Renderer/Private/VT/VirtualTextureSystem.cpp
@@ -1010,7 +1010,7 @@ void FVirtualTextureSystem::Update(FRHICommandListImmediate& RHICmdList, ERHIFea
 						{
 							Params.UniquePageList = new(MemStack) FUniquePageList;
 						}
-						Params.FeedbackBuffer = FeedbackBuffer.Buffer + (Rect.Min.Y + CurrentRow) * FeedbackBuffer.Pitch + Rect.Min.X;
+						Params.FeedbackBuffer = FeedbackBuffer.Buffer.GetData() + (Rect.Min.Y + CurrentRow) * FeedbackBuffer.Pitch + Rect.Min.X;
 						Params.FeedbackWidth = Rect.Size().X;
 						Params.FeedbackHeight = CurrentHeight;
 						Params.FeedbackPitch = FeedbackBuffer.Pitch;
@@ -1052,11 +1052,6 @@ void FVirtualTextureSystem::Update(FRHICommandListImmediate& RHICmdList, ERHIFea
 			}
 		}
 
-		for (uint32 i = 0; i < FeedbackBufferCount; ++i)
-		{
-			SceneContext.VirtualTextureFeedback.Unmap(RHICmdList, MappedFeedbackBuffers[i].MapHandle);
-		}
-
 		if(NumFeedbackTasks > 1u)
 		{
 			SCOPE_CYCLE_COUNTER(STAT_ProcessRequests_MergePages);
diff --git a/Engine/Source/Runtime/Renderer/Public/SceneRenderTargetParameters.h b/Engine/Source/Runtime/Renderer/Public/SceneRenderTargetParameters.h
--- a/Engine/Source/Runtime/Renderer/Public/SceneRenderTargetParameters.h
+++ b/Engine/Source/Runtime/Renderer/Public/SceneRenderTargetParameters.h
@@ -113,6 +113,7 @@ BEGIN_GLOBAL_SHADER_PARAMETER_STRUCT(FMobileSceneTextureUniformParameters, RENDE
 	SHADER_PARAMETER_SAMPLER(SamplerState, CustomDepthTextureSampler)
 	SHADER_PARAMETER_TEXTURE(Texture2D, MobileCustomStencilTexture)
 	SHADER_PARAMETER_SAMPLER(SamplerState, MobileCustomStencilTextureSampler)
+	SHADER_PARAMETER_UAV(RWBuffer<uint>, VirtualTextureFeedbackUAV)
 END_GLOBAL_SHADER_PARAMETER_STRUCT()
 
 extern void SetupMobileSceneTextureUniformParameters(
diff --git a/Engine/Source/Runtime/VulkanRHI/Private/VulkanCommands.cpp b/Engine/Source/Runtime/VulkanRHI/Private/VulkanCommands.cpp
--- a/Engine/Source/Runtime/VulkanRHI/Private/VulkanCommands.cpp
+++ b/Engine/Source/Runtime/VulkanRHI/Private/VulkanCommands.cpp
@@ -538,7 +538,10 @@ inline void /*FVulkanCommandListContext::*/SetShaderUniformBufferResources(FVulk
 		case UBMT_SRV:
 		{
 			const VkDescriptorType DescriptorType = DescriptorTypes[GlobalInfos[ResourceInfo.GlobalIndex].TypeIndex];
-			ensure(DescriptorType == VK_DESCRIPTOR_TYPE_UNIFORM_TEXEL_BUFFER || DescriptorType == VK_DESCRIPTOR_TYPE_COMBINED_IMAGE_SAMPLER || DescriptorType == VK_DESCRIPTOR_TYPE_STORAGE_BUFFER);
+			ensure(DescriptorType == VK_DESCRIPTOR_TYPE_UNIFORM_TEXEL_BUFFER 
+				|| DescriptorType == VK_DESCRIPTOR_TYPE_COMBINED_IMAGE_SAMPLER
+				|| DescriptorType == VK_DESCRIPTOR_TYPE_SAMPLED_IMAGE
+				|| DescriptorType == VK_DESCRIPTOR_TYPE_STORAGE_BUFFER);
 			FRHIShaderResourceView* CurrentSRV = (FRHIShaderResourceView*)(ResourceArray[ResourceInfo.SourceUBResourceIndex].GetReference());
 			if (CurrentSRV)
 			{
@@ -555,6 +558,29 @@ inline void /*FVulkanCommandListContext::*/SetShaderUniformBufferResources(FVulk
 			}
 		}
 		break;
+		case UBMT_UAV:
+		{
+			const VkDescriptorType DescriptorType = DescriptorTypes[GlobalInfos[ResourceInfo.GlobalIndex].TypeIndex];
+			ensure(DescriptorType == VK_DESCRIPTOR_TYPE_STORAGE_BUFFER 
+				|| DescriptorType == VK_DESCRIPTOR_TYPE_STORAGE_IMAGE
+				|| DescriptorType == VK_DESCRIPTOR_TYPE_STORAGE_TEXEL_BUFFER);
+			FRHIUnorderedAccessView* CurrentUAV = (FRHIUnorderedAccessView*)(ResourceArray[ResourceInfo.SourceUBResourceIndex].GetReference());
+			if (CurrentUAV)
+			{
+				FVulkanUnorderedAccessView* UAV = ResourceCast(CurrentUAV);
+				State->SetUAVForUBResource(GlobalRemappingInfo[ResourceInfo.GlobalIndex].NewDescriptorSet, GlobalRemappingInfo[ResourceInfo.GlobalIndex].NewBindingIndex, UAV);
+			}
+			else
+			{
+#if VULKAN_ENABLE_SHADER_DEBUG_NAMES
+				UE_LOG(LogVulkanRHI, Warning, TEXT("Invalid texture in SRT table for shader '%s'"), *Shader->GetDebugName());
+#else
+				UE_LOG(LogVulkanRHI, Warning, TEXT("Invalid texture in SRT table"));
+#endif
+			}
+		}
+		break;
+
 		default:
 			check(0);
 			break;
diff --git a/Engine/Source/Runtime/VulkanRHI/Private/VulkanPipelineState.cpp b/Engine/Source/Runtime/VulkanRHI/Private/VulkanPipelineState.cpp
--- a/Engine/Source/Runtime/VulkanRHI/Private/VulkanPipelineState.cpp
+++ b/Engine/Source/Runtime/VulkanRHI/Private/VulkanPipelineState.cpp
@@ -362,14 +362,6 @@ void FVulkanCommandListContext::RHISetGraphicsPipelineState(FRHIGraphicsPipeline
 		PendingGfxState->MarkNeedsDynamicStates();
 		PendingGfxState->StencilRef = 0;
 	}
-
-	// Yuck - Bind pending pixel shader UAVs from SetRenderTargets
-	{
-		for (int32 Index = 0; Index < PendingPixelUAVs.Num(); ++Index)
-		{
-			PendingGfxState->SetUAVForStage(ShaderStage::Pixel, PendingPixelUAVs[Index].BindIndex, PendingPixelUAVs[Index].UAV);
-		}
-	}
 }
 
 
