Fix OpenGL's RHICopyTexture to include all requested faces of cube textures.
#rnx

diff --git a/Engine/Source/Runtime/OpenGLDrv/Private/OpenGLTexture.cpp b/Engine/Source/Runtime/OpenGLDrv/Private/OpenGLTexture.cpp
--- a/Engine/Source/Runtime/OpenGLDrv/Private/OpenGLTexture.cpp
+++ b/Engine/Source/Runtime/OpenGLDrv/Private/OpenGLTexture.cpp
@@ -2611,7 +2611,12 @@ void FOpenGLDynamicRHI::RHICopyTexture(FRHITexture* SourceTextureRHI, FRHITextur
 	{
 		Width = CopyInfo.Size.X;
 		Height = CopyInfo.Size.Y;
-		Depth = (SourceTexture->Target == GL_TEXTURE_3D) ? CopyInfo.Size.Z : 1;
+		switch (SourceTexture->Target)
+		{
+			case GL_TEXTURE_3D:			Depth = CopyInfo.Size.Z; break;
+			case GL_TEXTURE_CUBE_MAP:	Depth = CopyInfo.NumSlices; break;
+			default:					Depth = 1; break;
+		}
 	}
 
 	GLint SrcMip = CopyInfo.SourceMipIndex;
@@ -2692,22 +2697,28 @@ void FOpenGLDynamicRHI::RHICopyTexture(FRHITexture* SourceTextureRHI, FRHITextur
 
 			switch (DestTexture->Target)
 			{
-			case GL_TEXTURE_1D:
-				FOpenGL::CopyTexSubImage1D(DestTexture->Target, DestMip, XOffset, X, 0, Width);
-				break;
-			case GL_TEXTURE_1D_ARRAY:
-				FOpenGL::CopyTexSubImage2D(DestTexture->Target, DestMip, XOffset, CopyInfo.DestSliceIndex + SliceIndex, X, 0, Width, 1);
-				break;
-			case GL_TEXTURE_2D:
-			case GL_TEXTURE_RECTANGLE:
-				FOpenGL::CopyTexSubImage2D(DestTexture->Target, DestMip, XOffset, YOffset, X, Y, Width, Height);
-				break;
-			case GL_TEXTURE_2D_ARRAY:
-				FOpenGL::CopyTexSubImage3D(DestTexture->Target, DestMip, XOffset, YOffset, CopyInfo.DestSliceIndex + SliceIndex, X, Y, Width, Height);
-				break;
-			case GL_TEXTURE_3D:
-				FOpenGL::CopyTexSubImage3D(DestTexture->Target, DestMip, XOffset, YOffset, ZOffset, X, Y, Width, Height);
-				break;
+				case GL_TEXTURE_1D:
+					FOpenGL::CopyTexSubImage1D(DestTexture->Target, DestMip, XOffset, X, 0, Width);
+					break;
+				case GL_TEXTURE_1D_ARRAY:
+					FOpenGL::CopyTexSubImage2D(DestTexture->Target, DestMip, XOffset, CopyInfo.DestSliceIndex + SliceIndex, X, 0, Width, 1);
+					break;
+				case GL_TEXTURE_2D:
+				case GL_TEXTURE_RECTANGLE:
+					FOpenGL::CopyTexSubImage2D(DestTexture->Target, DestMip, XOffset, YOffset, X, Y, Width, Height);
+					break;
+				case GL_TEXTURE_2D_ARRAY:
+					FOpenGL::CopyTexSubImage3D(DestTexture->Target, DestMip, XOffset, YOffset, CopyInfo.DestSliceIndex + SliceIndex, X, Y, Width, Height);
+					break;
+				case GL_TEXTURE_3D:
+					FOpenGL::CopyTexSubImage3D(DestTexture->Target, DestMip, XOffset, YOffset, ZOffset, X, Y, Width, Height);
+					break;
+				case GL_TEXTURE_CUBE_MAP:
+					for (int32 FaceIndex = FMath::Min((int32)CopyInfo.NumSlices, 6) - 1; FaceIndex >= 0; FaceIndex--)
+					{
+						FOpenGL::CopyTexSubImage2D(GL_TEXTURE_CUBE_MAP_POSITIVE_X + (uint32)FaceIndex, CopyInfo.DestMipIndex, XOffset, YOffset, X, Y, Width, Height);
+					}					
+					break;
 			}
 		}
 
