Fixed compile error about GL_TEXTURE_BUFFER_OFFSET_ALIGNMENT for OpenGLES RHI.

#rnx

diff --git a/Engine/Source/Runtime/OpenGLDrv/Private/OpenGLDrv.cpp b/Engine/Source/Runtime/OpenGLDrv/Private/OpenGLDrv.cpp
--- a/Engine/Source/Runtime/OpenGLDrv/Private/OpenGLDrv.cpp
+++ b/Engine/Source/Runtime/OpenGLDrv/Private/OpenGLDrv.cpp
@@ -405,6 +405,7 @@ GLint FOpenGLBase::MaxGeometryTextureImageUnits = -1;
 GLint FOpenGLBase::MaxHullTextureImageUnits = -1;
 GLint FOpenGLBase::MaxDomainTextureImageUnits = -1;
 GLint FOpenGLBase::MaxVaryingVectors = -1;
+GLint FOpenGLBase::TextureBufferAlignment = -1;
 GLint FOpenGLBase::MaxVertexUniformComponents = -1;
 GLint FOpenGLBase::MaxPixelUniformComponents = -1;
 GLint FOpenGLBase::MaxGeometryUniformComponents = -1;
@@ -424,6 +425,7 @@ void FOpenGLBase::ProcessQueryGLInt()
 	GET_GL_INT(GL_MAX_TEXTURE_IMAGE_UNITS, 0, MaxTextureImageUnits);
 	GET_GL_INT(GL_MAX_VERTEX_TEXTURE_IMAGE_UNITS, 0, MaxVertexTextureImageUnits);
 	GET_GL_INT(GL_MAX_COMBINED_TEXTURE_IMAGE_UNITS, 0, MaxCombinedTextureImageUnits);
+	GET_GL_INT(GL_TEXTURE_BUFFER_OFFSET_ALIGNMENT, 0, TextureBufferAlignment);
 }
 
 void FOpenGLBase::ProcessExtensions( const FString& ExtensionsString )
diff --git a/Engine/Source/Runtime/OpenGLDrv/Private/OpenGLESDeferred.cpp b/Engine/Source/Runtime/OpenGLDrv/Private/OpenGLESDeferred.cpp
--- a/Engine/Source/Runtime/OpenGLDrv/Private/OpenGLESDeferred.cpp
+++ b/Engine/Source/Runtime/OpenGLDrv/Private/OpenGLESDeferred.cpp
@@ -221,6 +221,7 @@ void FOpenGLESDeferred::ProcessQueryGLInt()
 		MaxGeometryTextureImageUnits = 0;
 		MaxHullTextureImageUnits = 0;
 		MaxDomainTextureImageUnits = 0;
+		TextureBufferAlignment = 0;
 	}
 	else
 	{
@@ -234,6 +235,8 @@ void FOpenGLESDeferred::ProcessQueryGLInt()
 		GET_GL_INT(GL_MAX_COMPUTE_TEXTURE_IMAGE_UNITS, 0, MaxComputeTextureImageUnits);
 		GET_GL_INT(GL_MAX_COMPUTE_UNIFORM_COMPONENTS, 0, MaxComputeUniformComponents);
 
+		GET_GL_INT(GL_TEXTURE_BUFFER_OFFSET_ALIGNMENT_EXT, 0, TextureBufferAlignment);
+
 		if (bSupportsTessellation)
 		{
 			GET_GL_INT(GL_MAX_TESS_CONTROL_UNIFORM_COMPONENTS_EXT, 0, MaxHullUniformComponents);
diff --git a/Engine/Source/Runtime/OpenGLDrv/Private/OpenGLUAV.cpp b/Engine/Source/Runtime/OpenGLDrv/Private/OpenGLUAV.cpp
--- a/Engine/Source/Runtime/OpenGLDrv/Private/OpenGLUAV.cpp
+++ b/Engine/Source/Runtime/OpenGLDrv/Private/OpenGLUAV.cpp
@@ -26,8 +26,7 @@ static void BindGLTexBufferRange(GLenum Target, GLenum InternalFormat, GLuint Bu
 		GLsizeiptr Size = NumElements * Stride;
 
 #if DO_CHECK
-		GLint Alignment = 0;
-		glGetIntegerv(GL_TEXTURE_BUFFER_OFFSET_ALIGNMENT, &Alignment);
+		GLint Alignment = FOpenGLBase::GetTextureBufferAlignment();
 		check(Stride > 0 && Offset % Alignment == 0);
 #endif
 
diff --git a/Engine/Source/Runtime/OpenGLDrv/Public/OpenGL.h b/Engine/Source/Runtime/OpenGLDrv/Public/OpenGL.h
--- a/Engine/Source/Runtime/OpenGLDrv/Public/OpenGL.h
+++ b/Engine/Source/Runtime/OpenGLDrv/Public/OpenGL.h
@@ -194,6 +194,7 @@ public:
 	static FORCEINLINE GLint GetMaxDomainTextureImageUnits()	{ check(MaxDomainTextureImageUnits != -1); return MaxDomainTextureImageUnits; }
 	static FORCEINLINE GLint GetMaxComputeTextureImageUnits()	{ return 0; }
 	static FORCEINLINE GLint GetMaxCombinedTextureImageUnits()	{ check(MaxCombinedTextureImageUnits != -1); return MaxCombinedTextureImageUnits; }
+	static FORCEINLINE GLint GetTextureBufferAlignment()		{ return TextureBufferAlignment; }
 
 
 	// Indices per unit are set in this order [Pixel, Vertex, Geometry, Hull, Domain]
@@ -431,6 +432,7 @@ protected:
 	static GLint MaxHullUniformComponents;
 	static GLint MaxDomainUniformComponents;
 	static GLint MaxVaryingVectors;
+	static GLint TextureBufferAlignment;
 
 	/** GL_ARB_clip_control */
 	static bool bSupportsClipControl;
