Appendix for fix in CL 11110898

#rnx

diff --git a/Engine/Source/Runtime/OpenGLDrv/Private/OpenGL4.cpp b/Engine/Source/Runtime/OpenGLDrv/Private/OpenGL4.cpp
--- a/Engine/Source/Runtime/OpenGLDrv/Private/OpenGL4.cpp
+++ b/Engine/Source/Runtime/OpenGLDrv/Private/OpenGL4.cpp
@@ -25,7 +25,8 @@ void FOpenGL4::ProcessQueryGLInt()
 {
 	GET_GL_INT(GL_MAX_COMBINED_IMAGE_UNIFORMS, 0, MaxCombinedUAVUnits);
 	GET_GL_INT(GL_MAX_FRAGMENT_IMAGE_UNIFORMS, 0, MaxPixelUAVUnits);
-	
+	GET_GL_INT(GL_TEXTURE_BUFFER_OFFSET_ALIGNMENT, 0, TextureBufferAlignment);
+
 	if (bSupportsComputeShaders)
 	{
 		GET_GL_INT(GL_MAX_COMPUTE_TEXTURE_IMAGE_UNITS, 0, MaxComputeTextureImageUnits);
diff --git a/Engine/Source/Runtime/OpenGLDrv/Private/OpenGLDrv.cpp b/Engine/Source/Runtime/OpenGLDrv/Private/OpenGLDrv.cpp
--- a/Engine/Source/Runtime/OpenGLDrv/Private/OpenGLDrv.cpp
+++ b/Engine/Source/Runtime/OpenGLDrv/Private/OpenGLDrv.cpp
@@ -425,7 +425,6 @@ void FOpenGLBase::ProcessQueryGLInt()
 	GET_GL_INT(GL_MAX_TEXTURE_IMAGE_UNITS, 0, MaxTextureImageUnits);
 	GET_GL_INT(GL_MAX_VERTEX_TEXTURE_IMAGE_UNITS, 0, MaxVertexTextureImageUnits);
 	GET_GL_INT(GL_MAX_COMBINED_TEXTURE_IMAGE_UNITS, 0, MaxCombinedTextureImageUnits);
-	GET_GL_INT(GL_TEXTURE_BUFFER_OFFSET_ALIGNMENT, 0, TextureBufferAlignment);
 }
 
 void FOpenGLBase::ProcessExtensions( const FString& ExtensionsString )
