Fixex compilation error in OpenGL Linux

diff --git a/Engine/Source/Runtime/OpenGLDrv/Private/Linux/OpenGLLinux.h b/Engine/Source/Runtime/OpenGLDrv/Private/Linux/OpenGLLinux.h
--- a/Engine/Source/Runtime/OpenGLDrv/Private/Linux/OpenGLLinux.h
+++ b/Engine/Source/Runtime/OpenGLDrv/Private/Linux/OpenGLLinux.h
@@ -288,6 +288,7 @@ THIRD_PARTY_INCLUDES_END
 	EnumMacro(PFNGLBINDVERTEXARRAYPROC,glBindVertexArray) \
 	EnumMacro(PFNGLCOPYBUFFERSUBDATAPROC,glCopyBufferSubData) \
 	EnumMacro(PFNGLTEXBUFFERPROC,glTexBuffer) \
+	EnumMacro(PFNGLTEXBUFFERRANGEPROC, glTexBufferRange)\
 	EnumMacro(PFNGLTEXIMAGE2DMULTISAMPLEPROC,glTexImage2DMultisample) \
 	EnumMacro(PFNGLQUERYCOUNTERPROC, glQueryCounter)\
 	EnumMacro(PFNGLISSYNCPROC, glIsSync)\
