Add some additional info to a check firing within OGL's LRU cache on android with quick smoke tests.
#rnx

diff --git a/Engine/Source/Runtime/OpenGLDrv/Private/OpenGLShaders.cpp b/Engine/Source/Runtime/OpenGLDrv/Private/OpenGLShaders.cpp
--- a/Engine/Source/Runtime/OpenGLDrv/Private/OpenGLShaders.cpp
+++ b/Engine/Source/Runtime/OpenGLDrv/Private/OpenGLShaders.cpp
@@ -2221,7 +2221,7 @@ class FGLProgramCacheLRU
 			LRUBinaryMemoryUse -= GetProgramBinarySize(LinkedProgram->Program);
 		}
 
-		check(!EvictedPrograms.Contains(LinkedProgram->Config.ProgramKey));
+		checkf(!EvictedPrograms.Contains(LinkedProgram->Config.ProgramKey), TEXT("Program is already in the evicted program list: %s"), *LinkedProgram->Config.ProgramKey.ToString());
 		//UE_LOG(LogRHI, Warning, TEXT("LRU: Evicting program %d"), LinkedProgram->Program);
 		FEvictedGLProgram& test = EvictedPrograms.Emplace(LinkedProgram->Config.ProgramKey, FEvictedGLProgram(LinkedProgram));
 		INC_DWORD_STAT(STAT_OpenGLShaderLRUEvictedProgramCount);
