Android GLES shader LRU
 - Fix case in which an evicted program could be enqueued for eviction again.
#rnx

diff --git a/Engine/Source/Runtime/OpenGLDrv/Private/OpenGLShaders.cpp b/Engine/Source/Runtime/OpenGLDrv/Private/OpenGLShaders.cpp
--- a/Engine/Source/Runtime/OpenGLDrv/Private/OpenGLShaders.cpp
+++ b/Engine/Source/Runtime/OpenGLDrv/Private/OpenGLShaders.cpp
@@ -2179,11 +2179,13 @@ class FGLProgramCacheLRU
 			// UE_LOG(LogRHI, Warning, TEXT("LRU: found and recovered EVICTED program %s"), *ProgramKey.ToString());
 			FoundEvicted->RestoreGLProgramFromBinary();
 			FOpenGLLinkedProgram* LinkedProgram = FoundEvicted->GetLinkedProgram();
-			// Add this back to the LRU
-			Add(ProgramKey, LinkedProgram);
 
 			// Remove from the evicted program map.
 			EvictedPrograms.Remove(ProgramKey);
+
+			// Add this back to the LRU
+			Add(ProgramKey, LinkedProgram);
+
 			DEC_DWORD_STAT(STAT_OpenGLShaderLRUEvictedProgramCount);
 
 			// reconfigure the new program:
@@ -2229,6 +2231,11 @@ class FGLProgramCacheLRU
 
 public:
 
+	bool IsEvicted(const FOpenGLProgramKey& ProgramKey)
+	{
+		return FindEvicted(ProgramKey) != nullptr;
+	}
+
 	void EvictLeastRecentFromLRU()
 	{
 		EvictFromLRU(LRU.RemoveLeastRecent());
@@ -2262,8 +2269,8 @@ public:
 	{
 		// Remove least recently used programs until we reach our limit.
 		// note that a single large shader could evict multiple smaller shaders.
-		// todo: logic for this!
-		check(!LRU.Contains(ProgramKey));
+		checkf(!LRU.Contains(ProgramKey), TEXT("Program is already in the LRU program list: %s"), *ProgramKey.ToString());
+		checkf(!IsEvicted(ProgramKey), TEXT("Program is already in the evicted program list: %s"), *ProgramKey.ToString());
 
 		// UE_LOG(LogRHI, Warning, TEXT("LRU: adding program %s (%d)"), *ProgramKey.ToString(), LinkedProgram->Program);
 
@@ -2466,6 +2473,12 @@ public:
 		ProgramCacheLRU.AddAsEvicted(ProgramKey, MoveTemp(ProgramBinary));
 	}
 
+	bool IsEvicted(const FOpenGLProgramKey& ProgramKey)
+	{
+		check(IsUsingLRU());
+		return ProgramCacheLRU.IsEvicted(ProgramKey);
+	}
+
 	void EnumerateLinkedPrograms(TFunction<void(FOpenGLLinkedProgram*)> EnumFunc)
 	{
 		if (bUseLRUCache)
@@ -2510,6 +2523,8 @@ void FDelayedEvictionContainer::Add(FOpenGLLinkedProgram* LinkedProgram)
 		return;
 	}
 
+	checkf(!GetOpenGLProgramsCache().IsEvicted(LinkedProgram->Config.ProgramKey), TEXT("FDelayedEvictionContainer::Add is already evicted! [%s], %d"), *LinkedProgram->Config.ProgramKey.ToString(), LinkedProgram->LRUInfo.EvictBucket);
+
 	if (LinkedProgram->LRUInfo.EvictBucket >=0 )
 	{
 		Remove(LinkedProgram);
@@ -4038,19 +4053,23 @@ FOpenGLBoundShaderState::~FOpenGLBoundShaderState()
 	check(LinkedProgram);
 	RunOnGLRenderContextThread([LinkedProgram = LinkedProgram]()
 	{
-		FOpenGLLinkedProgram* Prog = StaticLastReleasedPrograms[StaticLastReleasedProgramsIndex];
-		StaticLastReleasedPrograms[StaticLastReleasedProgramsIndex++] = LinkedProgram;
-		if (StaticLastReleasedProgramsIndex == LAST_RELEASED_PROGRAMS_CACHE_COUNT)
+		const bool bIsEvicted = GetOpenGLProgramsCache().IsUsingLRU() && GetOpenGLProgramsCache().IsEvicted(LinkedProgram->Config.ProgramKey);
+		if( !bIsEvicted )
 		{
-			StaticLastReleasedProgramsIndex = 0;
-		}
+			FOpenGLLinkedProgram* Prog = StaticLastReleasedPrograms[StaticLastReleasedProgramsIndex];
+			StaticLastReleasedPrograms[StaticLastReleasedProgramsIndex++] = LinkedProgram;
+			if (StaticLastReleasedProgramsIndex == LAST_RELEASED_PROGRAMS_CACHE_COUNT)
+			{
+				StaticLastReleasedProgramsIndex = 0;
+			}
 
-		if (CVarEvictOnBssDestruct.GetValueOnAnyThread() && GetOpenGLProgramsCache().IsUsingLRU())
-		{
-			FDelayedEvictionContainer::Get().Add(LinkedProgram);
-		}
+			if (CVarEvictOnBssDestruct.GetValueOnAnyThread() && GetOpenGLProgramsCache().IsUsingLRU())
+			{
+				FDelayedEvictionContainer::Get().Add(LinkedProgram);
+			}
 
-		OnProgramDeletion(LinkedProgram->Program);
+			OnProgramDeletion(LinkedProgram->Program);
+		}
 	});
 }
 
@@ -4772,7 +4791,7 @@ void FOpenGLProgramBinaryCache::ScanProgramCacheFile(const FGuid& ShaderPipeline
 							}
 							if (bAllShadersLoaded)
 							{
-								FPlatformMisc::LowLevelOutputDebugStringf(TEXT("*** All shaders for program %s already loaded\n"), *ProgramKey.ToString());
+								FPlatformMisc::LowLevelOutputDebugStringf(TEXT("*** All shaders for %s already loaded\n"), *ProgramKey.ToString());
 								NewEntry->ProgramBinaryData.AddUninitialized(ProgramBinarySize);
 								Ar.Serialize(NewEntry->ProgramBinaryData.GetData(), ProgramBinarySize);
 								NewEntry->GLProgramState = FGLProgramBinaryFileCacheEntry::EGLProgramState::ProgramLoaded;
