Disable firing check whilst investigations take place.
#rnx

diff --git a/Engine/Source/Runtime/OpenGLDrv/Private/OpenGLShaders.cpp b/Engine/Source/Runtime/OpenGLDrv/Private/OpenGLShaders.cpp
--- a/Engine/Source/Runtime/OpenGLDrv/Private/OpenGLShaders.cpp
+++ b/Engine/Source/Runtime/OpenGLDrv/Private/OpenGLShaders.cpp
@@ -2270,7 +2270,7 @@ public:
 		// Remove least recently used programs until we reach our limit.
 		// note that a single large shader could evict multiple smaller shaders.
 		checkf(!LRU.Contains(ProgramKey), TEXT("Program is already in the LRU program list: %s"), *ProgramKey.ToString());
-		checkf(!IsEvicted(ProgramKey), TEXT("Program is already in the evicted program list: %s"), *ProgramKey.ToString());
+		//checkf(!IsEvicted(ProgramKey), TEXT("Program is already in the evicted program list: %s"), *ProgramKey.ToString());
 
 		// UE_LOG(LogRHI, Warning, TEXT("LRU: adding program %s (%d)"), *ProgramKey.ToString(), LinkedProgram->Program);
 
