Flush RHI/RT before creating the opengl binary file cache. This ensures we do not have RT/RHI thread creating programs while binary file cache initialization adds programs to the LRU.

diff --git a/Engine/Source/Runtime/OpenGLDrv/Private/OpenGLShaders.cpp b/Engine/Source/Runtime/OpenGLDrv/Private/OpenGLShaders.cpp
--- a/Engine/Source/Runtime/OpenGLDrv/Private/OpenGLShaders.cpp
+++ b/Engine/Source/Runtime/OpenGLDrv/Private/OpenGLShaders.cpp
@@ -2270,7 +2270,7 @@ public:
 		// Remove least recently used programs until we reach our limit.
 		// note that a single large shader could evict multiple smaller shaders.
 		checkf(!LRU.Contains(ProgramKey), TEXT("Program is already in the LRU program list: %s"), *ProgramKey.ToString());
-		//checkf(!IsEvicted(ProgramKey), TEXT("Program is already in the evicted program list: %s"), *ProgramKey.ToString());
+		checkf(!IsEvicted(ProgramKey), TEXT("Program is already in the evicted program list: %s"), *ProgramKey.ToString());
 
 		// UE_LOG(LogRHI, Warning, TEXT("LRU: adding program %s (%d)"), *ProgramKey.ToString(), LinkedProgram->Program);
 
@@ -4687,6 +4687,10 @@ static FCriticalSection GPendingGLProgramCreateRequestsCS;
 // Scan the binary cache file and build a record of all programs.
 void FOpenGLProgramBinaryCache::ScanProgramCacheFile(const FGuid& ShaderPipelineCacheVersionGuid)
 {
+	// Flush any pending BSS create requests, this ensures we do not have RT/RHI thread creating programs while this step adds programs to the LRU.
+	check(IsInGameThread());
+	FlushRenderingCommands();
+
 	UE_LOG(LogRHI, Log, TEXT("OnShaderScanProgramCacheFile"));
 	FScopeLock Lock(&GProgramBinaryCacheCS);
 	FString ProgramCacheFilename = GetProgramBinaryCacheFilePath();
