GitHub #6267 : Mark Classes Having a Final Dtor as Final

diff --git a/Engine/Source/Runtime/Core/Public/Async/TaskGraphInterfaces.h b/Engine/Source/Runtime/Core/Public/Async/TaskGraphInterfaces.h
--- a/Engine/Source/Runtime/Core/Public/Async/TaskGraphInterfaces.h
+++ b/Engine/Source/Runtime/Core/Public/Async/TaskGraphInterfaces.h
@@ -732,7 +732,7 @@ public:
  *	Embeds a user defined task, as exemplified above, for doing the work and provides the functionality for setting up and handling prerequisites and subsequents
  **/
 template<typename TTask>
-class TGraphTask : public FBaseGraphTask
+class TGraphTask final : public FBaseGraphTask
 {
 public:
 	/** 
@@ -828,7 +828,7 @@ private:
 	 *	Dispatches the subsequents.
 	 *	Destroys myself.
 	 **/
-	virtual void ExecuteTask(TArray<FBaseGraphTask*>& NewTasks, ENamedThreads::Type CurrentThread) final override
+	void ExecuteTask(TArray<FBaseGraphTask*>& NewTasks, ENamedThreads::Type CurrentThread) override
 	{
 		checkThreadGraph(TaskConstructed);
 
@@ -885,7 +885,7 @@ private:
 	/** 
 	 *	Private destructor, just checks that the task appears to be completed
 	**/
-	virtual ~TGraphTask() final override
+	~TGraphTask() override
 	{
 		checkThreadGraph(!TaskConstructed);
 	}
diff --git a/Engine/Source/Runtime/D3D12RHI/Public/D3D12Resources.h b/Engine/Source/Runtime/D3D12RHI/Public/D3D12Resources.h
--- a/Engine/Source/Runtime/D3D12RHI/Public/D3D12Resources.h
+++ b/Engine/Source/Runtime/D3D12RHI/Public/D3D12Resources.h
@@ -1003,7 +1003,7 @@ private:
 	FD3D12ResourceLocation ResourceLocation;
 };
 
-class FD3D12StagingBuffer : public FRHIStagingBuffer
+class FD3D12StagingBuffer final : public FRHIStagingBuffer
 {
 	friend class FD3D12CommandContext;
 	friend class FD3D12DynamicRHI;
@@ -1014,7 +1014,7 @@ public:
 		, StagedRead(nullptr)
 		, ShadowBufferSize(0)
 	{}
-	virtual ~FD3D12StagingBuffer() final override;
+	~FD3D12StagingBuffer() override;
 
 	void SafeRelease()
 	{
@@ -1025,8 +1025,8 @@ public:
 		}
 	}
 
-	virtual void* Lock(uint32 Offset, uint32 NumBytes) final override;
-	virtual void Unlock() final override;
+	void* Lock(uint32 Offset, uint32 NumBytes) override;
+	void Unlock() override;
 
 private:
 	FD3D12Resource* StagedRead;
diff --git a/Engine/Source/Runtime/OpenGLDrv/Public/OpenGLResources.h b/Engine/Source/Runtime/OpenGLDrv/Public/OpenGLResources.h
--- a/Engine/Source/Runtime/OpenGLDrv/Public/OpenGLResources.h
+++ b/Engine/Source/Runtime/OpenGLDrv/Public/OpenGLResources.h
@@ -2088,7 +2088,7 @@ private:
 // Fences
 
 // Note that Poll() and WriteInternal() will stall the RHI thread if one is present.
-class FOpenGLGPUFence : public FRHIGPUFence
+class FOpenGLGPUFence final : public FRHIGPUFence
 {
 public:
 	FOpenGLGPUFence(FName InName)
@@ -2096,10 +2096,10 @@ public:
 		, bValidSync(false)
 	{}
 
-	virtual ~FOpenGLGPUFence() final override;
+	~FOpenGLGPUFence() override;
 
-	virtual void Clear() final override;
-	virtual bool Poll() const final override;
+	void Clear() override;
+	bool Poll() const override;
 	
 	void WriteInternal();
 private:
@@ -2108,7 +2108,7 @@ private:
 	bool bValidSync;
 };
 
-class FOpenGLStagingBuffer : public FRHIStagingBuffer
+class FOpenGLStagingBuffer final : public FRHIStagingBuffer
 {
 	friend class FOpenGLDynamicRHI;
 public:
@@ -2117,13 +2117,13 @@ public:
 		Initialize();
 	}
 
-	virtual ~FOpenGLStagingBuffer() final override;
+	~FOpenGLStagingBuffer() override;
 
 	// Locks the shadow of VertexBuffer for read. This will stall the RHI thread.
-	virtual void *Lock(uint32 Offset, uint32 NumBytes) final override;
+	void *Lock(uint32 Offset, uint32 NumBytes) override;
 
 	// Unlocks the shadow. This is an error if it was not locked previously.
-	virtual void Unlock() final override;
+	void Unlock() override;
 private:
 	void Initialize();
 
diff --git a/Engine/Source/Runtime/Windows/D3D11RHI/Public/D3D11Resources.h b/Engine/Source/Runtime/Windows/D3D11RHI/Public/D3D11Resources.h
--- a/Engine/Source/Runtime/Windows/D3D11RHI/Public/D3D11Resources.h
+++ b/Engine/Source/Runtime/Windows/D3D11RHI/Public/D3D11Resources.h
@@ -833,7 +833,7 @@ public:
 	}
 };
 
-class FD3D11StagingBuffer : public FRHIStagingBuffer
+class FD3D11StagingBuffer final : public FRHIStagingBuffer
 {
 	friend class FD3D11DynamicRHI;
 public:
@@ -841,10 +841,10 @@ public:
 		: FRHIStagingBuffer()
 	{}
 
-	virtual ~FD3D11StagingBuffer() final override;
+	~FD3D11StagingBuffer() override;
 
-	virtual void* Lock(uint32 Offset, uint32 NumBytes) final override;
-	virtual void Unlock() final override;
+	void* Lock(uint32 Offset, uint32 NumBytes) override;
+	void Unlock() override;
 
 private:
 	FD3D11DeviceContext* Context;
