Fixed: Graphical corruption occurs on specific Android devices. Screen is flashing and displaying wrong colors (disabled early_fragment_tests for now)

diff --git a/Engine/Source/Runtime/OpenGLDrv/Private/OpenGLShaders.cpp b/Engine/Source/Runtime/OpenGLDrv/Private/OpenGLShaders.cpp
--- a/Engine/Source/Runtime/OpenGLDrv/Private/OpenGLShaders.cpp
+++ b/Engine/Source/Runtime/OpenGLDrv/Private/OpenGLShaders.cpp
@@ -1054,6 +1054,11 @@ void OPENGLDRV_API GLSLToDeviceCompatibleGLSL(FAnsiCharArray& GlslCodeOriginal,
 		AppendCString(GlslCode, "#endif\n");
 	}
 
+	if (TypeEnum == GL_FRAGMENT_SHADER)
+	{
+		ReplaceCString(GlslCodeOriginal, "layout(early_fragment_tests) in;", "");
+	}
+
 	// The incoming glsl may have preprocessor code that is dependent on defines introduced via the engine.
 	// This is the place to insert such engine preprocessor defines, immediately after the glsl version declaration.
 	if (Capabilities.bRequiresUEShaderFramebufferFetchDef && TypeEnum == GL_FRAGMENT_SHADER )
