Fix for warning after upgrading to PVS-Studio 7.0.4:

V1013: Suspicious subexpression in a sequence of similar comparisons.

#rnx

diff --git a/Engine/Source/Runtime/OpenGLDrv/Private/OpenGLCommands.cpp b/Engine/Source/Runtime/OpenGLDrv/Private/OpenGLCommands.cpp
--- a/Engine/Source/Runtime/OpenGLDrv/Private/OpenGLCommands.cpp
+++ b/Engine/Source/Runtime/OpenGLDrv/Private/OpenGLCommands.cpp
@@ -1369,8 +1369,8 @@ void FOpenGLDynamicRHI::SetPendingBlendStateForActiveRenderTargets( FOpenGLConte
 				{
 					if (CachedRenderTargetBlendState.ColorSourceBlendFactor != RenderTargetBlendState.ColorSourceBlendFactor
 						|| CachedRenderTargetBlendState.ColorDestBlendFactor != RenderTargetBlendState.ColorDestBlendFactor
-						|| CachedRenderTargetBlendState.AlphaSourceBlendFactor != RenderTargetBlendState.ColorSourceBlendFactor
-						|| CachedRenderTargetBlendState.AlphaDestBlendFactor != RenderTargetBlendState.ColorDestBlendFactor)
+						|| CachedRenderTargetBlendState.AlphaSourceBlendFactor != RenderTargetBlendState.AlphaSourceBlendFactor
+						|| CachedRenderTargetBlendState.AlphaDestBlendFactor != RenderTargetBlendState.AlphaDestBlendFactor)
 					{
 						FOpenGL::BlendFunci(RenderTargetIndex, RenderTargetBlendState.ColorSourceBlendFactor, RenderTargetBlendState.ColorDestBlendFactor);
 					}
@@ -1442,8 +1442,8 @@ void FOpenGLDynamicRHI::SetPendingBlendStateForActiveRenderTargets( FOpenGLConte
 					{
 						if (CachedRenderTargetBlendState.ColorSourceBlendFactor != RenderTargetBlendState.ColorSourceBlendFactor
 							|| CachedRenderTargetBlendState.ColorDestBlendFactor != RenderTargetBlendState.ColorDestBlendFactor
-							|| CachedRenderTargetBlendState.AlphaSourceBlendFactor != RenderTargetBlendState.ColorSourceBlendFactor
-							|| CachedRenderTargetBlendState.AlphaDestBlendFactor != RenderTargetBlendState.ColorDestBlendFactor)
+							|| CachedRenderTargetBlendState.AlphaSourceBlendFactor != RenderTargetBlendState.AlphaSourceBlendFactor
+							|| CachedRenderTargetBlendState.AlphaDestBlendFactor != RenderTargetBlendState.AlphaDestBlendFactor)
 						{
 							glBlendFunc(RenderTargetBlendState.ColorSourceBlendFactor, RenderTargetBlendState.ColorDestBlendFactor);
 						}
