Remove RT flush from GL ScanProgramCacheFile, it can cause shader pipeline cache processing to deadlock.
Place binary cache mutex around createBSS to prevent it adding shaders while ScanProgramCacheFile runs.
#rnx

diff --git a/Engine/Source/Runtime/OpenGLDrv/Private/OpenGLShaders.cpp b/Engine/Source/Runtime/OpenGLDrv/Private/OpenGLShaders.cpp
--- a/Engine/Source/Runtime/OpenGLDrv/Private/OpenGLShaders.cpp
+++ b/Engine/Source/Runtime/OpenGLDrv/Private/OpenGLShaders.cpp
@@ -3287,6 +3287,7 @@ static void BindShaderStage(FOpenGLLinkedProgramConfiguration& Config, CrossComp
 }
 
 // ============================================================================================================================
+static FCriticalSection GProgramBinaryCacheCS;
 
 FBoundShaderStateRHIRef FOpenGLDynamicRHI::RHICreateBoundShaderState_OnThisThread(
 	FRHIVertexDeclaration* VertexDeclarationRHI,
@@ -3300,6 +3301,8 @@ FBoundShaderStateRHIRef FOpenGLDynamicRHI::RHICreateBoundShaderState_OnThisThrea
 {
 	check(IsInRenderingThread() || IsInRHIThread());
 
+	FScopeLock Lock(&GProgramBinaryCacheCS);
+
 	VERIFY_GL_SCOPE();
 
 	SCOPE_CYCLE_COUNTER(STAT_OpenGLCreateBoundShaderStateTime);
@@ -4382,16 +4385,11 @@ struct FGLProgramBinaryFileCacheEntry
 	}
 };
 
-static FCriticalSection GProgramBinaryCacheCS;
 static FCriticalSection GPendingGLProgramCreateRequestsCS;
 
 // Scan the binary cache file and build a record of all programs.
 void FOpenGLProgramBinaryCache::ScanProgramCacheFile(const FGuid& ShaderPipelineCacheVersionGuid)
 {
-	// Flush any pending BSS create requests, this ensures we do not have RT/RHI thread creating programs while this step adds programs to the LRU.
-	check(IsInGameThread());
-	FlushRenderingCommands();
-
 	UE_LOG(LogRHI, Log, TEXT("OnShaderScanProgramCacheFile"));
 	FScopeLock Lock(&GProgramBinaryCacheCS);
 	FString ProgramCacheFilename = GetProgramBinaryCacheFilePath();
