Introduce a limit to the time per frame that GL RHI will spend processing shader library requests.
 - Potentially a fix for GC timeout with large PSO cache on GLES devices with slow compile times.
#rnx

diff --git a/Engine/Source/Runtime/OpenGLDrv/Private/OpenGLShaders.cpp b/Engine/Source/Runtime/OpenGLDrv/Private/OpenGLShaders.cpp
--- a/Engine/Source/Runtime/OpenGLDrv/Private/OpenGLShaders.cpp
+++ b/Engine/Source/Runtime/OpenGLDrv/Private/OpenGLShaders.cpp
@@ -82,6 +82,15 @@ static TAutoConsoleVariable<int32> CVarUseExistingBinaryFileCache(
 	TEXT("1: When Pipeline Cache Version Guid changes re-use programs from the existing binary cache where possible (default)."),
 	ECVF_RenderThreadSafe);
 
+static int32 GMaxShaderLibProcessingTimeMS = 10;
+static FAutoConsoleVariableRef CVarMaxShaderLibProcessingTime(
+	TEXT("r.OpenGL.MaxShaderLibProcessingTime"),
+	GMaxShaderLibProcessingTimeMS,
+	TEXT("The maximum time per frame to process shader library requests in milliseconds.\n")
+	TEXT("default 10ms. Note: Driver compile time for a single program may exceed this limit."),
+	ECVF_RenderThreadSafe
+);
+
 #if PLATFORM_ANDROID
 bool GOpenGLShaderHackLastCompileSuccess = false;
 #endif
@@ -2152,6 +2161,7 @@ public:
 		}
 		else
 		{
+			check(!ProgramCache.Contains(ProgramKey));
 			ProgramCache.Add(ProgramKey, LinkedProgram);
 		}
 	}
@@ -4987,6 +4997,7 @@ void FOpenGLProgramBinaryCache::CheckPendingGLProgramCreateRequests()
 	FDelayedEvictionContainer::Get().Tick();
 	if (CachePtr)
 	{
+		QUICK_SCOPE_CYCLE_COUNTER(STAT_OpenGLShaderCreateShaderLibRequests);
 		CachePtr->CheckPendingGLProgramCreateRequests_internal();
 	}
 }
@@ -4996,10 +5007,17 @@ void FOpenGLProgramBinaryCache::CheckPendingGLProgramCreateRequests_internal()
 	check(IsInRenderingThread() || IsInRHIThread());
 	FScopeLock Lock(&GPendingGLProgramCreateRequestsCS);
 	//UE_LOG(LogRHI, Log, TEXT("CheckPendingGLProgramCreateRequests : PendingGLProgramCreateRequests = %d"), PendingGLProgramCreateRequests.Num());
-	while (PendingGLProgramCreateRequests.Num())
+
+	float TimeRemainingS = (float)GMaxShaderLibProcessingTimeMS / 1000.0f;
+	double StartTime = FPlatformTime::Seconds();
+	int32 Count = 0;
+	while(PendingGLProgramCreateRequests.Num() && TimeRemainingS > 0.0f)
 	{
 		CompleteLoadedGLProgramRequest_internal(PendingGLProgramCreateRequests.Pop());
+		TimeRemainingS -= (float)(FPlatformTime::Seconds() - StartTime);
+		Count++;
 	}
+	UE_CLOG(PendingGLProgramCreateRequests.Num()>0, LogRHI, Log, TEXT("CheckPendingGLProgramCreateRequests : iter count = %d, time taken = %d ms (remaining %d)"), Count, GMaxShaderLibProcessingTimeMS - (int32)(TimeRemainingS*1000.0f), PendingGLProgramCreateRequests.Num());
 }
 
 void FOpenGLProgramBinaryCache::CompleteLoadedGLProgramRequest_internal(FGLProgramBinaryFileCacheEntry* PendingGLCreate)
