Fix a crash bug of missing initialize some parameter in OpenGLShader.

diff --git a/Engine/Source/Runtime/OpenGLDrv/Private/OpenGLShaders.cpp b/Engine/Source/Runtime/OpenGLDrv/Private/OpenGLShaders.cpp
--- a/Engine/Source/Runtime/OpenGLDrv/Private/OpenGLShaders.cpp
+++ b/Engine/Source/Runtime/OpenGLDrv/Private/OpenGLShaders.cpp
@@ -387,6 +387,7 @@ struct FLibraryShaderCacheValue
 	FOpenGLCodeHeader* Header;
 	uint32 ShaderCrc;
 	GLuint GLShader;
+	TArray<FUniformBufferStaticSlot> StaticSlots;
 
 #if DEBUG_GL_SHADERS
 	TArray<ANSICHAR> GlslCode;
@@ -733,6 +734,7 @@ ShaderType* CompileOpenGLShader(TArrayView<const uint8> InShaderCode, const FSHA
 			Shader->Resource = Val->GLShader;
 			Shader->Bindings = Val->Header->Bindings;
 			Shader->UniformBuffersCopyInfo = Val->Header->UniformBuffersCopyInfo;
+			Shader->StaticSlots = Val->StaticSlots;
 			if (FOpenGL::SupportsSeparateShaderObjects())
 			{
 				FSHAHash Hash;
@@ -920,6 +922,7 @@ ShaderType* CompileOpenGLShader(TArrayView<const uint8> InShaderCode, const FSHA
 		Val.Header = new FOpenGLCodeHeader;
 		*Val.Header = Header;
 		Val.ShaderCrc = GlslCodeOriginalCRC;
+		Val.StaticSlots = Shader->StaticSlots;
 #if DEBUG_GL_SHADERS
 		Val.GlslCode = GlslCode;
 		Val.GlslCodeString = (ANSICHAR*)Shader->GlslCode.GetData();
