Fix for incorrect tex buffer offset alignment with GLES.
Fixes flickering particles.
#rnx

diff --git a/Engine/Source/Runtime/OpenGLDrv/Private/OpenGLDevice.cpp b/Engine/Source/Runtime/OpenGLDrv/Private/OpenGLDevice.cpp
--- a/Engine/Source/Runtime/OpenGLDrv/Private/OpenGLDevice.cpp
+++ b/Engine/Source/Runtime/OpenGLDrv/Private/OpenGLDevice.cpp
@@ -1468,6 +1468,12 @@ void FOpenGLDynamicRHI::SetCustomPresent(FRHICustomPresent* InCustomPresent)
 	CustomPresent = InCustomPresent;
 }
 
+uint64 FOpenGLDynamicRHI::RHIGetMinimumAlignmentForBufferBackedSRV(EPixelFormat Format)
+{
+	check(FOpenGL::GetTextureBufferAlignment()>=0);
+	return FOpenGL::GetTextureBufferAlignment();
+}
+
 bool FOpenGLDynamicRHIModule::IsSupported()
 {
 	return true;
diff --git a/Engine/Source/Runtime/OpenGLDrv/Private/OpenGLES.cpp b/Engine/Source/Runtime/OpenGLDrv/Private/OpenGLES.cpp
--- a/Engine/Source/Runtime/OpenGLDrv/Private/OpenGLES.cpp
+++ b/Engine/Source/Runtime/OpenGLDrv/Private/OpenGLES.cpp
@@ -309,6 +309,7 @@ void FOpenGLES::ProcessQueryGLInt()
 	LOG_AND_GET_GL_INT(GL_MAX_VARYING_VECTORS, 0, MaxVaryingVectors);
 	LOG_AND_GET_GL_INT(GL_MAX_VERTEX_UNIFORM_VECTORS, 0, MaxVertexUniformComponents);
 	LOG_AND_GET_GL_INT(GL_MAX_FRAGMENT_UNIFORM_VECTORS, 0, MaxPixelUniformComponents);
+	LOG_AND_GET_GL_INT(GL_TEXTURE_BUFFER_OFFSET_ALIGNMENT, 0, TextureBufferAlignment);
 
 	const GLint RequiredMaxVertexUniformComponents = 256;
 	if (MaxVertexUniformComponents < RequiredMaxVertexUniformComponents)
diff --git a/Engine/Source/Runtime/OpenGLDrv/Public/OpenGLDrv.h b/Engine/Source/Runtime/OpenGLDrv/Public/OpenGLDrv.h
--- a/Engine/Source/Runtime/OpenGLDrv/Public/OpenGLDrv.h
+++ b/Engine/Source/Runtime/OpenGLDrv/Public/OpenGLDrv.h
@@ -514,6 +514,8 @@ public:
 
 	virtual bool RHIRequiresComputeGenerateMips() const override;
 
+	virtual uint64 RHIGetMinimumAlignmentForBufferBackedSRV(EPixelFormat Format) final override;
+
 	void Cleanup();
 
 	void PurgeFramebufferFromCaches(GLuint Framebuffer);
diff --git a/Engine/Source/Runtime/OpenGLDrv/Public/OpenGLES.h b/Engine/Source/Runtime/OpenGLDrv/Public/OpenGLES.h
--- a/Engine/Source/Runtime/OpenGLDrv/Public/OpenGLES.h
+++ b/Engine/Source/Runtime/OpenGLDrv/Public/OpenGLES.h
@@ -333,6 +333,10 @@ typedef GLfloat GLdouble;
 #define GL_PROGRAM_BINARY_LENGTH 0x8741
 #endif
 
+#ifndef GL_TEXTURE_BUFFER_OFFSET_ALIGNMENT
+#define GL_TEXTURE_BUFFER_OFFSET_ALIGNMENT		0x919F
+#endif
+
 // Normalize debug macros due to naming differences across GL versions
 #if defined(GL_KHR_debug) && GL_KHR_debug
 #define GL_DEBUG_SOURCE_OTHER_ARB GL_DEBUG_SOURCE_OTHER_KHR
