Enable RHIT on OpenGL mobile by default

diff --git a/Engine/Source/Runtime/Core/Public/Android/AndroidPlatform.h b/Engine/Source/Runtime/Core/Public/Android/AndroidPlatform.h
--- a/Engine/Source/Runtime/Core/Public/Android/AndroidPlatform.h
+++ b/Engine/Source/Runtime/Core/Public/Android/AndroidPlatform.h
@@ -65,12 +65,7 @@ typedef FAndroidTypes FPlatformTypes;
 #define PLATFORM_CODE_SECTION(Name)					__attribute__((section(Name)))
 
 #define PLATFORM_GLOBAL_LOG_CATEGORY				LogAndroid
-
-#if defined(EXPERIMENTAL_OPENGL_RHITHREAD) && EXPERIMENTAL_OPENGL_RHITHREAD
-	#define PLATFORM_RHITHREAD_DEFAULT_BYPASS			0
-#else
-	#define PLATFORM_RHITHREAD_DEFAULT_BYPASS			1
-#endif
+#define PLATFORM_RHITHREAD_DEFAULT_BYPASS			0
 
 // Conditionally set in AndroidToolChain.cs
 // always set to 1 for ARM64 builds
diff --git a/Engine/Source/Runtime/OpenGLDrv/Private/OpenGLDevice.cpp b/Engine/Source/Runtime/OpenGLDrv/Private/OpenGLDevice.cpp
--- a/Engine/Source/Runtime/OpenGLDrv/Private/OpenGLDevice.cpp
+++ b/Engine/Source/Runtime/OpenGLDrv/Private/OpenGLDevice.cpp
@@ -46,16 +46,12 @@ static FOpenGLDynamicRHI* PrivateOpenGLDevicePtr = NULL;
 /** true if we're not using UBOs. (ES2) */
 bool GUseEmulatedUniformBuffers;
 
-#ifndef EXPERIMENTAL_OPENGL_RHITHREAD
-#define EXPERIMENTAL_OPENGL_RHITHREAD 0
-#endif
-
 static TAutoConsoleVariable<int32> CVarAllowRGLHIThread(
 	TEXT("r.OpenGL.AllowRHIThread"),
-	0,
-	TEXT("Toggle experimental OpenGL RHI thread support.\n")
-	TEXT("0: GL scene rendering operations are performed on the render thread. (default)\n")
-	TEXT("1: GL scene rendering operations are queued onto the RHI thread gaining some parallelism with the render thread. (mobile feature levels only)"),
+	1,
+	TEXT("Toggle OpenGL RHI thread support.\n")
+	TEXT("0: GL scene rendering operations are performed on the render thread.\n")
+	TEXT("1: GL scene rendering operations are queued onto the RHI thread gaining some parallelism with the render thread. (default, mobile feature levels only)"),
 	ECVF_RenderThreadSafe | ECVF_ReadOnly);
 
 void OnQueryCreation( FOpenGLRenderQuery* Query )
@@ -793,14 +789,8 @@ static void InitRHICapabilitiesForGL()
 	GMaxRHIFeatureLevel = FOpenGL::GetFeatureLevel();
 	GMaxRHIShaderPlatform = FOpenGL::GetShaderPlatform();
 	 
-	// Only enable the experimental OGL rhi thread if explicitly requested.
-#if EXPERIMENTAL_OPENGL_RHITHREAD
-	GRHISupportsRHIThread = GMaxRHIFeatureLevel <= ERHIFeatureLevel::ES3_1;
-#elif ((!PLATFORM_RHITHREAD_DEFAULT_BYPASS) || CAN_TOGGLE_COMMAND_LIST_BYPASS)
-	GRHISupportsRHIThread = GMaxRHIFeatureLevel <= ERHIFeatureLevel::ES3_1 && CVarAllowRGLHIThread.GetValueOnAnyThread();
-#else
-	GRHISupportsRHIThread = false;
-#endif
+	// Enable the OGL rhi thread if explicitly requested.
+	GRHISupportsRHIThread = (GMaxRHIFeatureLevel <= ERHIFeatureLevel::ES3_1 && CVarAllowRGLHIThread.GetValueOnAnyThread());
 	
 	// By default use emulated UBs on mobile
 	static auto* CVar = IConsoleManager::Get().FindTConsoleVariableDataInt(TEXT("OpenGL.UseEmulatedUBs"));
