Fix a bug that a texture will be force to use point filter.

diff --git a/Engine/Source/Runtime/OpenGLDrv/Private/OpenGLCommands.cpp b/Engine/Source/Runtime/OpenGLDrv/Private/OpenGLCommands.cpp
--- a/Engine/Source/Runtime/OpenGLDrv/Private/OpenGLCommands.cpp
+++ b/Engine/Source/Runtime/OpenGLDrv/Private/OpenGLCommands.cpp
@@ -826,8 +826,6 @@ void FOpenGLDynamicRHI::SetupTexturesForDraw( FOpenGLContextState& ContextState,
 					FOpenGL::BindSampler(TextureStageIndex, PendingSampler ? PendingSampler->Resource : 0);
 					ContextState.SamplerStates[TextureStageIndex] = PendingSampler;
 				}
-				//Clean the Pending Sampler State so the textures sampled by TexelFetch without setting point sampler could use sampler 0.
-				PendingState.SamplerStates[TextureStageIndex] = nullptr;
 			}
 			else if (TextureStage.Target != GL_TEXTURE_BUFFER)
 			{
