Log GLES shader compile and program link failures in shipping builds too.
#rnx

diff --git a/Engine/Source/Runtime/OpenGLDrv/Private/OpenGLShaders.cpp b/Engine/Source/Runtime/OpenGLDrv/Private/OpenGLShaders.cpp
--- a/Engine/Source/Runtime/OpenGLDrv/Private/OpenGLShaders.cpp
+++ b/Engine/Source/Runtime/OpenGLDrv/Private/OpenGLShaders.cpp
@@ -74,6 +74,14 @@ static TAutoConsoleVariable<int32> CVarIgnoreLinkFailure(
 	TEXT("1: Ignore link failures. this may allow a program to continue but could lead to undefined rendering behaviour."),
 	ECVF_RenderThreadSafe);
 
+static TAutoConsoleVariable<int32> CVarIgnoreShaderCompileFailure(
+	TEXT("r.OpenGL.IgnoreShaderCompileFailure"),
+	0,
+	TEXT("Ignore OpenGL shader compile failures.\n")
+	TEXT("0: Shader compile failure return an error when encountered. (default)\n")
+	TEXT("1: Ignore Shader compile failures."),
+	ECVF_RenderThreadSafe);
+
 static TAutoConsoleVariable<int32> CVarUseExistingBinaryFileCache(
 	TEXT("r.OpenGL.UseExistingBinaryFileCache"),
 	1,
@@ -95,7 +103,39 @@ static FAutoConsoleVariableRef CVarMaxShaderLibProcessingTime(
 bool GOpenGLShaderHackLastCompileSuccess = false;
 #endif
 
-#define VERIFY_GL_SHADER_LINK (UE_BUILD_DEBUG || UE_BUILD_DEVELOPMENT || UE_BUILD_TEST)
+#define VERIFY_GL_SHADER_LINK 1
+#define VERIFY_GL_SHADER_COMPILE 1
+
+static bool ReportShaderCompileFailures()
+{
+	bool bReportCompileFailures = true;
+#if PLATFORM_ANDROID
+	const FString * ConfigRulesReportGLShaderCompileFailures = FAndroidMisc::GetConfigRulesVariable(TEXT("ReportGLShaderCompileFailures"));
+	bReportCompileFailures = ConfigRulesReportGLShaderCompileFailures == nullptr || ConfigRulesReportGLShaderCompileFailures->Equals("true", ESearchCase::IgnoreCase);
+#endif
+
+#if VERIFY_GL_SHADER_COMPILE
+	return bReportCompileFailures;
+#else
+	return false;
+#endif
+}
+
+static bool ReportProgramLinkFailures()
+{
+	bool bReportLinkFailures = true;
+#if PLATFORM_ANDROID
+	const FString* ConfigRulesReportGLProgramLinkFailures = FAndroidMisc::GetConfigRulesVariable(TEXT("ReportGLProgramLinkFailures"));
+	bReportLinkFailures = ConfigRulesReportGLProgramLinkFailures == nullptr || ConfigRulesReportGLProgramLinkFailures->Equals("true", ESearchCase::IgnoreCase);
+#endif
+
+#if VERIFY_GL_SHADER_LINK
+	return bReportLinkFailures;
+#else
+	return false;
+#endif
+}
+
 
 static uint32 GCurrentDriverProgramBinaryAllocation = 0;
 static uint32 GNumPrograms = 0;
@@ -183,59 +223,42 @@ FORCEINLINE void FOpenGLShaderParameterCache::FRange::MarkDirtyRange(uint32 NewS
 	}
 }
 
-
-enum class VerifyProgramPipelineFailurePolicy : uint8
-{
-	CatchFailure,
-	LogFailure
-};
-
 /**
  * Verify that an OpenGL program has linked successfully.
  */
-static bool VerifyLinkedProgram(GLuint Program, VerifyProgramPipelineFailurePolicy FailurePolicy)
-{
-	auto LogError = [](const FString&& InMessage, VerifyProgramPipelineFailurePolicy InFailurePolicy)
+static bool VerifyLinkedProgram(GLuint Program)
 {
-		bool bCatchError = InFailurePolicy == VerifyProgramPipelineFailurePolicy::CatchFailure && CVarIgnoreLinkFailure.GetValueOnAnyThread() == 0;
-		if (bCatchError)
-		{
-			UE_LOG(LogRHI, Fatal, TEXT("%s"), *InMessage);
-		}
-		else
-		{
-			UE_LOG(LogRHI, Error, TEXT("%s"), *InMessage);
-		}
-	};
-
 	SCOPE_CYCLE_COUNTER(STAT_OpenGLShaderLinkVerifyTime);
 
 	GLint LinkStatus = 0;
 	glGetProgramiv(Program, GL_LINK_STATUS, &LinkStatus);
 	if (LinkStatus != GL_TRUE)
 	{
-#if VERIFY_GL_SHADER_LINK
-		GLint LogLength;
-		ANSICHAR DefaultLog[] = "No log";
-		ANSICHAR *CompileLog = DefaultLog;
-		glGetProgramiv(Program, GL_INFO_LOG_LENGTH, &LogLength);
-		if (LogLength > 1)
+		if (ReportProgramLinkFailures())
 		{
-			CompileLog = (ANSICHAR *)FMemory::Malloc(LogLength);
-			glGetProgramInfoLog(Program, LogLength, NULL, CompileLog);
-		}
-		LogError(FString::Printf(TEXT("Failed to link program. Current total programs: %d program binary bytes: %d\n  log:\n%s"),
-			GNumPrograms,
-			GCurrentDriverProgramBinaryAllocation,
-			ANSI_TO_TCHAR(CompileLog)), FailurePolicy);
+			GLint LogLength;
+			ANSICHAR DefaultLog[] = "No log";
+			ANSICHAR *CompileLog = DefaultLog;
+			glGetProgramiv(Program, GL_INFO_LOG_LENGTH, &LogLength);
+			if (LogLength > 1)
+			{
+				CompileLog = (ANSICHAR *)FMemory::Malloc(LogLength);
+				glGetProgramInfoLog(Program, LogLength, NULL, CompileLog);
+			}
+			UE_LOG(LogRHI, Error, TEXT("Failed to link program. Current total programs: %d program binary bytes: %d\n  log:\n%s"),
+				GNumPrograms,
+				GCurrentDriverProgramBinaryAllocation,
+				ANSI_TO_TCHAR(CompileLog));
 
-		if (LogLength > 1)
+			if (LogLength > 1)
+			{
+				FMemory::Free(CompileLog);
+			}
+		}
+		else
 		{
-			FMemory::Free(CompileLog);
+			UE_LOG(LogRHI, Error, TEXT("Failed to link program. Current total programs:%d"), GNumPrograms);
 		}
-#else
-		LogError(FString::Printf(TEXT("Failed to link program. Current total programs:%d"), GNumPrograms), FailurePolicy);
-#endif
 		// if we're required to ignore link failure then we return true here.
 		return CVarIgnoreLinkFailure.GetValueOnAnyThread() == 1;
 	}
@@ -249,27 +272,13 @@ static bool VerifyCompiledShader(GLuint Shader, const ANSICHAR* GlslCode )
 {
 	SCOPE_CYCLE_COUNTER(STAT_OpenGLShaderCompileVerifyTime);
 
-#if UE_BUILD_DEBUG || DEBUG_GL_SHADERS
 	if (FOpenGL::SupportsSeparateShaderObjects() && glIsProgram(Shader))
 	{
-		bool const bCompiledOK = VerifyLinkedProgram(Shader, VerifyProgramPipelineFailurePolicy::LogFailure);
+		bool const bCompiledOK = VerifyLinkedProgram(Shader);
 #if DEBUG_GL_SHADERS
 		if (!bCompiledOK && GlslCode)
 		{
 			UE_LOG(LogRHI,Error,TEXT("Shader:\n%s"),ANSI_TO_TCHAR(GlslCode));
-			
-#if 0
-			const ANSICHAR *Temp = GlslCode;
-			
-			for ( int i = 0; i < 30 && (*Temp != '\0'); ++i )
-			{
-				FString Converted = ANSI_TO_TCHAR( Temp );
-				Converted.LeftChop( 256 );
-				
-				UE_LOG(LogRHI,Display,TEXT("%s"), *Converted );
-				Temp += Converted.Len();
-			}
-#endif
 		}
 #endif
 		return bCompiledOK;
@@ -280,60 +289,49 @@ static bool VerifyCompiledShader(GLuint Shader, const ANSICHAR* GlslCode )
 		glGetShaderiv(Shader, GL_COMPILE_STATUS, &CompileStatus);
 		if (CompileStatus != GL_TRUE)
 		{
-			GLint LogLength;
-			ANSICHAR DefaultLog[] = "No log";
-			ANSICHAR *CompileLog = DefaultLog;
-			glGetShaderiv(Shader, GL_INFO_LOG_LENGTH, &LogLength);
-	#if PLATFORM_ANDROID
-			if ( LogLength == 0 )
+			if (ReportShaderCompileFailures())
 			{
-				// make it big anyway
-				// there was a bug in android 2.2 where glGetShaderiv would return 0 even though there was a error message
-				// https://code.google.com/p/android/issues/detail?id=9953
-				LogLength = 4096;
-			}
-	#endif
-			if (LogLength > 1)
-			{
-				CompileLog = (ANSICHAR *)FMemory::Malloc(LogLength);
-				glGetShaderInfoLog(Shader, LogLength, NULL, CompileLog);
-			}
-
-	#if DEBUG_GL_SHADERS
-			if (GlslCode)
-			{
-				UE_LOG(LogRHI,Error,TEXT("Shader:\n%s"),ANSI_TO_TCHAR(GlslCode));
-
-	#if 0
-				const ANSICHAR *Temp = GlslCode;
-
-				for ( int i = 0; i < 30 && (*Temp != '\0'); ++i )
+				GLint LogLength;
+				ANSICHAR DefaultLog[] = "No log";
+				ANSICHAR *CompileLog = DefaultLog;
+				glGetShaderiv(Shader, GL_INFO_LOG_LENGTH, &LogLength);
+#if PLATFORM_ANDROID
+				if ( LogLength == 0 )
 				{
-					FString Converted = ANSI_TO_TCHAR( Temp );
-					Converted.LeftChop( 256 );
+					// make it big anyway
+					// there was a bug in android 2.2 where glGetShaderiv would return 0 even though there was a error message
+					// https://code.google.com/p/android/issues/detail?id=9953
+					LogLength = 4096;
+				}
+#endif
+				if (LogLength > 1)
+				{
+					CompileLog = (ANSICHAR *)FMemory::Malloc(LogLength);
+					glGetShaderInfoLog(Shader, LogLength, NULL, CompileLog);
+				}
 
-					UE_LOG(LogRHI,Display,TEXT("%s"), *Converted );
-					Temp += Converted.Len();
+#if DEBUG_GL_SHADERS
+				if (GlslCode)
+				{
+					UE_LOG(LogRHI,Error,TEXT("Shader:\n%s"),ANSI_TO_TCHAR(GlslCode));
+				}
+#endif
+				UE_LOG(LogRHI,Error,TEXT("Failed to compile shader. Compile log:\n%s"), ANSI_TO_TCHAR(CompileLog));
+				if (LogLength > 1)
+				{
+					FMemory::Free(CompileLog);
 				}
-	#endif
-			}
-	#endif
-			UE_LOG(LogRHI,Fatal,TEXT("Failed to compile shader. Compile log:\n%s"), ANSI_TO_TCHAR(CompileLog));
-			if (LogLength > 1)
-			{
-				FMemory::Free(CompileLog);
 			}
-			return false;
+			// if we're required to ignore compile failure then we return true here, it will end with link failure.
+			return CVarIgnoreShaderCompileFailure.GetValueOnAnyThread() == 1;
 		}
 	}
-#endif
 	return true;
 }
 
-// Verify a program has created successfully.
-// FailurePolicy: default to fatal logging on failure, VerifyProgramPipelineFailurePolicy::LogFailure logs any errors and passes back success status.
+// Verify a program has created successfully, the non-SSO case will log errors and return back success status.
 // TODO: SupportsSeparateShaderObjects case.
-static bool VerifyProgramPipeline(GLuint Program, VerifyProgramPipelineFailurePolicy FailurePolicy = VerifyProgramPipelineFailurePolicy::CatchFailure)
+static bool VerifyProgramPipeline(GLuint Program)
 {
 	VERIFY_GL_SCOPE();
 	bool bOK = true;
@@ -346,7 +344,7 @@ static bool VerifyProgramPipeline(GLuint Program, VerifyProgramPipelineFailurePo
 	}
 	else
 	{
-		bOK = VerifyLinkedProgram(Program, FailurePolicy);
+		bOK = VerifyLinkedProgram(Program);
 	}
 	return bOK;
 }
@@ -855,7 +853,7 @@ ShaderType* CompileOpenGLShader(TArrayView<const uint8> InShaderCode, const FSHA
 				glAttachShader(SeparateResource, Resource);
 				
 				glLinkProgram(SeparateResource);
-				bool const bLinkedOK = VerifyLinkedProgram(SeparateResource, VerifyProgramPipelineFailurePolicy::LogFailure);
+				bool const bLinkedOK = VerifyLinkedProgram(SeparateResource);
 				if (!bLinkedOK)
 				{
 					const ANSICHAR* GlslCodeString = GlslCode.GetData();
@@ -1700,7 +1698,7 @@ static bool CreateGLProgramFromUncompressedBinary(GLuint& ProgramOUT, const TArr
 	//	UE_LOG(LogRHI, Warning, TEXT("LRU: CreateFromBinary %d, binary format: %x, BinSize: %d"), GLProgramName, ((GLenum*)ProgramBinaryPtr)[0], BinarySize - sizeof(GLenum));
 
 	ProgramOUT = GLProgramName;
-	return VerifyLinkedProgram(GLProgramName, VerifyProgramPipelineFailurePolicy::LogFailure);
+	return VerifyLinkedProgram(GLProgramName);
 }
 
 struct FCompressedProgramBinaryHeader
@@ -2913,7 +2911,7 @@ static FOpenGLLinkedProgram* LinkProgram( const FOpenGLLinkedProgramConfiguratio
 		}
 	}
 
-	if (VerifyProgramPipeline(Program, VerifyProgramPipelineFailurePolicy::LogFailure))
+	if (VerifyProgramPipeline(Program))
 	{
 		if(bShouldLinkProgram && !FOpenGL::SupportsSeparateShaderObjects())
 		{
@@ -4924,7 +4922,12 @@ void FOpenGLProgramBinaryCache::CompilePendingShaders(const FOpenGLLinkedProgram
 			{
 				TArray<ANSICHAR> GlslCode;
 				UncompressShader(*PendingShaderCodePtr, GlslCode);
-				CompileCurrentShader(ShaderResource, GlslCode);
+				if(CompileCurrentShader(ShaderResource, GlslCode) != GL_TRUE)
+				{
+					// log shader compile failure.
+					VerifyCompiledShader(ShaderResource, GlslCode.Num() > 0 ? GlslCode.GetData() : nullptr);
+				}
+
 				CachePtr->ShadersPendingCompilation.Remove(ShaderResource);
 			}
 		}
@@ -5053,7 +5056,6 @@ void FOpenGLProgramBinaryCache::CompleteLoadedGLProgramRequest_internal(FGLProgr
 			RHIGetPanicDelegate().ExecuteIfBound(FName("FailedBinaryProgramCreateLoadRequest"));
 			UE_LOG(LogRHI, Fatal, TEXT("CompleteLoadedGLProgramRequest_internal : Failed to create GL program from binary data! [%s]"), *ProgramKey.ToString());
 		}
-		VerifyProgramPipeline(PendingGLCreate->GLProgramId);
 		FOpenGLLinkedProgram* NewLinkedProgram = new FOpenGLLinkedProgram(ProgramKey, PendingGLCreate->GLProgramId);
 		GetOpenGLProgramsCache().Add(ProgramKey, NewLinkedProgram);
 		PendingGLCreate->GLProgramState = FGLProgramBinaryFileCacheEntry::EGLProgramState::ProgramAvailable;
