Fix a crash bug on OpenGL platform if the LoadAction of the frame buffer is NoAction and the ClearColor is unbound.

diff --git a/Engine/Source/Runtime/OpenGLDrv/Private/OpenGLCommands.cpp b/Engine/Source/Runtime/OpenGLDrv/Private/OpenGLCommands.cpp
--- a/Engine/Source/Runtime/OpenGLDrv/Private/OpenGLCommands.cpp
+++ b/Engine/Source/Runtime/OpenGLDrv/Private/OpenGLCommands.cpp
@@ -1680,32 +1680,47 @@ void FOpenGLDynamicRHI::RHISetRenderTargetsAndClear(const FRHISetRenderTargetsIn
 	this->RHISetRenderTargets(RenderTargetsInfo.NumColorRenderTargets,
 		RenderTargetsInfo.ColorRenderTarget,
 		&RenderTargetsInfo.DepthStencilRenderTarget);
-	if (RenderTargetsInfo.bClearColor || RenderTargetsInfo.bClearStencil || RenderTargetsInfo.bClearDepth)
-	{
-		FLinearColor ClearColors[MaxSimultaneousRenderTargets];
-		float DepthClear = 0.0;
-		uint32 StencilClear = 0;
 
-		if (RenderTargetsInfo.bClearColor)
+	/**
+	 * Convert all load action from NoAction to Clear for OpenGL platform, especially for mobile to avoid an unnecessary load action.
+	 */
+	bool bClearColor = RenderTargetsInfo.bClearColor;
+	bool bClearStencil = RenderTargetsInfo.bClearStencil;
+	bool bClearDepth = RenderTargetsInfo.bClearDepth;
+
+	FLinearColor ClearColors[MaxSimultaneousRenderTargets];
+	float DepthClear = 0.0;
+	uint32 StencilClear = 0;
+
+	for (int32 i = 0; i < RenderTargetsInfo.NumColorRenderTargets; ++i)
+	{
+		if (RenderTargetsInfo.ColorRenderTarget[i].Texture != nullptr)
 		{
-			for (int32 i = 0; i < RenderTargetsInfo.NumColorRenderTargets; ++i)
-			{
-				if (RenderTargetsInfo.ColorRenderTarget[i].Texture != nullptr)
-				{
-					const FClearValueBinding& ClearValue = RenderTargetsInfo.ColorRenderTarget[i].Texture->GetClearBinding();
-					checkf(ClearValue.ColorBinding == EClearBinding::EColorBound, TEXT("Texture: %s does not have a color bound for fast clears"), *RenderTargetsInfo.ColorRenderTarget[i].Texture->GetName().GetPlainNameString());
-					ClearColors[i] = ClearValue.GetClearColor();
-				}
-			}
+			bClearColor |= RenderTargetsInfo.ColorRenderTarget[i].LoadAction == ERenderTargetLoadAction::ENoAction;
+
+			const FClearValueBinding& ClearValue = RenderTargetsInfo.ColorRenderTarget[i].Texture->GetClearBinding();
+			
+			ClearColors[i] = ClearValue.ColorBinding == EClearBinding::EColorBound ? ClearValue.GetClearColor() : FLinearColor::Black;
 		}
-		if (RenderTargetsInfo.bClearDepth || RenderTargetsInfo.bClearStencil)
+	}
+
+	if (RenderTargetsInfo.DepthStencilRenderTarget.Texture != nullptr)
+	{
+		bClearStencil |= RenderTargetsInfo.DepthStencilRenderTarget.StencilLoadAction == ERenderTargetLoadAction::ENoAction;
+
+		bClearDepth |= RenderTargetsInfo.DepthStencilRenderTarget.DepthLoadAction == ERenderTargetLoadAction::ENoAction;
+
+		const FClearValueBinding& ClearValue = RenderTargetsInfo.DepthStencilRenderTarget.Texture->GetClearBinding();
+
+		if (ClearValue.ColorBinding == EClearBinding::EDepthStencilBound)
 		{
-			const FClearValueBinding& ClearValue = RenderTargetsInfo.DepthStencilRenderTarget.Texture->GetClearBinding();
-			checkf(ClearValue.ColorBinding == EClearBinding::EDepthStencilBound, TEXT("Texture: %s does not have a DS value bound for fast clears"), *RenderTargetsInfo.DepthStencilRenderTarget.Texture->GetName().GetPlainNameString());
 			ClearValue.GetDepthStencil(DepthClear, StencilClear);
 		}
+	}
 
-		this->RHIClearMRT(RenderTargetsInfo.bClearColor, RenderTargetsInfo.NumColorRenderTargets, ClearColors, RenderTargetsInfo.bClearDepth, DepthClear, RenderTargetsInfo.bClearStencil, StencilClear);
+	if (bClearColor || bClearStencil || bClearDepth)
+	{
+		this->RHIClearMRT(bClearColor, RenderTargetsInfo.NumColorRenderTargets, ClearColors, bClearDepth, DepthClear, bClearStencil, StencilClear);
 	}
 }
 
diff --git a/Engine/Source/Runtime/RHI/Private/RHI.cpp b/Engine/Source/Runtime/RHI/Private/RHI.cpp
--- a/Engine/Source/Runtime/RHI/Private/RHI.cpp
+++ b/Engine/Source/Runtime/RHI/Private/RHI.cpp
@@ -1107,7 +1107,7 @@ void FRHIRenderPassInfo::ConvertToRenderTargetsInfo(FRHISetRenderTargetsInfo& Ou
 		OutRTInfo.ColorRenderTarget[Index].MipIndex = ColorRenderTargets[Index].MipIndex;
 		++OutRTInfo.NumColorRenderTargets;
 
-		OutRTInfo.bClearColor |= (LoadAction == ERenderTargetLoadAction::EClear) || (IsOpenGLPlatform(GMaxRHIShaderPlatform) && LoadAction == ERenderTargetLoadAction::ENoAction);
+		OutRTInfo.bClearColor |= (LoadAction == ERenderTargetLoadAction::EClear);
 
 		ensure(!OutRTInfo.bHasResolveAttachments || ColorRenderTargets[Index].ResolveTarget);
 		if (ColorRenderTargets[Index].ResolveTarget)
