Expand generic FGenericPlatformRHIFramePacer to get/set frame pace rate rate mapping to rhi.syncinterval
Add Android OpenGL FPlatformRHIFramePacer support optionally using swappy and relocate code from AndroidEGL
Add FIOSPlatformRHIFramePacer get/set frame pace functions to set rhi.syncinterval taking into account GetMaxRefreshRate

diff --git a/Engine/Source/Runtime/ApplicationCore/Private/Android/AndroidPlatformFramePacer.cpp b/Engine/Source/Runtime/ApplicationCore/Private/Android/AndroidPlatformFramePacer.cpp
new file mode 100644
--- /dev/null
+++ b/Engine/Source/Runtime/ApplicationCore/Private/Android/AndroidPlatformFramePacer.cpp
@@ -0,0 +1,108 @@
+// Copyright Epic Games, Inc. All Rights Reserved.
+
+#include "Android/AndroidPlatformFramePacer.h"
+#include "HAL/IConsoleManager.h"
+#include "Misc/ConfigCacheIni.h"
+#include "Misc/Parse.h"
+#include "Misc/CommandLine.h"
+#include "RHI.h"
+#include "RHIUtilities.h"
+
+/*******************************************************************
+ * FAndroidPlatformRHIFramePacer implementation
+ *******************************************************************/
+
+TAutoConsoleVariable<int32> FAndroidPlatformRHIFramePacer::CVarUseSwappyForFramePacing(
+	TEXT("a.UseSwappyForFramePacing"),
+	0,
+	TEXT("True to use Swappy for frame pacing."));
+
+TAutoConsoleVariable<int32> FAndroidPlatformRHIFramePacer::CVarAllowFrameTimestamps(
+	TEXT("a.AllowFrameTimestamps"),
+	1,
+	TEXT("True to allow the use use eglGetFrameTimestampsANDROID et al for frame pacing or spew."));
+
+TAutoConsoleVariable<int32> FAndroidPlatformRHIFramePacer::CVarTimeStampErrorRetryCount(
+	TEXT("a.TimeStampErrorRetryCount"),
+	100,
+	TEXT("Number of consecutive frames eglGetFrameTimestampsANDROID can fail before reverting to the naive frame pacer.\n")
+	TEXT("Retry count is reset after any successful eglGetFrameTimestampsANDROID calls or a change to sync interval."));
+
+
+// Legacy pacer stuff
+TAutoConsoleVariable<int32> FAndroidPlatformRHIFramePacer::CVarUseGetFrameTimestamps(
+	TEXT("a.UseFrameTimeStampsForPacing"),
+	0,
+	TEXT("True to use eglGetFrameTimestampsANDROID for frame pacing on android (if supported). Only active if a.UseChoreographer is false or the various things needed to use that are not available."));
+
+TAutoConsoleVariable<int32> FAndroidPlatformRHIFramePacer::CVarSpewGetFrameTimestamps(
+	TEXT("a.SpewFrameTimeStamps"),
+	0,
+	TEXT("True to information about frame pacing to the log (if supported). Setting this to 2 results in more detail."));
+
+TAutoConsoleVariable<float> FAndroidPlatformRHIFramePacer::CVarStallSwap(
+	TEXT("CriticalPathStall.Swap"),
+	0.0f,
+	TEXT("Sleep for the given time after the swap (android only for now). Time is given in ms. This is a debug option used for critical path analysis and forcing a change in the critical path."));
+
+TAutoConsoleVariable<int32> FAndroidPlatformRHIFramePacer::CVarDisableOpenGLGPUSync(
+	TEXT("r.Android.DisableOpenGLGPUSync"),
+	1,
+	TEXT("When true, android OpenGL will not prevent the GPU from running more than one frame behind. This will allow higher performance on some devices but increase input latency."),
+	ECVF_RenderThreadSafe);
+
+IAndroidFramePacer* FAndroidPlatformRHIFramePacer::FramePacer = nullptr;
+int32 FAndroidPlatformRHIFramePacer::InternalFramePace = 0;
+
+void FAndroidPlatformRHIFramePacer::Init(IAndroidFramePacer* InFramePacer)
+{
+	FramePacer = InFramePacer;
+	FramePacer->Init();
+}
+
+bool FAndroidPlatformRHIFramePacer::IsEnabled()
+{  
+	return FramePacer != nullptr;
+}
+
+int32 FAndroidPlatformRHIFramePacer::GetFramePace()
+{
+	if (InternalFramePace == 0)
+	{
+		// Internal frame pacing FPS is not set, fall back to generic framepacer based on RHI.SyncInterval
+		int32 SyncIntervalFramePace = FGenericPlatformRHIFramePacer::GetFramePaceFromSyncInterval();
+		InternalFramePace = SyncIntervalFramePace;
+	}
+
+	return InternalFramePace;
+}
+
+int32 FAndroidPlatformRHIFramePacer::SetFramePace(int32 InFramePace)
+{
+	int32 NewFramePace = SupportsFramePace(InFramePace) ? InFramePace : 0;
+
+	// Call generic framepacer to update rhi.SyncInterval where possible
+	FGenericPlatformRHIFramePacer::SetFramePaceToSyncInterval(InFramePace);
+
+	/* Update cvar if necessary */
+	if (InternalFramePace != NewFramePace)
+	{
+		InternalFramePace = NewFramePace;
+	}
+
+	return NewFramePace;
+}
+
+int32 FAndroidPlatformRHIFramePacer::GetLegacySyncInterval()
+{
+	return RHIGetSyncInterval();
+}
+
+void FAndroidPlatformRHIFramePacer::Destroy()
+{
+	if (FramePacer != nullptr)
+	{
+		delete FramePacer;
+		FramePacer = nullptr;
+	}
+}
diff --git a/Engine/Source/Runtime/ApplicationCore/Private/GenericPlatform/GenericPlatformFramePacer.cpp b/Engine/Source/Runtime/ApplicationCore/Private/GenericPlatform/GenericPlatformFramePacer.cpp
new file mode 100644
--- /dev/null
+++ b/Engine/Source/Runtime/ApplicationCore/Private/GenericPlatform/GenericPlatformFramePacer.cpp
@@ -0,0 +1,66 @@
+// Copyright Epic Games, Inc. All Rights Reserved.
+
+#include "GenericPlatform/GenericPlatformFramePacer.h"
+#include "HAL/PlatformApplicationMisc.h"
+#include "HAL/IConsoleManager.h"
+
+int32 FGenericPlatformRHIFramePacer::GetFramePaceFromSyncInterval()
+{
+	static IConsoleVariable* SyncIntervalCVar = IConsoleManager::Get().FindConsoleVariable(TEXT("rhi.SyncInterval"));
+	if (ensure(SyncIntervalCVar != nullptr))
+	{
+		int SyncInterval = SyncIntervalCVar->GetInt();
+		if (SyncInterval > 0)
+		{
+			return 60 / SyncInterval;
+		}
+	}
+
+	return 0;
+}
+
+int32 FGenericPlatformRHIFramePacer::GetFramePace()
+{
+	int32 SyncIntervalFramePace = GetFramePaceFromSyncInterval();
+	return SyncIntervalFramePace;
+}
+
+bool FGenericPlatformRHIFramePacer::SupportsFramePace(int32 QueryFramePace)
+{
+	// Support frame rates that are an integer multiple of max refresh rate, or 0 for no pacing
+	return QueryFramePace == 0 || (60 % QueryFramePace) == 0;
+}
+
+int32 FGenericPlatformRHIFramePacer::SetFramePace(int32 InFramePace)
+{
+	int32 NewFramePace = SetFramePaceToSyncInterval(InFramePace);
+	return NewFramePace;
+}
+
+int32 FGenericPlatformRHIFramePacer::SetFramePaceToSyncInterval(int32 InFramePace)
+{
+	int32 NewPace = 0;
+
+	// Disable frame pacing if an unsupported frame rate is requested
+	if (!SupportsFramePace(InFramePace))
+	{
+		InFramePace = 0;
+	}
+
+	static IConsoleVariable* SyncIntervalCVar = IConsoleManager::Get().FindConsoleVariable(TEXT("rhi.SyncInterval"));
+	if (ensure(SyncIntervalCVar != nullptr))
+	{
+		int32 NewSyncInterval = InFramePace > 0 ? 60 / InFramePace : 0;
+
+		SyncIntervalCVar->Set(NewSyncInterval, ECVF_SetByCode);
+
+		if (NewSyncInterval > 0)
+		{
+			NewPace = 60 / NewSyncInterval;
+		}
+	}
+
+	return NewPace;
+}
+
+
diff --git a/Engine/Source/Runtime/ApplicationCore/Private/HAL/PlatformFramePacer.cpp b/Engine/Source/Runtime/ApplicationCore/Private/HAL/PlatformFramePacer.cpp
new file mode 100644
--- /dev/null
+++ b/Engine/Source/Runtime/ApplicationCore/Private/HAL/PlatformFramePacer.cpp
@@ -0,0 +1,25 @@
+// Copyright Epic Games, Inc. All Rights Reserved.
+
+#include "HAL/PlatformFramePacer.h"
+#include "HAL/IConsoleManager.h"
+
+DEFINE_LOG_CATEGORY_STATIC(LogFramePacer, Log, All);
+
+static FAutoConsoleCommand CFramePace(
+	TEXT("r.SetFramePace"),
+	TEXT("Set a target frame rate for the frame pacer.")
+	TEXT("To set 30fps: \"r.SetFramePace 30\""),
+	FConsoleCommandWithArgsDelegate::CreateLambda([](const TArray<FString>& Args)
+{
+	if (Args.Num() > 0)
+	{
+		uint32 RequestFramePace = FCString::Atod(*Args[0]);
+		uint32 ResultFramePace = FPlatformRHIFramePacer::SetFramePace(RequestFramePace);
+		UE_LOG(LogFramePacer, Display, TEXT("r.FramePace : requesting %d, set as %d"), RequestFramePace, ResultFramePace);
+	}
+	else
+	{
+		UE_LOG(LogFramePacer, Display, TEXT("r.FramePace : current frame pace is %d"), FPlatformRHIFramePacer::GetFramePace());
+	}
+}),
+ECVF_Cheat);
\ No newline at end of file
diff --git a/Engine/Source/Runtime/ApplicationCore/Private/IOS/IOSPlatformFramePacer.cpp b/Engine/Source/Runtime/ApplicationCore/Private/IOS/IOSPlatformFramePacer.cpp
--- a/Engine/Source/Runtime/ApplicationCore/Private/IOS/IOSPlatformFramePacer.cpp
+++ b/Engine/Source/Runtime/ApplicationCore/Private/IOS/IOSPlatformFramePacer.cpp
@@ -218,6 +218,46 @@ uint32 FIOSPlatformRHIFramePacer::GetMaxRefreshRate()
 	return bEnableDynamicMaxFPS ? [UIScreen mainScreen].maximumFramesPerSecond : IOSDisplayConstants::MaxRefreshRate;
 }
 
+bool FIOSPlatformRHIFramePacer::SupportsFramePace(int32 QueryFramePace)
+{
+	// Support frame rates that are an integer multiple of max refresh rate, or 0 for no pacing
+	return QueryFramePace == 0 || (GetMaxRefreshRate() % QueryFramePace) == 0;
+}
+
+int32 FIOSPlatformRHIFramePacer::SetFramePace(int32 InFramePace)
+{
+	int32 NewPace = 0;
+
+	static IConsoleVariable* SyncIntervalCVar = IConsoleManager::Get().FindConsoleVariable(TEXT("rhi.SyncInterval"));
+	if (ensure(SyncIntervalCVar != nullptr))
+	{
+		int32 MaxRefreshRate = GetMaxRefreshRate();
+		int32 NewSyncInterval = InFramePace > 0 ? MaxRefreshRate / InFramePace : 0;
+		SyncIntervalCVar->Set(NewSyncInterval, ECVF_SetByCode);
+
+		if (NewSyncInterval > 0)
+		{
+			NewPace = MaxRefreshRate / NewSyncInterval;
+		}
+	}
+	return NewPace;
+}
+
+int32 FIOSPlatformRHIFramePacer::GetFramePace()
+{
+	int32 CurrentPace = 0;
+	static IConsoleVariable* SyncIntervalCVar = IConsoleManager::Get().FindConsoleVariable(TEXT("rhi.SyncInterval"));
+	if (ensure(SyncIntervalCVar != nullptr))
+	{
+		int SyncInterval = SyncIntervalCVar->GetInt();
+		if (SyncInterval > 0)
+		{
+			CurrentPace = GetMaxRefreshRate() / SyncInterval;
+		}
+	}
+	return CurrentPace;
+}
+
 void FIOSPlatformRHIFramePacer::InitWithEvent(FEvent* TriggeredEvent)
 {
     // Create display link thread
diff --git a/Engine/Source/Runtime/ApplicationCore/Public/Android/AndroidPlatformFramePacer.h b/Engine/Source/Runtime/ApplicationCore/Public/Android/AndroidPlatformFramePacer.h
new file mode 100644
--- /dev/null
+++ b/Engine/Source/Runtime/ApplicationCore/Public/Android/AndroidPlatformFramePacer.h
@@ -0,0 +1,85 @@
+// Copyright Epic Games, Inc. All Rights Reserved.
+
+/*=============================================================================================
+	AndroidPlatformFramePacer.h: Android platform frame pacer classes.
+==============================================================================================*/
+
+#pragma once
+#include "GenericPlatform/GenericPlatformFramePacer.h"
+
+/**
+ * Android implementation of FGenericPlatformRHIFramePacer
+ **/
+
+struct IAndroidFramePacer
+{
+	virtual void Init() = 0;
+	virtual ~IAndroidFramePacer() {}
+	virtual bool SwapBuffers(bool bLockToVsync) { return true; }
+	virtual bool SupportsFramePace(int32 QueryFramePace) = 0;
+};
+
+struct FAndroidOpenGLFramePacer : public IAndroidFramePacer
+{
+	virtual void Init() override;
+	virtual ~FAndroidOpenGLFramePacer();
+	virtual bool SwapBuffers(bool bLockToVsync) override;
+	virtual bool SupportsFramePace(int32 QueryFramePace) override;
+
+private:
+	void InitSwappy();
+	// legacy framepacer state
+	int32 DesiredSyncIntervalRelativeTo60Hz = -1;
+	int32 DesiredSyncIntervalRelativeToDevice = -1;
+	int32 DriverSyncIntervalRelativeToDevice = -1;
+	float DriverRefreshRate = 60.0f;
+	int64 DriverRefreshNanos = 16666666;
+	double LastTimeEmulatedSync = -1.0;
+	uint32 SwapBufferFailureCount = 0;
+
+	bool bSwappyInit = false;
+};
+
+struct FAndroidVulkanFramePacer : public IAndroidFramePacer
+{
+	// Temporary stub until we enable Swappy on Vulkan
+	virtual void Init() override {}
+	virtual ~FAndroidVulkanFramePacer() {}
+	virtual bool SupportsFramePace(int32 QueryFramePace) override { return FGenericPlatformRHIFramePacer::SupportsFramePace(QueryFramePace); }
+};
+
+struct FAndroidPlatformRHIFramePacer : public FGenericPlatformRHIFramePacer
+{
+    // FGenericPlatformRHIFramePacer interface
+    static bool IsEnabled();
+    static void Destroy();
+	static int32 GetFramePace();
+	static int32 SetFramePace(int32 FramePace);
+	static bool SupportsFramePace(int32 QueryFramePace) { return (ensure(FramePacer)) ? FramePacer->SupportsFramePace(QueryFramePace) : false; }
+
+	// FAndroidPlatformRHIFramePacer interface
+	static void Init(IAndroidFramePacer* InFramePacer);
+	static int32 GetLegacySyncInterval();
+	static void SwapBuffers(bool bLockToVsync) { if (ensure(FramePacer)) { FramePacer->SwapBuffers(bLockToVsync); } }
+
+	static TAutoConsoleVariable<int32> CVarUseSwappyForFramePacing;
+
+	// Legacy pacer stuff
+	static TAutoConsoleVariable<int32> CVarAllowFrameTimestamps;
+	static TAutoConsoleVariable<int32> CVarTimeStampErrorRetryCount;
+	static TAutoConsoleVariable<int32> CVarUseGetFrameTimestamps;
+	static TAutoConsoleVariable<int32> CVarSpewGetFrameTimestamps;
+	static TAutoConsoleVariable<float> CVarStallSwap;
+	static TAutoConsoleVariable<int32> CVarDisableOpenGLGPUSync;
+	   
+private:
+    /** The actual GL or Vulkan frame pacer */
+	static IAndroidFramePacer* FramePacer;
+	friend struct FAndroidOpenGLFramePacer;
+	friend struct FAndroidVulkanFramePacer;
+
+	/* Actual current frame pace */
+	static int32 InternalFramePace;
+};
+
+typedef FAndroidPlatformRHIFramePacer FPlatformRHIFramePacer;
diff --git a/Engine/Source/Runtime/ApplicationCore/Public/GenericPlatform/GenericPlatformFramePacer.h b/Engine/Source/Runtime/ApplicationCore/Public/GenericPlatform/GenericPlatformFramePacer.h
--- a/Engine/Source/Runtime/ApplicationCore/Public/GenericPlatform/GenericPlatformFramePacer.h
+++ b/Engine/Source/Runtime/ApplicationCore/Public/GenericPlatform/GenericPlatformFramePacer.h
@@ -20,15 +20,7 @@ struct APPLICATIONCORE_API FGenericPlatformRHIFramePacer
      * Should the Frame Pacer be enabled?
      */
     bool IsEnabled() { return false; }
-    
-    /**
-     * Init the RHI frame pacer
-     *
-     * @param InTriggeredEvent - The event we wish to trigger when the frame interval has been triggered by the hardware.
-     * @param InFrameInterval - How often should the event be triggered, in Frames.
-     */
-	static void InitWithEvent(FEvent* InTriggeredEvent) {}
-    
+
     /**
      * Teardown the Frame Pacer.
      */
@@ -36,11 +28,33 @@ struct APPLICATIONCORE_API FGenericPlatformRHIFramePacer
 	
 	/**
 	 * The pace we are running at (30 = 30fps, 0 = unpaced)
+	 * The generic implementation returns a result based on rhi.SyncInterval assuming a 60Hz native refresh rate.
+	 */
+	static int32 GetFramePace();
+
+	/**
+	 * Sets the pace we would like to running at (30 = 30fps, 0 = unpaced).
+	 * The generic implementation sets the value for rhi.SyncInterval assuming a 60Hz native refresh rate.
+	 *
+	 * @return the pace we will run at.
+	 */
+	static int32 SetFramePace(int32 FramePace);
+
+	/**
+	 * Rethers whether the hardware is able to frame pace at the specificed frame rate
+	 */
+	static bool SupportsFramePace(int32 QueryFramePace);
+
+protected:
+	/**
+	 * The generic implementation returns a result based on rhi.SyncInterval assuming a 60Hz native refresh rate.
 	 */
-	static uint32 GetFramePace() { return 0; };
+	static int32 GetFramePaceFromSyncInterval();
 
 	/**
-	 * The maximum refresh rate the pacer supports
+	 * The generic sets rhi.SyncInterval assuming a 60Hz native refresh rate.
+	 *
+	 * @return the pace we will run at.
 	 */
-	static uint32 GetMaxRefreshRate() { return 60; }
+	static int32 SetFramePaceToSyncInterval(int32 FramePace);
 };
diff --git a/Engine/Source/Runtime/ApplicationCore/Public/HAL/PlatformFramePacer.h b/Engine/Source/Runtime/ApplicationCore/Public/HAL/PlatformFramePacer.h
--- a/Engine/Source/Runtime/ApplicationCore/Public/HAL/PlatformFramePacer.h
+++ b/Engine/Source/Runtime/ApplicationCore/Public/HAL/PlatformFramePacer.h
@@ -7,6 +7,8 @@
 #include "IOS/IOSPlatformFramePacer.h"
 #elif PLATFORM_MAC
 #include "Mac/MacPlatformFramePacer.h"
+#elif PLATFORM_ANDROID
+#include "Android/AndroidPlatformFramePacer.h"
 #else
 #include "GenericPlatform/GenericPlatformFramePacer.h"
 typedef FGenericPlatformRHIFramePacer FPlatformRHIFramePacer;
diff --git a/Engine/Source/Runtime/ApplicationCore/Public/IOS/IOSPlatformFramePacer.h b/Engine/Source/Runtime/ApplicationCore/Public/IOS/IOSPlatformFramePacer.h
--- a/Engine/Source/Runtime/ApplicationCore/Public/IOS/IOSPlatformFramePacer.h
+++ b/Engine/Source/Runtime/ApplicationCore/Public/IOS/IOSPlatformFramePacer.h
@@ -21,14 +21,17 @@ struct FIOSPlatformRHIFramePacer : public FGenericPlatformRHIFramePacer
 {
     // FGenericPlatformRHIFramePacer interface
     static bool IsEnabled();
-	static uint32 GetMaxRefreshRate();
+    static void Destroy();
+	static int32 SetFramePace(int32 FramePace);
+	static int32 GetFramePace();
+	static bool SupportsFramePace(int32 QueryFramePace);
+
+	// FIOSPlatformRHIFramePacer interface
 	static void InitWithEvent(class FEvent* TriggeredEvent);
 	static void AddHandler(FIOSFramePacerHandler Handler);
 	static void RemoveHandler(FIOSFramePacerHandler Handler);
-    static void Destroy();
-	static uint32 GetFramePace() { return Pace; };
+	static uint32 GetMaxRefreshRate();
 
-    
     /** Access to the IOS Frame Pacer: CADisplayLink */
     static FIOSFramePacer* FramePacer;
     
diff --git a/Engine/Source/Runtime/OpenGLDrv/OpenGLDrv.Build.cs b/Engine/Source/Runtime/OpenGLDrv/OpenGLDrv.Build.cs
--- a/Engine/Source/Runtime/OpenGLDrv/OpenGLDrv.Build.cs
+++ b/Engine/Source/Runtime/OpenGLDrv/OpenGLDrv.Build.cs
@@ -52,7 +52,22 @@ public class OpenGLDrv : ModuleRules
 			PrivateDependencyModuleNames.Add("detex");
 		}
 
-		if(Target.Platform != UnrealTargetPlatform.Win32 && Target.Platform != UnrealTargetPlatform.Win64
+        if (Target.Platform == UnrealTargetPlatform.Android)
+        {
+            // for Swappy
+            PublicDefinitions.Add("USE_ANDROID_OPENGL_SWAPPY=0");
+/*
+            PrivateDependencyModuleNames.AddRange(
+                new string[]
+                {
+                "Launch",
+                "GoogleGameSDK"
+                }
+            );
+*/
+        }
+
+        if (Target.Platform != UnrealTargetPlatform.Win32 && Target.Platform != UnrealTargetPlatform.Win64
 			&& Target.Platform != UnrealTargetPlatform.IOS && Target.Platform != UnrealTargetPlatform.Android
 			&& !Target.IsInPlatformGroup(UnrealPlatformGroup.Linux)
 			&& Target.Platform != UnrealTargetPlatform.TVOS && Target.Platform != UnrealTargetPlatform.Lumin)
diff --git a/Engine/Source/Runtime/OpenGLDrv/Private/Android/AndroidEGL.cpp b/Engine/Source/Runtime/OpenGLDrv/Private/Android/AndroidEGL.cpp
--- a/Engine/Source/Runtime/OpenGLDrv/Private/Android/AndroidEGL.cpp
+++ b/Engine/Source/Runtime/OpenGLDrv/Private/Android/AndroidEGL.cpp
@@ -12,30 +12,14 @@
 #if USE_ANDROID_JNI
 #include <android/native_window_jni.h>
 #endif
-#include "OpenGLDrvPrivate.h"
 #include "Misc/ScopeLock.h"
 #include "UObject/GarbageCollection.h"
+#include "Android/AndroidPlatformFramePacer.h"
 
 
 AndroidEGL* AndroidEGL::Singleton = NULL;
 DEFINE_LOG_CATEGORY(LogEGL);
 
-static TAutoConsoleVariable<int32> CVarAllowFrameTimestamps(
-	TEXT("a.AllowFrameTimestamps"),
-	1,
-	TEXT("True to allow the use use eglGetFrameTimestampsANDROID et al for frame pacing or spew."));
-
-static TAutoConsoleVariable<int32> CVarTimeStampErrorRetryCount(
-	TEXT("a.TimeStampErrorRetryCount"),
-	100,
-	TEXT("Number of consecutive frames eglGetFrameTimestampsANDROID can fail before reverting to the naive frame pacer.\n")
-	TEXT("Retry count is reset after any successful eglGetFrameTimestampsANDROID calls or a change to sync interval."));
-
-#define ENABLE_CONFIG_FILTER 1
-#define ENABLE_EGL_DEBUG 0
-#define ENABLE_VERIFY_EGL 0
-#define ENABLE_VERIFY_EGL_TRACE 0
-
 #if USE_ANDROID_EGL_NO_ERROR_CONTEXT
 #ifndef EGL_KHR_create_context_no_error
 #define EGL_KHR_create_context_no_error 1
@@ -43,85 +27,7 @@ static TAutoConsoleVariable<int32> CVarTimeStampErrorRetryCount(
 #endif // EGL_KHR_create_context_no_error
 #endif // USE_ANDROID_EGL_NO_ERROR_CONTEXT
 
-#if ENABLE_VERIFY_EGL
-
-#define VERIFY_EGL(msg) { VerifyEGLResult(eglGetError(),TEXT(#msg),TEXT(""),TEXT(__FILE__),__LINE__); }
-
-void VerifyEGLResult(EGLint ErrorCode, const TCHAR* Msg1, const TCHAR* Msg2, const TCHAR* Filename, uint32 Line)
-{
-	if (ErrorCode != EGL_SUCCESS)
-	{
-		static const TCHAR* EGLErrorStrings[] =
-		{
-			TEXT("EGL_NOT_INITIALIZED"),
-			TEXT("EGL_BAD_ACCESS"),
-			TEXT("EGL_BAD_ALLOC"),
-			TEXT("EGL_BAD_ATTRIBUTE"),
-			TEXT("EGL_BAD_CONFIG"),
-			TEXT("EGL_BAD_CONTEXT"),
-			TEXT("EGL_BAD_CURRENT_SURFACE"),
-			TEXT("EGL_BAD_DISPLAY"),
-			TEXT("EGL_BAD_MATCH"),
-			TEXT("EGL_BAD_NATIVE_PIXMAP"),
-			TEXT("EGL_BAD_NATIVE_WINDOW"),
-			TEXT("EGL_BAD_PARAMETER"),
-			TEXT("EGL_BAD_SURFACE"),
-			TEXT("EGL_CONTEXT_LOST"),
-			TEXT("UNKNOWN EGL ERROR")
-		};
 
-		uint32 ErrorIndex = FMath::Min<uint32>(ErrorCode - EGL_SUCCESS, UE_ARRAY_COUNT(EGLErrorStrings) - 1);
-		UE_LOG(LogRHI, Warning, TEXT("%s(%u): %s%s failed with error %s (0x%x)"),
-			Filename, Line, Msg1, Msg2, EGLErrorStrings[ErrorIndex], ErrorCode);
-		check(0);
-	}
-}
-
-class FEGLErrorScope
-{
-public:
-	FEGLErrorScope(
-		const TCHAR* InFunctionName,
-		const TCHAR* InFilename,
-		const uint32 InLine)
-		: FunctionName(InFunctionName)
-		, Filename(InFilename)
-		, Line(InLine)
-	{
-#if ENABLE_VERIFY_EGL_TRACE
-		UE_LOG(LogRHI, Log, TEXT("EGL log before %s(%d): %s"), InFilename, InLine, InFunctionName);
-#endif
-		CheckForErrors(TEXT("Before "));
-	}
-
-	~FEGLErrorScope()
-	{
-#if ENABLE_VERIFY_EGL_TRACE
-		UE_LOG(LogRHI, Log, TEXT("EGL log after  %s(%d): %s"), Filename, Line, FunctionName);
-#endif
-		CheckForErrors(TEXT("After "));
-	}
-
-private:
-	const TCHAR* FunctionName;
-	const TCHAR* Filename;
-	const uint32 Line;
-
-	void CheckForErrors(const TCHAR* PrefixString)
-	{
-		VerifyEGLResult(eglGetError(), PrefixString, FunctionName, Filename, Line);
-	}
-};
-
-#define MACRO_TOKENIZER(IdentifierName, Msg, FileName, LineNumber) FEGLErrorScope IdentifierName_ ## LineNumber (Msg, FileName, LineNumber)
-#define MACRO_TOKENIZER2(IdentifierName, Msg, FileName, LineNumber) MACRO_TOKENIZER(IdentiferName, Msg, FileName, LineNumber)
-#define VERIFY_EGL_SCOPE_WITH_MSG_STR(MsgStr) MACRO_TOKENIZER2(ErrorScope_, MsgStr, TEXT(__FILE__), __LINE__)
-#define VERIFY_EGL_SCOPE() VERIFY_EGL_SCOPE_WITH_MSG_STR(ANSI_TO_TCHAR(__FUNCTION__))
-#define VERIFY_EGL_FUNC(Func, ...) { VERIFY_EGL_SCOPE_WITH_MSG_STR(TEXT(#Func)); Func(__VA_ARGS__); }
-#else
-#define VERIFY_EGL(...)
-#define VERIFY_EGL_SCOPE(...)
-#endif
 const  int EGLMinRedBits		= 5;
 const  int EGLMinGreenBits		= 6;
 const  int EGLMinBlueBits		= 5;
@@ -177,26 +83,11 @@ struct AndroidESPImpl
 	float eglRatio;
 	EGLConfigParms Parms;
 	int DepthSize;
-	uint32_t SwapBufferFailureCount;
 	ANativeWindow* Window;
 	bool Initalized ;
 	EOpenGLCurrentContext CurrentContextType;
 	GLuint OnScreenColorRenderBuffer;
 	GLuint ResolveFrameBuffer;
-	int32 DesiredSyncIntervalRelativeTo60Hz;
-	int32 DesiredSyncIntervalRelativeToDevice;
-	int32 DriverSyncIntervalRelativeToDevice;
-	float DriverRefreshRate;
-	int64 DriverRefreshNanos;
-
-	//unknown google mystery meat, maybe search for open source google code called swappy
-	int64 DriverAppVsyncOffsetNanos; 
-	int64 DriverDeadlineNanos;
-	int64 DriverSlopNanos;
-	EGLSyncKHR SyncFenceForChoreographerMethod;
-
-
-	double LastTimeEmulatedSync;
 	FPlatformRect CachedWindowRect;
 
 	AndroidESPImpl();
@@ -367,7 +258,7 @@ void AndroidEGL::CreateEGLSurface(ANativeWindow* InWindow, bool bCreateWndSurfac
 		//need ANativeWindow
 		PImplData->eglSurface = eglCreateWindowSurface(PImplData->eglDisplay, PImplData->eglConfigParam,InWindow, NULL);
 
-		if (CVarAllowFrameTimestamps.GetValueOnAnyThread())
+		if (FAndroidPlatformRHIFramePacer::CVarAllowFrameTimestamps.GetValueOnAnyThread())
 		{
 			eglSurfaceAttrib(PImplData->eglDisplay, PImplData->eglSurface, EGL_TIMESTAMPS_ANDROID, EGL_TRUE);
 		}
@@ -609,26 +500,13 @@ eglDisplay(EGL_NO_DISPLAY)
 	,eglHeight(8) // required for Gear VR apps with internal win surf mgmt
 	,eglRatio(0)
 	,DepthSize(0)
-	,SwapBufferFailureCount(0)
 	,Window(NULL)
 	,Initalized(false)
 	,OnScreenColorRenderBuffer(0)
 	,ResolveFrameBuffer(0)
 	,NativeVisualID(0)
 	,CurrentContextType(CONTEXT_Invalid)
-	,DesiredSyncIntervalRelativeTo60Hz(-1)
-	,DesiredSyncIntervalRelativeToDevice(-1)
-	,DriverSyncIntervalRelativeToDevice(1)
-	,DriverRefreshRate(60.0f)
-	,DriverRefreshNanos(16666666)
-	,DriverAppVsyncOffsetNanos(2000000)
-	,DriverDeadlineNanos(10666666)
-	,DriverSlopNanos(1000000)
-	,SyncFenceForChoreographerMethod(EGL_NO_SYNC_KHR)
-
-	,LastTimeEmulatedSync(-1.0)
-	,CachedWindowRect(FPlatformRect(0,0,0,0)
-	)
+	,CachedWindowRect(FPlatformRect(0,0,0,0))
 {
 }
 
@@ -852,623 +730,6 @@ int32 AndroidEGL::GetError()
 	return eglGetError();
 }
 
-
-#include <chrono>
-
-static int64 ChoreographerClock()
-{
-	return std::chrono::steady_clock::now().time_since_epoch().count();
-}
-
-extern void StartChoreographer(TFunction<int64(int64)> Callback);
-extern bool ChoreographerIsAvailable();
-
-static TAutoConsoleVariable<int32> CVarUseChoreographerEventArrivalTimes(
-	TEXT("a.UseChoreographerEventArriveTimes"),
-	0,
-	TEXT("On some devices and drivers the choreographer time stamps are wobbly and bad. If this is set to 1, then use the arrival times of the choreographer event instead"));
-
-struct AChoreographer;
-struct FChoreographerFramePacer
-{
-	bool bSetup = false;
-	bool bStopped = false;
-	FCriticalSection ChoreographerSetupLock;
-
-	FEvent* ChoreographerThreadWaitEvent = nullptr;
-	FEvent* RHIThreadWaitEvent = nullptr;
-	bool bRHIThreadWaiting = false;
-
-	int64 NextDelay = -1;
-	int64 SleepTime = -1;
-	int64 SlopTime = 1000000;
-
-	void SetupChoreographer()
-	{
-		if (!bSetup)
-		{
-			bSetup = true;
-			bRHIThreadWaiting = false;
-			if (!ChoreographerThreadWaitEvent)
-			{
-				ChoreographerThreadWaitEvent = FPlatformProcess::GetSynchEventFromPool(false);
-				RHIThreadWaitEvent = FPlatformProcess::GetSynchEventFromPool(false);
-			}
-			else
-			{
-				ChoreographerThreadWaitEvent->Reset();
-				RHIThreadWaitEvent->Reset();
-			}
-			StartChoreographer([this](int64 FrameCounter) {return DoCallback(FrameCounter); });
-		}
-	}
-
-	void StopChoreographer()
-	{
-		FScopeLock Lock(&ChoreographerSetupLock);
-		bSetup = false;
-		if (bRHIThreadWaiting)
-		{
-			bRHIThreadWaiting = false;
-			NextDelay = -1;
-			ChoreographerThreadWaitEvent->Trigger();
-		}
-	}
-
-	int64 DoCallback(int64 FrameTime)
-	{
-		if (!bSetup)
-		{
-			return -1;
-		}
-		int64 WakeUpDelay = (CVarUseChoreographerEventArrivalTimes.GetValueOnAnyThread() > 0) ? 0.0 : (ChoreographerClock() - FrameTime);
-		int64 AdjustedSleepTime = SlopTime + SleepTime - WakeUpDelay;
-
-		if (AdjustedSleepTime > 0)
-		{
-			QUICK_SCOPE_CYCLE_COUNTER(STAT_ChoreographerSleep);
-			//UE_LOG(LogRHI, Log, TEXT("Sleep %6.2fms"), float(AdjustedSleepTime) / 1000000.0f);
-			FPlatformProcess::Sleep(float(AdjustedSleepTime) / 1000000000.0f);
-		}
-		{
-			FScopeLock Lock(&ChoreographerSetupLock);
-			if (!bSetup)
-			{
-				return -1;
-			}
-			if (!bRHIThreadWaiting)
-			{
-				//UE_LOG(LogRHI, Log, TEXT("Missed VSync (CPU)"));
-				return SlopTime;
-			}
-		}
-		RHIThreadWaitEvent->Trigger();
-		{
-			QUICK_SCOPE_CYCLE_COUNTER(STAT_ChoreographerWait);
-			if (!ChoreographerThreadWaitEvent->Wait(FTimespan::FromMilliseconds(3000.0)))
-			{
-				UE_LOG(LogRHI, Fatal, TEXT("Timed out waiting for ChoreographerThreadWaitEvent."));
-			}
-		}
-		FScopeLock Lock(&ChoreographerSetupLock);
-		ChoreographerThreadWaitEvent->Reset();
-		//UE_LOG(LogRHI, Log, TEXT("Next Delay %6.2fms"), float(NextDelay) / 1000000.0f);
-		return NextDelay;
-	}
-	void StartSync(int64 InNextDelay, int64 InSleepTime)
-	{
-		FScopeLock Lock(&ChoreographerSetupLock);
-		SleepTime = InSleepTime;
-		NextDelay = InNextDelay >= 0 ? InNextDelay + SlopTime : SlopTime;
-		SetupChoreographer();
-		if (bRHIThreadWaiting)
-		{
-			bRHIThreadWaiting = false;
-			ChoreographerThreadWaitEvent->Trigger();
-		}
-	}
-	void WaitSync()
-	{
-		QUICK_SCOPE_CYCLE_COUNTER(STAT_StallForChoreographerSyncInterval);
-		check(RHIThreadWaitEvent);
-		{
-			QUICK_SCOPE_CYCLE_COUNTER(STAT_StallForChoreographerSyncInterval_Lock);
-			FScopeLock Lock(&ChoreographerSetupLock);
-			check(!bRHIThreadWaiting);
-			bRHIThreadWaiting = true;
-		}
-		if (!RHIThreadWaitEvent->Wait(FTimespan::FromMilliseconds(3000.0)))
-		{
-			UE_LOG(LogRHI, Fatal, TEXT("Timed out waiting for RHIThreadWaitEvent."));
-		}
-		RHIThreadWaitEvent->Reset();
-	}
-};
-FChoreographerFramePacer TheChoreographerFramePacer;
-
-static TAutoConsoleVariable<int32> CVarUseChoreographer(
-	TEXT("a.UseChoreographerForPacing"),
-	0,
-	TEXT("True to use Choreographer to do frame pacing on android."));
-
-bool ShouldUseChoreographer()
-{
-	// should check the ndk version, etc
-	return CVarUseChoreographer.GetValueOnAnyThread() > 0 && ChoreographerIsAvailable() && eglGetSyncAttribKHR_p;
-}
-
-
-extern int64 AndroidThunkCpp_GetMetaDataLong(const FString& Key);
-extern float AndroidThunkCpp_GetMetaDataFloat(const FString& Key);
-
-static uint32 NextFrameIDSlot = 0;
-#define NUM_FRAMES_TO_MONITOR (4)
-static EGLuint64KHR FrameIDs[NUM_FRAMES_TO_MONITOR] = { 0 };
-
-static int32 RecordedFrameInterval[100];
-static int32 NumRecordedFrameInterval = 0;
-
-static TAutoConsoleVariable<int32> CVarUseGetFrameTimestamps(
-	TEXT("a.UseFrameTimeStampsForPacing"),
-	0,
-	TEXT("True to use eglGetFrameTimestampsANDROID for frame pacing on android (if supported). Only active if a.UseChoreographer is false or the various things needed to use that are not available."));
-
-static TAutoConsoleVariable<int32> CVarSpewGetFrameTimestamps(
-	TEXT("a.SpewFrameTimeStamps"),
-	0,
-	TEXT("True to information about frame pacing to the log (if supported). Setting this to 2 results in more detail."));
-
-static TAutoConsoleVariable<float> CVarStallSwap(
-	TEXT("CriticalPathStall.Swap"),
-	0.0f,
-	TEXT("Sleep for the given time after the swap (android only for now). Time is given in ms. This is a debug option used for critical path analysis and forcing a change in the critical path."));
-
-static TAutoConsoleVariable<int32> CVarDisableOpenGLGPUSync(
-	TEXT("r.Android.DisableOpenGLGPUSync"),
-	1,
-	TEXT("When true, android OpenGL will not prevent the GPU from running more than one frame behind. This will allow higher performance on some devices but increase input latency."),
-	ECVF_RenderThreadSafe);
-
-static bool GGetTimeStampsSucceededThisFrame = true;
-static uint32 GGetTimeStampsRetryCount = 0;
-
-static bool CanUseGetFrameTimestamps()
-{
-	return CVarUseGetFrameTimestamps.GetValueOnAnyThread()
-		&& eglGetFrameTimestampsANDROID_p
-		&& eglGetNextFrameIdANDROID_p
-		&& eglPresentationTimeANDROID_p
-		&& (GGetTimeStampsRetryCount < CVarTimeStampErrorRetryCount.GetValueOnAnyThread());
-}
-static bool CanUseGetFrameTimestampsForThisFrame()
-{
-	return CanUseGetFrameTimestamps() && GGetTimeStampsSucceededThisFrame;
-}
-
-bool ShouldUseGPUFencesToLimitLatency()
-{
-	if (ShouldUseChoreographer())
-	{
-		return false; // this method does its own GPU fences as part of the swap.
-	}
-	if (CanUseGetFrameTimestampsForThisFrame())
-	{
-		return true; // this method requires a GPU fence to give steady results
-	}
-	return CVarDisableOpenGLGPUSync.GetValueOnAnyThread() == 0; // otherwise just based on the cvar; thought to be bad to use GPU fences on PowerVR
-}
-
-
-bool AndroidEGL::SwapBuffers(int32 SyncInterval)
-{
-#if !UE_BUILD_SHIPPING
-	if (CVarStallSwap.GetValueOnAnyThread() > 0.0f)
-	{
-		QUICK_SCOPE_CYCLE_COUNTER(STAT_Swap_Intentional_Stall);
-		FPlatformProcess::Sleep(CVarStallSwap.GetValueOnRenderThread() / 1000.0f);
-	}
-#endif
-
-	VERIFY_EGL_SCOPE();
-
-
-	static bool bUseChoreographer = false;
-
-	bool bChoreographerActive = false;
-
-	if (bUseChoreographer)
-	{
-		TheChoreographerFramePacer.WaitSync();
-		UE_CLOG(PImplData->SyncFenceForChoreographerMethod == EGL_NO_SYNC_KHR, LogRHI, Fatal, TEXT("SyncFenceForChoreographerMethod was EGL_NO_SYNC_KHR"));
-		bool bReady = false;
-		for (int32 Retry = 0; Retry < 150 && !bReady; Retry++)
-		{
-			EGLint Status = 0;
-			EGLBoolean Result = eglGetSyncAttribKHR_p(PImplData->eglDisplay, PImplData->SyncFenceForChoreographerMethod, EGL_SYNC_STATUS_KHR, &Status);
-			if (Result == EGL_FALSE)
-			{
-				UE_LOG(LogRHI, Error, TEXT("eglGetSyncAttribKHR returned false"));
-			}
-			else if (Status == EGL_SIGNALED_KHR)
-			{
-				bReady = true;
-			}
-			else if (Status != EGL_UNSIGNALED_KHR)
-			{
-				UE_LOG(LogRHI, Fatal, TEXT("eglGetSyncAttribKHR unexpected value %d"), Status);
-			}
-
-			if (!bReady)
-			{
-				//UE_LOG(LogRHI, Log, TEXT("Missed VSync (GPU)"));
-				TheChoreographerFramePacer.StartSync(-1, PImplData->DriverRefreshNanos - PImplData->DriverDeadlineNanos - PImplData->DriverAppVsyncOffsetNanos);
-				TheChoreographerFramePacer.WaitSync();
-			}
-		}
-		UE_CLOG(!bReady, LogRHI, Fatal, TEXT("Exhausted retries waiting for a GPU fence....GPU hang?"));
-		EGLBoolean Result = eglDestroySyncKHR_p(PImplData->eglDisplay, PImplData->SyncFenceForChoreographerMethod);
-		if (Result == EGL_FALSE)
-		{
-			UE_LOG(LogRHI, Fatal, TEXT("eglDestroySyncKHR_p returned false"));
-		}
-		PImplData->SyncFenceForChoreographerMethod = EGL_NO_SYNC_KHR;
-
-		bUseChoreographer = false;
-		bChoreographerActive = true;
-	}
-
-	bool bPrintMethod = false;
-	if (PImplData->DesiredSyncIntervalRelativeTo60Hz != SyncInterval)
-	{
-		GGetTimeStampsRetryCount = 0;
-
-		bPrintMethod = true;
-		PImplData->DesiredSyncIntervalRelativeTo60Hz = SyncInterval;
-		PImplData->DriverRefreshRate = 60.0f;
-		PImplData->DriverRefreshNanos = 16666666;
-		PImplData->DriverAppVsyncOffsetNanos = 2000000;
-		PImplData->DriverDeadlineNanos = 10666666;
-		PImplData->DriverSlopNanos = 1000000;
-
-
-		EGLnsecsANDROID EGL_COMPOSITE_DEADLINE_ANDROID_Value = -1;
-		EGLnsecsANDROID EGL_COMPOSITE_INTERVAL_ANDROID_Value = -1;
-		EGLnsecsANDROID EGL_COMPOSITE_TO_PRESENT_LATENCY_ANDROID_Value = -1;
-
-		if (eglGetCompositorTimingANDROID_p)
-		{
-			{
-				EGLint Item = EGL_COMPOSITE_DEADLINE_ANDROID;
-				if (!eglGetCompositorTimingANDROID_p(PImplData->eglDisplay, PImplData->eglSurface, 1, &Item, &EGL_COMPOSITE_DEADLINE_ANDROID_Value))
-				{
-					EGL_COMPOSITE_DEADLINE_ANDROID_Value = -1;
-				}
-			}
-			{
-				EGLint Item = EGL_COMPOSITE_INTERVAL_ANDROID;
-				if (!eglGetCompositorTimingANDROID_p(PImplData->eglDisplay, PImplData->eglSurface, 1, &Item, &EGL_COMPOSITE_INTERVAL_ANDROID_Value))
-				{
-					EGL_COMPOSITE_INTERVAL_ANDROID_Value = -1;
-				}
-			}
-			{
-				EGLint Item = EGL_COMPOSITE_TO_PRESENT_LATENCY_ANDROID;
-				if (!eglGetCompositorTimingANDROID_p(PImplData->eglDisplay, PImplData->eglSurface, 1, &Item, &EGL_COMPOSITE_TO_PRESENT_LATENCY_ANDROID_Value))
-				{
-					EGL_COMPOSITE_TO_PRESENT_LATENCY_ANDROID_Value = -1;
-				}
-			}
-			UE_LOG(LogRHI, Log, TEXT("AndroidEGL:SwapBuffers eglGetCompositorTimingANDROID EGL_COMPOSITE_DEADLINE_ANDROID=%lld, EGL_COMPOSITE_INTERVAL_ANDROID=%lld, EGL_COMPOSITE_TO_PRESENT_LATENCY_ANDROID=%lld"),
-				EGL_COMPOSITE_DEADLINE_ANDROID_Value,
-				EGL_COMPOSITE_INTERVAL_ANDROID_Value,
-				EGL_COMPOSITE_TO_PRESENT_LATENCY_ANDROID_Value
-			);
-		}
-
-		int64 PresentationDeadlineNanos = AndroidThunkCpp_GetMetaDataLong(TEXT("ue4.display.PresentationDeadlineNanos"));
-		int64 AppVsyncOffsetNanos = AndroidThunkCpp_GetMetaDataLong(TEXT("ue4.display.AppVsyncOffsetNanos"));
-
-		float RefreshRate = AndroidThunkCpp_GetMetaDataFloat(TEXT("ue4.display.getRefreshRate"));
-
-		UE_LOG(LogRHI, Log, TEXT("JNI Display getPresentationDeadlineNanos=%lld getAppVsyncOffsetNanos=%lld getRefreshRate=%f"),
-			PresentationDeadlineNanos,
-			AppVsyncOffsetNanos,
-			RefreshRate
-		);
-
-		if (EGL_COMPOSITE_INTERVAL_ANDROID_Value >= 4000000 && EGL_COMPOSITE_INTERVAL_ANDROID_Value <= 41666666)
-		{
-			PImplData->DriverRefreshRate = float(1000000000.0 / double(EGL_COMPOSITE_INTERVAL_ANDROID_Value));
-			PImplData->DriverRefreshNanos = EGL_COMPOSITE_INTERVAL_ANDROID_Value;
-		}
-		else if (RefreshRate >= 24.0f && RefreshRate <= 250.0f)
-		{
-			PImplData->DriverRefreshRate = RefreshRate;
-			PImplData->DriverRefreshNanos = int64(0.5 + 1000000000.0 / double(RefreshRate));
-		}
-
-		if (EGL_COMPOSITE_TO_PRESENT_LATENCY_ANDROID_Value > 0 && EGL_COMPOSITE_TO_PRESENT_LATENCY_ANDROID_Value <= PImplData->DriverRefreshNanos)
-		{
-			PImplData->DriverDeadlineNanos = EGL_COMPOSITE_TO_PRESENT_LATENCY_ANDROID_Value;
-		}
-		if (PresentationDeadlineNanos > 1000000 && PresentationDeadlineNanos - 1000000 <= PImplData->DriverRefreshNanos)
-		{
-			PImplData->DriverDeadlineNanos = PresentationDeadlineNanos - 1000000;
-		}
-
-		if (AppVsyncOffsetNanos >= 0 && AppVsyncOffsetNanos < PImplData->DriverRefreshNanos)
-		{
-			PImplData->DriverAppVsyncOffsetNanos = AppVsyncOffsetNanos;
-		}
-
-		UE_LOG(LogRHI, Log, TEXT("Final display timing metrics: DriverRefreshRate=%7.4f  DriverRefreshNanos=%lld  DriverDeadlineNanos=%lld  DriverAppVsyncOffsetNanos=%lld"),
-			PImplData->DriverRefreshRate,
-			PImplData->DriverRefreshNanos,
-			PImplData->DriverDeadlineNanos,
-			PImplData->DriverAppVsyncOffsetNanos
-		);
-
-
-		// make sure requested interval is in supported range
-		EGLint MinSwapInterval, MaxSwapInterval;
-		eglGetConfigAttrib(PImplData->eglDisplay, PImplData->eglConfigParam, EGL_MIN_SWAP_INTERVAL, &MinSwapInterval);
-		eglGetConfigAttrib(PImplData->eglDisplay, PImplData->eglConfigParam, EGL_MAX_SWAP_INTERVAL, &MaxSwapInterval);
-
-		int64 SyncIntervalNanos = (30 + 1000000000l * int64(SyncInterval)) / 60;
-
-		int32 UnderDriverInterval = int32(SyncIntervalNanos / PImplData->DriverRefreshNanos);
-		int32 OverDriverInterval = UnderDriverInterval + 1;
-
-		int64 UnderNanos = int64(UnderDriverInterval) * PImplData->DriverRefreshNanos;
-		int64 OverNanos = int64(OverDriverInterval) * PImplData->DriverRefreshNanos;
-
-		PImplData->DesiredSyncIntervalRelativeToDevice = (FMath::Abs(SyncIntervalNanos - UnderNanos) < FMath::Abs(SyncIntervalNanos - OverNanos)) ?
-			UnderDriverInterval : OverDriverInterval;
-
-		int32 DesiredDriverSyncInterval = FMath::Clamp<int32>(PImplData->DesiredSyncIntervalRelativeToDevice, MinSwapInterval, MaxSwapInterval);
-		
-		UE_LOG(LogRHI, Log, TEXT("AndroidEGL:SwapBuffers Min=%d, Max=%d, Request=%d, ClosestDriver=%d, SetDriver=%d"), MinSwapInterval, MaxSwapInterval, PImplData->DesiredSyncIntervalRelativeTo60Hz, PImplData->DesiredSyncIntervalRelativeToDevice, DesiredDriverSyncInterval);
-
-		if (DesiredDriverSyncInterval != PImplData->DriverSyncIntervalRelativeToDevice)
-		{
-			PImplData->DriverSyncIntervalRelativeToDevice = DesiredDriverSyncInterval;
-			UE_LOG(LogRHI, Log, TEXT("Called eglSwapInterval %d"), DesiredDriverSyncInterval);
-			eglSwapInterval(PImplData->eglDisplay, PImplData->DriverSyncIntervalRelativeToDevice);
-		}
-	}
-
-	if (PImplData->DesiredSyncIntervalRelativeToDevice > PImplData->DriverSyncIntervalRelativeToDevice)
-	{
-		// this is a prototype currently unused, left here for possible future use
-		if (ShouldUseChoreographer())
-		{
-			UE_CLOG(bPrintMethod, LogRHI, Display, TEXT("Using google choreographer method for frame pacing"));
-
-			TheChoreographerFramePacer.StartSync((PImplData->DesiredSyncIntervalRelativeToDevice - 1) * PImplData->DriverRefreshNanos, PImplData->DriverRefreshNanos - PImplData->DriverDeadlineNanos - PImplData->DriverAppVsyncOffsetNanos);
-			UE_CLOG(PImplData->SyncFenceForChoreographerMethod != EGL_NO_SYNC_KHR, LogRHI, Fatal, TEXT("SyncFenceForChoreographerMethod was NOT EGL_NO_SYNC_KHR"));
-			PImplData->SyncFenceForChoreographerMethod = eglCreateSyncKHR_p(PImplData->eglDisplay, EGL_SYNC_FENCE_KHR, NULL);
-			if (eglPresentationTimeANDROID_p)
-			{
-				EGLnsecsANDROID PresentationTime = ChoreographerClock() + PImplData->DesiredSyncIntervalRelativeToDevice * PImplData->DriverRefreshNanos;
-				eglPresentationTimeANDROID_p(PImplData->eglDisplay, PImplData->eglSurface, PresentationTime);
-			}
-
-			bUseChoreographer = true;
-		}
-		else
-		{
-			UE_CLOG(bPrintMethod, LogRHI, Display, TEXT("Using niave method for frame pacing (possible with timestamps method)"));
-			if (PImplData->LastTimeEmulatedSync > 0.0)
-			{
-				QUICK_SCOPE_CYCLE_COUNTER(STAT_StallForEmulatedSyncInterval);
-				float MinTimeBetweenFrames = (float(PImplData->DesiredSyncIntervalRelativeToDevice) / PImplData->DriverRefreshRate);
-
-				float ThisTime = FPlatformTime::Seconds() - PImplData->LastTimeEmulatedSync;
-				if (ThisTime > 0 && ThisTime < MinTimeBetweenFrames)
-				{
-					FPlatformProcess::Sleep(MinTimeBetweenFrames - ThisTime);
-				}
-			}
-		}
-	}
-	if (bChoreographerActive && !bUseChoreographer)
-	{
-		TheChoreographerFramePacer.StopChoreographer();
-	}
-	if (!bUseChoreographer && CanUseGetFrameTimestamps())
-	{
-		UE_CLOG(bPrintMethod, LogRHI, Display, TEXT("Using eglGetFrameTimestampsANDROID method for frame pacing"));
-
-		//static bool bPrintOnce = true;
-		if (FrameIDs[(int32(NextFrameIDSlot) - 1) % NUM_FRAMES_TO_MONITOR])
-			// not supported   && eglGetFrameTimestampsSupportedANDROID_p && eglGetFrameTimestampsSupportedANDROID_p(PImplData->eglDisplay, PImplData->eglSurface, EGL_FIRST_COMPOSITION_START_TIME_ANDROID))
-		{
-			//UE_CLOG(bPrintOnce, LogRHI, Log, TEXT("eglGetFrameTimestampsSupportedANDROID retured true for EGL_FIRST_COMPOSITION_START_TIME_ANDROID"));
-			EGLint TimestampList = EGL_FIRST_COMPOSITION_START_TIME_ANDROID;
-			//EGLint TimestampList = EGL_COMPOSITION_LATCH_TIME_ANDROID;
-			//EGLint TimestampList = EGL_LAST_COMPOSITION_START_TIME_ANDROID;
-			//EGLint TimestampList = EGL_DISPLAY_PRESENT_TIME_ANDROID;
-			EGLnsecsANDROID Result = 0;
-			int32 DeltaFrameIndex = 1;
-			for (int32 Index = int32(NextFrameIDSlot) - 1; Index >= int32(NextFrameIDSlot) - NUM_FRAMES_TO_MONITOR && Index >= 0; Index--)
-			{
-				Result = 0;
-				if (FrameIDs[Index % NUM_FRAMES_TO_MONITOR])
-				{
-					eglGetFrameTimestampsANDROID_p(PImplData->eglDisplay, PImplData->eglSurface, FrameIDs[Index % NUM_FRAMES_TO_MONITOR], 1, &TimestampList, &Result);
-				}
-				if (Result > 0)
-				{
-					break;
-				}
-				DeltaFrameIndex++;
-			}
-
-			GGetTimeStampsSucceededThisFrame = Result > 0;
-			if (GGetTimeStampsSucceededThisFrame)
-			{
-				EGLnsecsANDROID FudgeFactor = 0; //  8333 * 1000;
-				EGLnsecsANDROID DeltaNanos = EGLnsecsANDROID(PImplData->DesiredSyncIntervalRelativeToDevice) * EGLnsecsANDROID(DeltaFrameIndex) * PImplData->DriverRefreshNanos;
-				EGLnsecsANDROID PresentationTime = Result + DeltaNanos + FudgeFactor;
-				eglPresentationTimeANDROID_p(PImplData->eglDisplay, PImplData->eglSurface, PresentationTime);
-				GGetTimeStampsRetryCount = 0;
-			}
-			else
-			{
-				GGetTimeStampsRetryCount++;
-				if (GGetTimeStampsRetryCount == CVarTimeStampErrorRetryCount.GetValueOnAnyThread())
-				{
-					UE_LOG(LogRHI, Log, TEXT("eglGetFrameTimestampsANDROID_p failed for %d consecutive frames, reverting to naive frame pacer."), GGetTimeStampsRetryCount);
-				}
-			}
-		}
-		else
-		{
-			//UE_CLOG(bPrintOnce, LogRHI, Log, TEXT("eglGetFrameTimestampsSupportedANDROID doesn't exist or retured false for EGL_FIRST_COMPOSITION_START_TIME_ANDROID, discarding eglGetNextFrameIdANDROID_p and eglGetFrameTimestampsANDROID_p"));
-		}
-		//bPrintOnce = false;
-	}
-
-	PImplData->LastTimeEmulatedSync = FPlatformTime::Seconds();
-
-	{
-		QUICK_SCOPE_CYCLE_COUNTER(STAT_eglSwapBuffers);
-
-		FrameIDs[(NextFrameIDSlot) % NUM_FRAMES_TO_MONITOR] = 0;
-		if (eglGetNextFrameIdANDROID_p && (CanUseGetFrameTimestamps() || CVarSpewGetFrameTimestamps.GetValueOnAnyThread()))
-		{
-			eglGetNextFrameIdANDROID_p(PImplData->eglDisplay, PImplData->eglSurface, &FrameIDs[(NextFrameIDSlot) % NUM_FRAMES_TO_MONITOR]);
-		}
-		NextFrameIDSlot++;
-
-		if (PImplData->eglSurface == NULL || !eglSwapBuffers(PImplData->eglDisplay, PImplData->eglSurface))
-		{
-			// shutdown if swapbuffering goes down
-			if (PImplData->SwapBufferFailureCount > 10)
-			{
-				//Process.killProcess(Process.myPid());		//@todo android
-			}
-			PImplData->SwapBufferFailureCount++;
-
-			// basic reporting
-			if (PImplData->eglSurface == NULL)
-			{
-				return false;
-			}
-			else
-			{
-				if (eglGetError() == EGL_CONTEXT_LOST)
-				{
-					//Logger.LogOut("swapBuffers: EGL11.EGL_CONTEXT_LOST err: " + eglGetError());					
-					//Process.killProcess(Process.myPid());		//@todo android
-				}
-			}
-
-			return false;
-		}
-	}
-
-	if (PImplData->DesiredSyncIntervalRelativeToDevice > 0 && eglGetFrameTimestampsANDROID_p && CVarSpewGetFrameTimestamps.GetValueOnAnyThread())
-	{
-		static EGLint TimestampList[9] = 
-		{
-			EGL_REQUESTED_PRESENT_TIME_ANDROID,
-			EGL_RENDERING_COMPLETE_TIME_ANDROID,
-			EGL_COMPOSITION_LATCH_TIME_ANDROID,
-			EGL_FIRST_COMPOSITION_START_TIME_ANDROID,
-			EGL_LAST_COMPOSITION_START_TIME_ANDROID,
-			EGL_FIRST_COMPOSITION_GPU_FINISHED_TIME_ANDROID,
-			EGL_DISPLAY_PRESENT_TIME_ANDROID,
-			EGL_DEQUEUE_READY_TIME_ANDROID,
-			EGL_READS_DONE_TIME_ANDROID
-		};
-
-		static const TCHAR* TimestampStrings[9] =
-		{
-			TEXT("EGL_REQUESTED_PRESENT_TIME_ANDROID"),
-			TEXT("EGL_RENDERING_COMPLETE_TIME_ANDROID"),
-			TEXT("EGL_COMPOSITION_LATCH_TIME_ANDROID"),
-			TEXT("EGL_FIRST_COMPOSITION_START_TIME_ANDROID"),
-			TEXT("EGL_LAST_COMPOSITION_START_TIME_ANDROID"),
-			TEXT("EGL_FIRST_COMPOSITION_GPU_FINISHED_TIME_ANDROID"),
-			TEXT("EGL_DISPLAY_PRESENT_TIME_ANDROID"),
-			TEXT("EGL_DEQUEUE_READY_TIME_ANDROID"),
-			TEXT("EGL_READS_DONE_TIME_ANDROID")
-		};
-
-
-		EGLnsecsANDROID Results[NUM_FRAMES_TO_MONITOR][9] = { {0} };
-		EGLnsecsANDROID FirstRealValue = 0;
-		for (int32 Index = int32(NextFrameIDSlot) - NUM_FRAMES_TO_MONITOR; Index < int32(NextFrameIDSlot); Index++)
-		{
-			eglGetFrameTimestampsANDROID_p(PImplData->eglDisplay, PImplData->eglSurface, FrameIDs[Index % NUM_FRAMES_TO_MONITOR], 9, TimestampList, Results[Index % NUM_FRAMES_TO_MONITOR]);
-			for (int32 IndexInner = 0; IndexInner < 9; IndexInner++)
-			{
-				if (!FirstRealValue || (Results[Index % NUM_FRAMES_TO_MONITOR][IndexInner] > 1 && Results[Index % NUM_FRAMES_TO_MONITOR][IndexInner] < FirstRealValue))
-				{
-					FirstRealValue = Results[Index % NUM_FRAMES_TO_MONITOR][IndexInner];
-				}
-			}
-		}
-		UE_CLOG(CVarSpewGetFrameTimestamps.GetValueOnAnyThread() > 1, LogRHI, Log, TEXT("************************************  frame %d   base time is %lld"), NextFrameIDSlot - 1, FirstRealValue);
-
-		for (int32 Index = int32(NextFrameIDSlot) - NUM_FRAMES_TO_MONITOR; Index < int32(NextFrameIDSlot); Index++)
-		{
-
-			UE_CLOG(CVarSpewGetFrameTimestamps.GetValueOnAnyThread() > 1, LogRHI, Log, TEXT("eglGetFrameTimestampsANDROID_p  frame %d"), Index);
-			for (int32 IndexInner = 0; IndexInner < 9; IndexInner++)
-			{
-				int32 MsVal = (Results[Index % NUM_FRAMES_TO_MONITOR][IndexInner] > 1) ? int32((Results[Index % NUM_FRAMES_TO_MONITOR][IndexInner] - FirstRealValue) / 1000000) : int32(Results[Index % NUM_FRAMES_TO_MONITOR][IndexInner]);
-
-				UE_CLOG(CVarSpewGetFrameTimestamps.GetValueOnAnyThread() > 1, LogRHI, Log, TEXT("     %8d    %s"), MsVal, TimestampStrings[IndexInner]);
-			}
-		}
-
-		int32 IndexLast = int32(NextFrameIDSlot) - NUM_FRAMES_TO_MONITOR;
-		int32 IndexLastNext = IndexLast + 1;
-
-		if (Results[IndexLast % NUM_FRAMES_TO_MONITOR][3] > 1 && Results[IndexLastNext % NUM_FRAMES_TO_MONITOR][3] > 1)
-		{
-			int32 MsVal = int32((Results[IndexLastNext % NUM_FRAMES_TO_MONITOR][3] - Results[IndexLast % NUM_FRAMES_TO_MONITOR][3]) / 1000000);
-
-			RecordedFrameInterval[NumRecordedFrameInterval++] = MsVal;
-			if (NumRecordedFrameInterval == 100)
-			{
-				FString All;
-				int32 NumOnTarget = 0;
-				int32 NumBelowTarget = 0;
-				int32 NumAboveTarget = 0;
-				for (int32 Index = 0; Index < 100; Index++)
-				{
-					if (Index)
-					{
-						All += TCHAR(' ');
-					}
-					All += FString::Printf(TEXT("%d"), RecordedFrameInterval[Index]);
-
-					if (RecordedFrameInterval[Index] > PImplData->DesiredSyncIntervalRelativeTo60Hz * 16 - 8 && RecordedFrameInterval[Index] < PImplData->DesiredSyncIntervalRelativeTo60Hz * 16 + 8)
-					{
-						NumOnTarget++;
-					}
-					else if (RecordedFrameInterval[Index] < PImplData->DesiredSyncIntervalRelativeTo60Hz * 16)
-					{
-						NumBelowTarget++;
-					}
-					else
-					{
-						NumAboveTarget++;
-					}
-				}
-				UE_LOG(LogRHI, Log, TEXT("%3d fast  %3d ok  %3d slow   %s"), NumBelowTarget, NumOnTarget, NumAboveTarget, *All);
-				NumRecordedFrameInterval = 0;
-			}
-		}
-	}
-
-
-	return true;
-}
-
 bool AndroidEGL::IsInitialized()
 {
 	return PImplData->Initalized;
@@ -1502,6 +763,18 @@ EGLDisplay AndroidEGL::GetDisplay() const
 	return PImplData->eglDisplay;
 }
 
+EGLSurface AndroidEGL::GetSurface() const
+{
+	return PImplData->eglSurface;
+}
+
+void AndroidEGL::GetSwapIntervalRange(EGLint& OutMinSwapInterval, EGLint& OutMaxSwapInterval) const
+{
+	eglGetConfigAttrib(PImplData->eglDisplay, PImplData->eglConfigParam, EGL_MIN_SWAP_INTERVAL, &OutMinSwapInterval);
+	eglGetConfigAttrib(PImplData->eglDisplay, PImplData->eglConfigParam, EGL_MAX_SWAP_INTERVAL, &OutMaxSwapInterval);
+}
+
+
 ANativeWindow* AndroidEGL::GetNativeWindow() const
 {
 	return PImplData->Window;
diff --git a/Engine/Source/Runtime/OpenGLDrv/Private/Android/AndroidEGL.h b/Engine/Source/Runtime/OpenGLDrv/Private/Android/AndroidEGL.h
--- a/Engine/Source/Runtime/OpenGLDrv/Private/Android/AndroidEGL.h
+++ b/Engine/Source/Runtime/OpenGLDrv/Private/Android/AndroidEGL.h
@@ -68,14 +68,16 @@ public:
 	void Init( APIVariant API, uint32 MajorVersion, uint32 MinorVersion, bool bDebug);
 	void ReInit();
 	void UnBind();
-	bool SwapBuffers(int32 SyncInterval);
+	bool SwapBuffers();
 	void Terminate();
 	void InitSurface(bool bUseSmallSurface, bool bCreateWndSurface);
 
 	void GetDimensions(uint32& OutWidth, uint32& OutHeight);
 	
 	EGLDisplay GetDisplay() const;
+	EGLSurface GetSurface() const;
 	ANativeWindow* GetNativeWindow() const;
+	void GetSwapIntervalRange(EGLint& OutMinSwapInterval, EGLint& OutMaxSwapInterval) const;
 
 	EGLContext CreateContext(EGLContext InSharedContext = EGL_NO_CONTEXT);
 	int32 GetError();
@@ -132,4 +134,91 @@ private:
 	int *ContextAttributes;
 };
 
+#define ENABLE_CONFIG_FILTER 1
+#define ENABLE_EGL_DEBUG 0
+#define ENABLE_VERIFY_EGL 0
+#define ENABLE_VERIFY_EGL_TRACE 0
+
+#if ENABLE_VERIFY_EGL
+
+#define VERIFY_EGL(msg) { VerifyEGLResult(eglGetError(),TEXT(#msg),TEXT(""),TEXT(__FILE__),__LINE__); }
+
+void VerifyEGLResult(EGLint ErrorCode, const TCHAR* Msg1, const TCHAR* Msg2, const TCHAR* Filename, uint32 Line)
+{
+	if (ErrorCode != EGL_SUCCESS)
+	{
+		static const TCHAR* EGLErrorStrings[] =
+		{
+			TEXT("EGL_NOT_INITIALIZED"),
+			TEXT("EGL_BAD_ACCESS"),
+			TEXT("EGL_BAD_ALLOC"),
+			TEXT("EGL_BAD_ATTRIBUTE"),
+			TEXT("EGL_BAD_CONFIG"),
+			TEXT("EGL_BAD_CONTEXT"),
+			TEXT("EGL_BAD_CURRENT_SURFACE"),
+			TEXT("EGL_BAD_DISPLAY"),
+			TEXT("EGL_BAD_MATCH"),
+			TEXT("EGL_BAD_NATIVE_PIXMAP"),
+			TEXT("EGL_BAD_NATIVE_WINDOW"),
+			TEXT("EGL_BAD_PARAMETER"),
+			TEXT("EGL_BAD_SURFACE"),
+			TEXT("EGL_CONTEXT_LOST"),
+			TEXT("UNKNOWN EGL ERROR")
+		};
+
+		uint32 ErrorIndex = FMath::Min<uint32>(ErrorCode - EGL_SUCCESS, UE_ARRAY_COUNT(EGLErrorStrings) - 1);
+		UE_LOG(LogRHI, Warning, TEXT("%s(%u): %s%s failed with error %s (0x%x)"),
+			Filename, Line, Msg1, Msg2, EGLErrorStrings[ErrorIndex], ErrorCode);
+		check(0);
+	}
+}
+
+class FEGLErrorScope
+{
+public:
+	FEGLErrorScope(
+		const TCHAR* InFunctionName,
+		const TCHAR* InFilename,
+		const uint32 InLine)
+		: FunctionName(InFunctionName)
+		, Filename(InFilename)
+		, Line(InLine)
+	{
+#if ENABLE_VERIFY_EGL_TRACE
+		UE_LOG(LogRHI, Log, TEXT("EGL log before %s(%d): %s"), InFilename, InLine, InFunctionName);
+#endif
+		CheckForErrors(TEXT("Before "));
+	}
+
+	~FEGLErrorScope()
+	{
+#if ENABLE_VERIFY_EGL_TRACE
+		UE_LOG(LogRHI, Log, TEXT("EGL log after  %s(%d): %s"), Filename, Line, FunctionName);
+#endif
+		CheckForErrors(TEXT("After "));
+	}
+
+private:
+	const TCHAR* FunctionName;
+	const TCHAR* Filename;
+	const uint32 Line;
+
+	void CheckForErrors(const TCHAR* PrefixString)
+	{
+		VerifyEGLResult(eglGetError(), PrefixString, FunctionName, Filename, Line);
+	}
+};
+
+#define MACRO_TOKENIZER(IdentifierName, Msg, FileName, LineNumber) FEGLErrorScope IdentifierName_ ## LineNumber (Msg, FileName, LineNumber)
+#define MACRO_TOKENIZER2(IdentifierName, Msg, FileName, LineNumber) MACRO_TOKENIZER(IdentiferName, Msg, FileName, LineNumber)
+#define VERIFY_EGL_SCOPE_WITH_MSG_STR(MsgStr) MACRO_TOKENIZER2(ErrorScope_, MsgStr, TEXT(__FILE__), __LINE__)
+#define VERIFY_EGL_SCOPE() VERIFY_EGL_SCOPE_WITH_MSG_STR(ANSI_TO_TCHAR(__FUNCTION__))
+#define VERIFY_EGL_FUNC(Func, ...) { VERIFY_EGL_SCOPE_WITH_MSG_STR(TEXT(#Func)); Func(__VA_ARGS__); }
+#else
+#define VERIFY_EGL(...)
+#define VERIFY_EGL_SCOPE(...)
+#endif
+
+
+
 #endif
diff --git a/Engine/Source/Runtime/OpenGLDrv/Private/Android/AndroidOpenGL.cpp b/Engine/Source/Runtime/OpenGLDrv/Private/Android/AndroidOpenGL.cpp
--- a/Engine/Source/Runtime/OpenGLDrv/Private/Android/AndroidOpenGL.cpp
+++ b/Engine/Source/Runtime/OpenGLDrv/Private/Android/AndroidOpenGL.cpp
@@ -12,6 +12,8 @@
 #include "Android/AndroidWindow.h"
 #include "AndroidOpenGLPrivate.h"
 #include "Android/AndroidPlatformMisc.h"
+#include "Android/AndroidPlatformFramePacer.h"
+
 
 PFNeglPresentationTimeANDROID eglPresentationTimeANDROID_p = NULL;
 PFNeglGetNextFrameIdANDROID eglGetNextFrameIdANDROID_p = NULL;
@@ -56,6 +58,8 @@ struct FPlatformOpenGLDevice
 
 FPlatformOpenGLDevice::~FPlatformOpenGLDevice()
 {
+	FPlatformRHIFramePacer::Destroy();
+
 	FAndroidAppEntry::ReleaseEGL();
 }
 
@@ -73,6 +77,9 @@ static bool bRunningUnderRenderDoc = false;
 
 void FPlatformOpenGLDevice::Init()
 {
+	// Initialize frame pacer
+	FPlatformRHIFramePacer::Init(new FAndroidOpenGLFramePacer());
+
 	extern void InitDebugContext();
 
 	bRunningUnderRenderDoc = glIsEnabled(GL_DEBUG_TOOL_EXT) != GL_FALSE;
@@ -117,16 +124,17 @@ void* PlatformGetWindow(FPlatformOpenGLContext* Context, void** AddParam)
 	return (void*)&Context->eglContext;
 }
 
-bool PlatformBlitToViewport( FPlatformOpenGLDevice* Device, const FOpenGLViewport& Viewport, uint32 BackbufferSizeX, uint32 BackbufferSizeY, bool bPresent,bool bLockToVsync, int32 SyncInterval )
+bool PlatformBlitToViewport(FPlatformOpenGLDevice* Device, const FOpenGLViewport& Viewport, uint32 BackbufferSizeX, uint32 BackbufferSizeY, bool bPresent,bool bLockToVsync )
 {
 	if (bPresent && Viewport.GetCustomPresent())
 	{
 		QUICK_SCOPE_CYCLE_COUNTER(STAT_FAndroidOpenGL_PlatformBlitToViewport_CustomPresent);
+		int32 SyncInterval = FAndroidPlatformRHIFramePacer::GetLegacySyncInterval();
 		bPresent = Viewport.GetCustomPresent()->Present(SyncInterval);
 	}
 	if (bPresent)
 	{
-		AndroidEGL::GetInstance()->SwapBuffers(bLockToVsync ? SyncInterval : 0);
+		FAndroidPlatformRHIFramePacer::SwapBuffers(bLockToVsync);
 	}
 	static IConsoleVariable* CVar = IConsoleManager::Get().FindConsoleVariable(TEXT("a.UseFrameTimeStampsForPacing"));
 	const bool bForceGPUFence = CVar ? CVar->GetInt() != 0 : false;
diff --git a/Engine/Source/Runtime/OpenGLDrv/Private/Android/AndroidOpenGLFramePacer.cpp b/Engine/Source/Runtime/OpenGLDrv/Private/Android/AndroidOpenGLFramePacer.cpp
new file mode 100644
--- /dev/null
+++ b/Engine/Source/Runtime/OpenGLDrv/Private/Android/AndroidOpenGLFramePacer.cpp
@@ -0,0 +1,461 @@
+// Copyright Epic Games, Inc. All Rights Reserved.
+
+#include "AndroidOpenGLFramePacer.h"
+
+#if USE_ANDROID_OPENGL
+
+/*******************************************************************
+ * FAndroidOpenGLFramePacer implementation
+ *******************************************************************/
+
+#include "OpenGLDrvPrivate.h"
+
+#if USE_ANDROID_OPENGL_SWAPPY
+#include "Android/AndroidJNI.h"
+#include "Android/AndroidApplication.h"
+#include "swappyGL.h"
+#include "swappyGL_extra.h"
+#endif
+
+#include "AndroidEGL.h"
+#include <EGL/egl.h>
+
+void FAndroidOpenGLFramePacer::Init()
+{
+	bSwappyInit = false;
+#if USE_ANDROID_OPENGL_SWAPPY
+	if (FAndroidPlatformRHIFramePacer::CVarUseSwappyForFramePacing.GetValueOnAnyThread() == 1)
+	{
+		// initialize now if set on startup
+		InitSwappy();
+	}
+	else
+	{
+		// initialize later if set by console
+		FAndroidPlatformRHIFramePacer::CVarUseSwappyForFramePacing.AsVariable()->SetOnChangedCallback(FConsoleVariableDelegate::CreateLambda([this](IConsoleVariable* Variable)
+		{
+			if (Variable->GetInt() == 1)
+			{
+				InitSwappy();
+			}
+		}));
+	}
+#else
+	bSwappyInit = false;
+#endif
+}
+
+#if USE_ANDROID_OPENGL_SWAPPY
+void FAndroidOpenGLFramePacer::InitSwappy()
+{
+	if (!bSwappyInit)
+	{
+		// initialize Swappy
+		JNIEnv* Env = FAndroidApplication::GetJavaEnv();
+		if (ensure(Env))
+		{
+			FPlatformMisc::LowLevelOutputDebugStringf(TEXT("Init Swappy"));
+			SwappyGL_init(Env, FJavaWrapper::GameActivityThis);
+		}
+		bSwappyInit = true;
+	}
+}
+#endif
+
+FAndroidOpenGLFramePacer::~FAndroidOpenGLFramePacer()
+{
+#if USE_ANDROID_OPENGL_SWAPPY
+	FAndroidPlatformRHIFramePacer::CVarUseSwappyForFramePacing.AsVariable()->SetOnChangedCallback(FConsoleVariableDelegate());
+	if (bSwappyInit)
+	{
+		FPlatformMisc::LowLevelOutputDebugStringf(TEXT("Shutdown Swappy"));
+		SwappyGL_destroy();
+		bSwappyInit = false;
+	}
+#endif
+}
+
+static bool GGetTimeStampsSucceededThisFrame = true;
+static uint32 GGetTimeStampsRetryCount = 0;
+
+static bool CanUseGetFrameTimestamps()
+{
+	return FAndroidPlatformRHIFramePacer::CVarUseGetFrameTimestamps.GetValueOnAnyThread()
+		&& eglGetFrameTimestampsANDROID_p
+		&& eglGetNextFrameIdANDROID_p
+		&& eglPresentationTimeANDROID_p
+		&& (GGetTimeStampsRetryCount < FAndroidPlatformRHIFramePacer::CVarTimeStampErrorRetryCount.GetValueOnAnyThread());
+}
+static bool CanUseGetFrameTimestampsForThisFrame()
+{
+	return CanUseGetFrameTimestamps() && GGetTimeStampsSucceededThisFrame;
+}
+
+bool ShouldUseGPUFencesToLimitLatency()
+{
+	if (CanUseGetFrameTimestampsForThisFrame())
+	{
+		return true; // this method requires a GPU fence to give steady results
+	}
+	return FAndroidPlatformRHIFramePacer::CVarDisableOpenGLGPUSync.GetValueOnAnyThread() == 0; // otherwise just based on the FAndroidPlatformRHIFramePacer::CVar; thought to be bad to use GPU fences on PowerVR
+}
+
+static uint32 NextFrameIDSlot = 0;
+#define NUM_FRAMES_TO_MONITOR (4)
+static EGLuint64KHR FrameIDs[NUM_FRAMES_TO_MONITOR] = { 0 };
+
+static int32 RecordedFrameInterval[100];
+static int32 NumRecordedFrameInterval = 0;
+
+extern float AndroidThunkCpp_GetMetaDataFloat(const FString& Key);
+
+
+bool FAndroidOpenGLFramePacer::SupportsFramePace(int32 QueryFramePace)
+{
+#if USE_ANDROID_OPENGL_SWAPPY
+	if (FAndroidPlatformRHIFramePacer::CVarUseSwappyForFramePacing.GetValueOnAnyThread() == 1)
+	{
+		// temporary
+		if (QueryFramePace == 90 || QueryFramePace == 45)
+		{
+			return true;
+		}
+	}
+#endif
+	return FGenericPlatformRHIFramePacer::SupportsFramePace(QueryFramePace);
+}
+
+bool FAndroidOpenGLFramePacer::SwapBuffers(bool bLockToVsync)
+{
+#if !UE_BUILD_SHIPPING
+	if (FAndroidPlatformRHIFramePacer::CVarStallSwap.GetValueOnAnyThread() > 0.0f)
+	{
+		QUICK_SCOPE_CYCLE_COUNTER(STAT_Swap_Intentional_Stall);
+		FPlatformProcess::Sleep(FAndroidPlatformRHIFramePacer::CVarStallSwap.GetValueOnRenderThread() / 1000.0f);
+	}
+#endif
+
+	VERIFY_EGL_SCOPE();
+
+	EGLDisplay eglDisplay = AndroidEGL::GetInstance()->GetDisplay();
+	EGLSurface eglSurface = AndroidEGL::GetInstance()->GetSurface();
+	int32 SyncInterval = FAndroidPlatformRHIFramePacer::GetLegacySyncInterval();
+
+	bool bPrintMethod = false;
+
+#if USE_ANDROID_OPENGL_SWAPPY
+	if (FAndroidPlatformRHIFramePacer::CVarUseSwappyForFramePacing.GetValueOnRenderThread() != 0 && ensure(bSwappyInit))
+	{
+		int64 DesiredFrameNS = (1000000000L) / (int64)FAndroidPlatformRHIFramePacer::GetFramePace();
+		SwappyGL_setSwapIntervalNS(DesiredFrameNS);
+		SwappyGL_setAutoSwapInterval(false);
+		SwappyGL_swap(eglDisplay, eglSurface);
+	}
+	else
+#endif
+	{
+		if (DesiredSyncIntervalRelativeTo60Hz != SyncInterval)
+		{
+			GGetTimeStampsRetryCount = 0;
+
+			bPrintMethod = true;
+			DesiredSyncIntervalRelativeTo60Hz = SyncInterval;
+			DriverRefreshRate = 60.0f;
+			DriverRefreshNanos = 16666666;
+
+
+			EGLnsecsANDROID EGL_COMPOSITE_DEADLINE_ANDROID_Value = -1;
+			EGLnsecsANDROID EGL_COMPOSITE_INTERVAL_ANDROID_Value = -1;
+			EGLnsecsANDROID EGL_COMPOSITE_TO_PRESENT_LATENCY_ANDROID_Value = -1;
+
+			if (eglGetCompositorTimingANDROID_p)
+			{
+				{
+					EGLint Item = EGL_COMPOSITE_DEADLINE_ANDROID;
+					if (!eglGetCompositorTimingANDROID_p(eglDisplay, eglSurface, 1, &Item, &EGL_COMPOSITE_DEADLINE_ANDROID_Value))
+					{
+						EGL_COMPOSITE_DEADLINE_ANDROID_Value = -1;
+					}
+				}
+				{
+					EGLint Item = EGL_COMPOSITE_INTERVAL_ANDROID;
+					if (!eglGetCompositorTimingANDROID_p(eglDisplay, eglSurface, 1, &Item, &EGL_COMPOSITE_INTERVAL_ANDROID_Value))
+					{
+						EGL_COMPOSITE_INTERVAL_ANDROID_Value = -1;
+					}
+				}
+				{
+					EGLint Item = EGL_COMPOSITE_TO_PRESENT_LATENCY_ANDROID;
+					if (!eglGetCompositorTimingANDROID_p(eglDisplay, eglSurface, 1, &Item, &EGL_COMPOSITE_TO_PRESENT_LATENCY_ANDROID_Value))
+					{
+						EGL_COMPOSITE_TO_PRESENT_LATENCY_ANDROID_Value = -1;
+					}
+				}
+				UE_LOG(LogRHI, Log, TEXT("AndroidEGL:SwapBuffers eglGetCompositorTimingANDROID EGL_COMPOSITE_DEADLINE_ANDROID=%lld, EGL_COMPOSITE_INTERVAL_ANDROID=%lld, EGL_COMPOSITE_TO_PRESENT_LATENCY_ANDROID=%lld"),
+					EGL_COMPOSITE_DEADLINE_ANDROID_Value,
+					EGL_COMPOSITE_INTERVAL_ANDROID_Value,
+					EGL_COMPOSITE_TO_PRESENT_LATENCY_ANDROID_Value
+				);
+			}
+
+			float RefreshRate = AndroidThunkCpp_GetMetaDataFloat(TEXT("ue4.display.getRefreshRate"));
+
+			UE_LOG(LogRHI, Log, TEXT("JNI Display getRefreshRate=%f"),
+				RefreshRate
+			);
+
+			if (EGL_COMPOSITE_INTERVAL_ANDROID_Value >= 4000000 && EGL_COMPOSITE_INTERVAL_ANDROID_Value <= 41666666)
+			{
+				DriverRefreshRate = float(1000000000.0 / double(EGL_COMPOSITE_INTERVAL_ANDROID_Value));
+				DriverRefreshNanos = EGL_COMPOSITE_INTERVAL_ANDROID_Value;
+			}
+			else if (RefreshRate >= 24.0f && RefreshRate <= 250.0f)
+			{
+				DriverRefreshRate = RefreshRate;
+				DriverRefreshNanos = int64(0.5 + 1000000000.0 / double(RefreshRate));
+			}
+
+			UE_LOG(LogRHI, Log, TEXT("Final display timing metrics: DriverRefreshRate=%7.4f  DriverRefreshNanos=%lld"),
+				DriverRefreshRate,
+				DriverRefreshNanos
+			);
+
+			// make sure requested interval is in supported range
+			EGLint MinSwapInterval, MaxSwapInterval;
+			AndroidEGL::GetInstance()->GetSwapIntervalRange(MinSwapInterval, MaxSwapInterval);
+
+			int64 SyncIntervalNanos = (30 + 1000000000l * int64(SyncInterval)) / 60;
+
+			int32 UnderDriverInterval = int32(SyncIntervalNanos / DriverRefreshNanos);
+			int32 OverDriverInterval = UnderDriverInterval + 1;
+
+			int64 UnderNanos = int64(UnderDriverInterval) * DriverRefreshNanos;
+			int64 OverNanos = int64(OverDriverInterval) * DriverRefreshNanos;
+
+			DesiredSyncIntervalRelativeToDevice = (FMath::Abs(SyncIntervalNanos - UnderNanos) < FMath::Abs(SyncIntervalNanos - OverNanos)) ?
+				UnderDriverInterval : OverDriverInterval;
+
+			int32 DesiredDriverSyncInterval = FMath::Clamp<int32>(DesiredSyncIntervalRelativeToDevice, MinSwapInterval, MaxSwapInterval);
+
+			UE_LOG(LogRHI, Log, TEXT("AndroidEGL:SwapBuffers Min=%d, Max=%d, Request=%d, ClosestDriver=%d, SetDriver=%d"), MinSwapInterval, MaxSwapInterval, DesiredSyncIntervalRelativeTo60Hz, DesiredSyncIntervalRelativeToDevice, DesiredDriverSyncInterval);
+
+			if (DesiredDriverSyncInterval != DriverSyncIntervalRelativeToDevice)
+			{
+				DriverSyncIntervalRelativeToDevice = DesiredDriverSyncInterval;
+				UE_LOG(LogRHI, Log, TEXT("Called eglSwapInterval %d"), DesiredDriverSyncInterval);
+				eglSwapInterval(eglDisplay, DriverSyncIntervalRelativeToDevice);
+			}
+		}
+
+		if (DesiredSyncIntervalRelativeToDevice > DriverSyncIntervalRelativeToDevice)
+		{
+			{
+				UE_CLOG(bPrintMethod, LogRHI, Display, TEXT("Using niave method for frame pacing (possible with timestamps method)"));
+				if (LastTimeEmulatedSync > 0.0)
+				{
+					QUICK_SCOPE_CYCLE_COUNTER(STAT_StallForEmulatedSyncInterval);
+					float MinTimeBetweenFrames = (float(DesiredSyncIntervalRelativeToDevice) / DriverRefreshRate);
+
+					float ThisTime = FPlatformTime::Seconds() - LastTimeEmulatedSync;
+					if (ThisTime > 0 && ThisTime < MinTimeBetweenFrames)
+					{
+						FPlatformProcess::Sleep(MinTimeBetweenFrames - ThisTime);
+					}
+				}
+			}
+		}
+		if (CanUseGetFrameTimestamps())
+		{
+			UE_CLOG(bPrintMethod, LogRHI, Display, TEXT("Using eglGetFrameTimestampsANDROID method for frame pacing"));
+
+			//static bool bPrintOnce = true;
+			if (FrameIDs[(int32(NextFrameIDSlot) - 1) % NUM_FRAMES_TO_MONITOR])
+				// not supported   && eglGetFrameTimestampsSupportedANDROID_p && eglGetFrameTimestampsSupportedANDROID_p(eglDisplay, eglSurface, EGL_FIRST_COMPOSITION_START_TIME_ANDROID))
+			{
+				//UE_CLOG(bPrintOnce, LogRHI, Log, TEXT("eglGetFrameTimestampsSupportedANDROID retured true for EGL_FIRST_COMPOSITION_START_TIME_ANDROID"));
+				EGLint TimestampList = EGL_FIRST_COMPOSITION_START_TIME_ANDROID;
+				//EGLint TimestampList = EGL_COMPOSITION_LATCH_TIME_ANDROID;
+				//EGLint TimestampList = EGL_LAST_COMPOSITION_START_TIME_ANDROID;
+				//EGLint TimestampList = EGL_DISPLAY_PRESENT_TIME_ANDROID;
+				EGLnsecsANDROID Result = 0;
+				int32 DeltaFrameIndex = 1;
+				for (int32 Index = int32(NextFrameIDSlot) - 1; Index >= int32(NextFrameIDSlot) - NUM_FRAMES_TO_MONITOR && Index >= 0; Index--)
+				{
+					Result = 0;
+					if (FrameIDs[Index % NUM_FRAMES_TO_MONITOR])
+					{
+						eglGetFrameTimestampsANDROID_p(eglDisplay, eglSurface, FrameIDs[Index % NUM_FRAMES_TO_MONITOR], 1, &TimestampList, &Result);
+					}
+					if (Result > 0)
+					{
+						break;
+					}
+					DeltaFrameIndex++;
+				}
+
+				GGetTimeStampsSucceededThisFrame = Result > 0;
+				if (GGetTimeStampsSucceededThisFrame)
+				{
+					EGLnsecsANDROID FudgeFactor = 0; //  8333 * 1000;
+					EGLnsecsANDROID DeltaNanos = EGLnsecsANDROID(DesiredSyncIntervalRelativeToDevice) * EGLnsecsANDROID(DeltaFrameIndex) * DriverRefreshNanos;
+					EGLnsecsANDROID PresentationTime = Result + DeltaNanos + FudgeFactor;
+					eglPresentationTimeANDROID_p(eglDisplay, eglSurface, PresentationTime);
+					GGetTimeStampsRetryCount = 0;
+				}
+				else
+				{
+					GGetTimeStampsRetryCount++;
+					if (GGetTimeStampsRetryCount == FAndroidPlatformRHIFramePacer::CVarTimeStampErrorRetryCount.GetValueOnAnyThread())
+					{
+						UE_LOG(LogRHI, Log, TEXT("eglGetFrameTimestampsANDROID_p failed for %d consecutive frames, reverting to naive frame pacer."), GGetTimeStampsRetryCount);
+					}
+				}
+			}
+			else
+			{
+				//UE_CLOG(bPrintOnce, LogRHI, Log, TEXT("eglGetFrameTimestampsSupportedANDROID doesn't exist or retured false for EGL_FIRST_COMPOSITION_START_TIME_ANDROID, discarding eglGetNextFrameIdANDROID_p and eglGetFrameTimestampsANDROID_p"));
+			}
+			//bPrintOnce = false;
+		}
+
+		LastTimeEmulatedSync = FPlatformTime::Seconds();
+
+		{
+			QUICK_SCOPE_CYCLE_COUNTER(STAT_eglSwapBuffers);
+
+			FrameIDs[(NextFrameIDSlot) % NUM_FRAMES_TO_MONITOR] = 0;
+			if (eglGetNextFrameIdANDROID_p && (CanUseGetFrameTimestamps() || FAndroidPlatformRHIFramePacer::CVarSpewGetFrameTimestamps.GetValueOnAnyThread()))
+			{
+				eglGetNextFrameIdANDROID_p(eglDisplay, eglSurface, &FrameIDs[(NextFrameIDSlot) % NUM_FRAMES_TO_MONITOR]);
+			}
+			NextFrameIDSlot++;
+
+			if (eglSurface == NULL || !eglSwapBuffers(eglDisplay, eglSurface))
+			{
+				// shutdown if swapbuffering goes down
+				if (SwapBufferFailureCount > 10)
+				{
+					//Process.killProcess(Process.myPid());		//@todo android
+				}
+				SwapBufferFailureCount++;
+
+				// basic reporting
+				if (eglSurface == NULL)
+				{
+					return false;
+				}
+				else
+				{
+					if (eglGetError() == EGL_CONTEXT_LOST)
+					{
+						//Logger.LogOut("swapBuffers: EGL11.EGL_CONTEXT_LOST err: " + eglGetError());					
+						//Process.killProcess(Process.myPid());		//@todo android
+					}
+				}
+
+				return false;
+			}
+		}
+
+		if (DesiredSyncIntervalRelativeToDevice > 0 && eglGetFrameTimestampsANDROID_p && FAndroidPlatformRHIFramePacer::CVarSpewGetFrameTimestamps.GetValueOnAnyThread())
+		{
+			static EGLint TimestampList[9] =
+			{
+				EGL_REQUESTED_PRESENT_TIME_ANDROID,
+				EGL_RENDERING_COMPLETE_TIME_ANDROID,
+				EGL_COMPOSITION_LATCH_TIME_ANDROID,
+				EGL_FIRST_COMPOSITION_START_TIME_ANDROID,
+				EGL_LAST_COMPOSITION_START_TIME_ANDROID,
+				EGL_FIRST_COMPOSITION_GPU_FINISHED_TIME_ANDROID,
+				EGL_DISPLAY_PRESENT_TIME_ANDROID,
+				EGL_DEQUEUE_READY_TIME_ANDROID,
+				EGL_READS_DONE_TIME_ANDROID
+			};
+
+			static const TCHAR* TimestampStrings[9] =
+			{
+				TEXT("EGL_REQUESTED_PRESENT_TIME_ANDROID"),
+				TEXT("EGL_RENDERING_COMPLETE_TIME_ANDROID"),
+				TEXT("EGL_COMPOSITION_LATCH_TIME_ANDROID"),
+				TEXT("EGL_FIRST_COMPOSITION_START_TIME_ANDROID"),
+				TEXT("EGL_LAST_COMPOSITION_START_TIME_ANDROID"),
+				TEXT("EGL_FIRST_COMPOSITION_GPU_FINISHED_TIME_ANDROID"),
+				TEXT("EGL_DISPLAY_PRESENT_TIME_ANDROID"),
+				TEXT("EGL_DEQUEUE_READY_TIME_ANDROID"),
+				TEXT("EGL_READS_DONE_TIME_ANDROID")
+			};
+
+
+			EGLnsecsANDROID Results[NUM_FRAMES_TO_MONITOR][9] = { {0} };
+			EGLnsecsANDROID FirstRealValue = 0;
+			for (int32 Index = int32(NextFrameIDSlot) - NUM_FRAMES_TO_MONITOR; Index < int32(NextFrameIDSlot); Index++)
+			{
+				eglGetFrameTimestampsANDROID_p(eglDisplay, eglSurface, FrameIDs[Index % NUM_FRAMES_TO_MONITOR], 9, TimestampList, Results[Index % NUM_FRAMES_TO_MONITOR]);
+				for (int32 IndexInner = 0; IndexInner < 9; IndexInner++)
+				{
+					if (!FirstRealValue || (Results[Index % NUM_FRAMES_TO_MONITOR][IndexInner] > 1 && Results[Index % NUM_FRAMES_TO_MONITOR][IndexInner] < FirstRealValue))
+					{
+						FirstRealValue = Results[Index % NUM_FRAMES_TO_MONITOR][IndexInner];
+					}
+				}
+			}
+			UE_CLOG(FAndroidPlatformRHIFramePacer::CVarSpewGetFrameTimestamps.GetValueOnAnyThread() > 1, LogRHI, Log, TEXT("************************************  frame %d   base time is %lld"), NextFrameIDSlot - 1, FirstRealValue);
+
+			for (int32 Index = int32(NextFrameIDSlot) - NUM_FRAMES_TO_MONITOR; Index < int32(NextFrameIDSlot); Index++)
+			{
+
+				UE_CLOG(FAndroidPlatformRHIFramePacer::CVarSpewGetFrameTimestamps.GetValueOnAnyThread() > 1, LogRHI, Log, TEXT("eglGetFrameTimestampsANDROID_p  frame %d"), Index);
+				for (int32 IndexInner = 0; IndexInner < 9; IndexInner++)
+				{
+					int32 MsVal = (Results[Index % NUM_FRAMES_TO_MONITOR][IndexInner] > 1) ? int32((Results[Index % NUM_FRAMES_TO_MONITOR][IndexInner] - FirstRealValue) / 1000000) : int32(Results[Index % NUM_FRAMES_TO_MONITOR][IndexInner]);
+
+					UE_CLOG(FAndroidPlatformRHIFramePacer::CVarSpewGetFrameTimestamps.GetValueOnAnyThread() > 1, LogRHI, Log, TEXT("     %8d    %s"), MsVal, TimestampStrings[IndexInner]);
+				}
+			}
+
+			int32 IndexLast = int32(NextFrameIDSlot) - NUM_FRAMES_TO_MONITOR;
+			int32 IndexLastNext = IndexLast + 1;
+
+			if (Results[IndexLast % NUM_FRAMES_TO_MONITOR][3] > 1 && Results[IndexLastNext % NUM_FRAMES_TO_MONITOR][3] > 1)
+			{
+				int32 MsVal = int32((Results[IndexLastNext % NUM_FRAMES_TO_MONITOR][3] - Results[IndexLast % NUM_FRAMES_TO_MONITOR][3]) / 1000000);
+
+				RecordedFrameInterval[NumRecordedFrameInterval++] = MsVal;
+				if (NumRecordedFrameInterval == 100)
+				{
+					FString All;
+					int32 NumOnTarget = 0;
+					int32 NumBelowTarget = 0;
+					int32 NumAboveTarget = 0;
+					for (int32 Index = 0; Index < 100; Index++)
+					{
+						if (Index)
+						{
+							All += TCHAR(' ');
+						}
+						All += FString::Printf(TEXT("%d"), RecordedFrameInterval[Index]);
+
+						if (RecordedFrameInterval[Index] > DesiredSyncIntervalRelativeTo60Hz * 16 - 8 && RecordedFrameInterval[Index] < DesiredSyncIntervalRelativeTo60Hz * 16 + 8)
+						{
+							NumOnTarget++;
+						}
+						else if (RecordedFrameInterval[Index] < DesiredSyncIntervalRelativeTo60Hz * 16)
+						{
+							NumBelowTarget++;
+						}
+						else
+						{
+							NumAboveTarget++;
+						}
+					}
+					UE_LOG(LogRHI, Log, TEXT("%3d fast  %3d ok  %3d slow   %s"), NumBelowTarget, NumOnTarget, NumAboveTarget, *All);
+					NumRecordedFrameInterval = 0;
+				}
+			}
+		}
+	}
+
+	return true;
+}
+
+#endif
diff --git a/Engine/Source/Runtime/OpenGLDrv/Private/Android/AndroidOpenGLFramePacer.h b/Engine/Source/Runtime/OpenGLDrv/Private/Android/AndroidOpenGLFramePacer.h
new file mode 100644
--- /dev/null
+++ b/Engine/Source/Runtime/OpenGLDrv/Private/Android/AndroidOpenGLFramePacer.h
@@ -0,0 +1,4 @@
+// Copyright Epic Games, Inc. All Rights Reserved.
+
+#include "Android/AndroidPlatform.h"
+#include "Android/AndroidPlatformFramePacer.h"
diff --git a/Engine/Source/Runtime/OpenGLDrv/Private/Lumin/LuminGL4.cpp b/Engine/Source/Runtime/OpenGLDrv/Private/Lumin/LuminGL4.cpp
--- a/Engine/Source/Runtime/OpenGLDrv/Private/Lumin/LuminGL4.cpp
+++ b/Engine/Source/Runtime/OpenGLDrv/Private/Lumin/LuminGL4.cpp
@@ -498,11 +498,11 @@ bool PlatformBlitToViewport(FPlatformOpenGLDevice* Device,
 	uint32 BackbufferSizeX,
 	uint32 BackbufferSizeY,
 	bool bPresent,
-	bool bLockToVsync,
-	int32 SyncInterval)
+	bool bLockToVsync)
 {
 	FScopeLock ScopeLock(Device->ContextUsageGuard);
 	{
+		int32 SyncInterval = RHIGetSyncInterval();
 		FPlatformOpenGLContext* const Context = Viewport.GetGLContext();
 		check(Context && Context->OpenGLContext);
 		FScopeContext ScopeContext(Context);
diff --git a/Engine/Source/Runtime/OpenGLDrv/Private/Lumin/LuminOpenGL.cpp b/Engine/Source/Runtime/OpenGLDrv/Private/Lumin/LuminOpenGL.cpp
--- a/Engine/Source/Runtime/OpenGLDrv/Private/Lumin/LuminOpenGL.cpp
+++ b/Engine/Source/Runtime/OpenGLDrv/Private/Lumin/LuminOpenGL.cpp
@@ -88,8 +88,9 @@ void* PlatformGetWindow(FPlatformOpenGLContext* Context, void** AddParam)
 	return (void*)&Context->eglContext;
 }
 
-bool PlatformBlitToViewport(FPlatformOpenGLDevice* Device, const FOpenGLViewport& Viewport, uint32 BackbufferSizeX, uint32 BackbufferSizeY, bool bPresent, bool bLockToVsync, int32 SyncInterval)
+bool PlatformBlitToViewport(FPlatformOpenGLDevice* Device, const FOpenGLViewport& Viewport, uint32 BackbufferSizeX, uint32 BackbufferSizeY, bool bPresent, bool bLockToVsync)
 {
+	int32 SyncInterval = RHIGetSyncInterval();
 	FPlatformOpenGLContext* const Context = Viewport.GetGLContext();
 	check(Context && Context->eglContext);
 	FScopeContext ScopeContext(Context);
diff --git a/Engine/Source/Runtime/OpenGLDrv/Private/OpenGLDrvPrivate.h b/Engine/Source/Runtime/OpenGLDrv/Private/OpenGLDrvPrivate.h
--- a/Engine/Source/Runtime/OpenGLDrv/Private/OpenGLDrvPrivate.h
+++ b/Engine/Source/Runtime/OpenGLDrv/Private/OpenGLDrvPrivate.h
@@ -250,7 +250,7 @@ FRHITexture* PlatformCreateBuiltinBackBuffer(FOpenGLDynamicRHI* OpenGLRHI, uint3
  * On Windows it temporarily switches OpenGL context, on Mac only context's output view.
  * Should return true if frame was presented and it is necessary to finish frame rendering.
  */
-bool PlatformBlitToViewport( FPlatformOpenGLDevice* Device, const FOpenGLViewport& Viewport, uint32 BackbufferSizeX, uint32 BackbufferSizeY, bool bPresent,bool bLockToVsync, int32 SyncInterval );
+bool PlatformBlitToViewport(FPlatformOpenGLDevice* Device, const FOpenGLViewport& Viewport, uint32 BackbufferSizeX, uint32 BackbufferSizeY, bool bPresent,bool bLockToVsync);
 
 /**
  * Resize the GL context for platform.
diff --git a/Engine/Source/Runtime/OpenGLDrv/Private/OpenGLViewport.cpp b/Engine/Source/Runtime/OpenGLDrv/Private/OpenGLViewport.cpp
--- a/Engine/Source/Runtime/OpenGLDrv/Private/OpenGLViewport.cpp
+++ b/Engine/Source/Runtime/OpenGLDrv/Private/OpenGLViewport.cpp
@@ -11,19 +11,6 @@
 #include "OpenGLDrv.h"
 #include "OpenGLDrvPrivate.h"
 
-/**
- * RHI console variables used by viewports.
- */
-namespace RHIOpenGLConsoleVariables
-{
-	int32 SyncInterval = 1;
-	static FAutoConsoleVariableRef CVarSyncInterval(
-		TEXT("RHI.SyncIntervalOgl"),
-		SyncInterval,
-		TEXT("When synchronizing with OpenGL, specifies the interval at which to refresh.")
-		);
-};
-
 void FOpenGLDynamicRHI::RHIGetSupportedResolution(uint32 &Width, uint32 &Height)
 {
 	PlatformGetSupportedResolution(Width, Height);
@@ -145,8 +132,7 @@ void FOpenGLDynamicRHI::RHIEndDrawingViewport(FRHIViewport* ViewportRHI,bool bPr
 		BackBuffer->GetSizeX(),
 		BackBuffer->GetSizeY(),
 		bPresent,
-		bLockToVsync,
-		RHIGetSyncInterval()
+		bLockToVsync
 	);
 
 	// Always consider the Framebuffer in the rendering context dirty after the blit
diff --git a/Engine/Source/Runtime/OpenGLDrv/Private/Windows/OpenGLWindows.cpp b/Engine/Source/Runtime/OpenGLDrv/Private/Windows/OpenGLWindows.cpp
--- a/Engine/Source/Runtime/OpenGLDrv/Private/Windows/OpenGLWindows.cpp
+++ b/Engine/Source/Runtime/OpenGLDrv/Private/Windows/OpenGLWindows.cpp
@@ -456,8 +456,10 @@ void PlatformDestroyOpenGLContext(FPlatformOpenGLDevice* Device, FPlatformOpenGL
  * Main function for transferring data to on-screen buffers.
  * On Windows it temporarily switches OpenGL context, on Mac only context's output view.
  */
-bool PlatformBlitToViewport( FPlatformOpenGLDevice* Device, const FOpenGLViewport& Viewport, uint32 BackbufferSizeX, uint32 BackbufferSizeY, bool bPresent,bool bLockToVsync, int32 SyncInterval )
+bool PlatformBlitToViewport( FPlatformOpenGLDevice* Device, const FOpenGLViewport& Viewport, uint32 BackbufferSizeX, uint32 BackbufferSizeY, bool bPresent,bool bLockToVsync)
 {
+	int32 SyncInterval = RHIGetSyncInterval();
+
 	FPlatformOpenGLContext* const Context = Viewport.GetGLContext();
 
 	check(Context && Context->DeviceContext);
diff --git a/Engine/Source/Runtime/VulkanRHI/Private/Android/VulkanAndroidPlatform.cpp b/Engine/Source/Runtime/VulkanRHI/Private/Android/VulkanAndroidPlatform.cpp
--- a/Engine/Source/Runtime/VulkanRHI/Private/Android/VulkanAndroidPlatform.cpp
+++ b/Engine/Source/Runtime/VulkanRHI/Private/Android/VulkanAndroidPlatform.cpp
@@ -14,6 +14,7 @@
 #include "../VulkanRHIPrivate.h"
 #include <dlfcn.h>
 #include "Android/AndroidWindow.h"
+#include "Android/AndroidPlatformFramePacer.h"
 
 // Vulkan function pointers
 #define DEFINE_VK_ENTRYPOINTS(Type,Func) Type VulkanDynamicAPI::Func = NULL;
@@ -62,6 +63,9 @@ bool FVulkanAndroidPlatform::LoadVulkanLibrary()
 
 #undef GET_VK_ENTRYPOINTS
 
+	// Init frame pacer
+	FPlatformRHIFramePacer::Init(new FAndroidVulkanFramePacer());
+
 	return true;
 }
 
