Enabling linking with GameSDK

diff --git a/Engine/Source/Runtime/OpenGLDrv/OpenGLDrv.Build.cs b/Engine/Source/Runtime/OpenGLDrv/OpenGLDrv.Build.cs
--- a/Engine/Source/Runtime/OpenGLDrv/OpenGLDrv.Build.cs
+++ b/Engine/Source/Runtime/OpenGLDrv/OpenGLDrv.Build.cs
@@ -55,8 +55,8 @@ public class OpenGLDrv : ModuleRules
         if (Target.Platform == UnrealTargetPlatform.Android)
         {
             // for Swappy
-            PublicDefinitions.Add("USE_ANDROID_OPENGL_SWAPPY=0");
-/*
+            PublicDefinitions.Add("USE_ANDROID_OPENGL_SWAPPY=1");
+
             PrivateDependencyModuleNames.AddRange(
                 new string[]
                 {
@@ -64,7 +64,6 @@ public class OpenGLDrv : ModuleRules
                 "GoogleGameSDK"
                 }
             );
-*/
         }
 
         if (Target.Platform != UnrealTargetPlatform.Win32 && Target.Platform != UnrealTargetPlatform.Win64
