Fixed: Runtime Virtual Texture crash on mobile

diff --git a/Engine/Shaders/Private/MobileBasePassPixelShader.usf b/Engine/Shaders/Private/MobileBasePassPixelShader.usf
--- a/Engine/Shaders/Private/MobileBasePassPixelShader.usf
+++ b/Engine/Shaders/Private/MobileBasePassPixelShader.usf
@@ -362,13 +362,11 @@ void ApplyPixelDepthOffsetForMobileBasePass(inout FMaterialPixelParameters Mater
     float PixelDepthOffset = ApplyPixelDepthOffsetToMaterialParameters(MaterialParameters, PixelMaterialInputs, OutDepth);
 }
 
-
-// Disable early depth_stencil for masked materials and depth writes
-// Cross-compiler (hlslcc) checks for 'discard' and depth write usage anyway and disables it, this ifdef is needed for HLSL
-#if MATERIALBLENDING_MASKED || USE_DITHERED_LOD_TRANSITION || OUTPUT_PIXEL_DEPTH_OFFSET
-	#define PIXELSHADER_EARLYDEPTHSTENCIL 
+// Force early depth_stencil for non-masked material that use VT feedback
+#if (NUM_VIRTUALTEXTURE_SAMPLES || LIGHTMAP_VT_ENABLED) && !(MATERIALBLENDING_MASKED || USE_DITHERED_LOD_TRANSITION || OUTPUT_PIXEL_DEPTH_OFFSET)
+	#define PIXELSHADER_EARLYDEPTHSTENCIL EARLYDEPTHSTENCIL	
 #else
-	#define PIXELSHADER_EARLYDEPTHSTENCIL EARLYDEPTHSTENCIL
+	#define PIXELSHADER_EARLYDEPTHSTENCIL 	
 #endif
 
 PIXELSHADER_EARLYDEPTHSTENCIL
diff --git a/Engine/Source/Runtime/OpenGLDrv/Private/Android/AndroidOpenGL.cpp b/Engine/Source/Runtime/OpenGLDrv/Private/Android/AndroidOpenGL.cpp
--- a/Engine/Source/Runtime/OpenGLDrv/Private/Android/AndroidOpenGL.cpp
+++ b/Engine/Source/Runtime/OpenGLDrv/Private/Android/AndroidOpenGL.cpp
@@ -39,6 +39,12 @@ static TAutoConsoleVariable<int32> CVarEnableAdrenoTilingHint(
 	TEXT("  1 = hinting enabled for Adreno devices running Andorid 8 or earlier [default]\n")
 	TEXT("  2 = hinting always enabled for Adreno devices\n"));
 
+static TAutoConsoleVariable<int32> CVarDisableEarlyFragmentTests(
+	TEXT("r.Android.DisableEarlyFragmentTests"),
+	0,
+	TEXT("Whether to disable early_fragment_tests if any \n"),
+	ECVF_ReadOnly);
+
 struct FPlatformOpenGLDevice
 {
 
@@ -1000,6 +1006,12 @@ void FAndroidOpenGL::ProcessExtensions(const FString& ExtensionsString)
 			UE_LOG(LogRHI, Log, TEXT("Using QCOM_shader_framebuffer_fetch_noncoherent"));
 		}
 	}
+
+	if (CVarDisableEarlyFragmentTests.GetValueOnAnyThread() != 0)
+	{
+		bRequiresDisabledEarlyFragmentTests = true;
+		UE_LOG(LogRHI, Log, TEXT("Disabling early_fragment_tests"));
+	}
 }
 
 FString FAndroidMisc::GetGPUFamily()
diff --git a/Engine/Source/Runtime/OpenGLDrv/Private/OpenGLES.cpp b/Engine/Source/Runtime/OpenGLDrv/Private/OpenGLES.cpp
--- a/Engine/Source/Runtime/OpenGLDrv/Private/OpenGLES.cpp
+++ b/Engine/Source/Runtime/OpenGLDrv/Private/OpenGLES.cpp
@@ -245,6 +245,9 @@ bool FOpenGLES::bRequiresTexture2DPrecisionHack = false;
 /* This is a hack to add a round() function when not available to a shader compiler */
 bool FOpenGLES::bRequiresRoundFunctionHack = true;
 
+/* Some Mali devices do not work correctly with early_fragment_test enabled */
+bool FOpenGLES::bRequiresDisabledEarlyFragmentTests = false;
+
 /* This is to avoid a bug in Adreno drivers that define GL_ARM_shader_framebuffer_fetch_depth_stencil even when device does not support this extension  */
 bool FOpenGLES::bRequiresARMShaderFramebufferFetchDepthStencilUndef = false;
 
diff --git a/Engine/Source/Runtime/OpenGLDrv/Private/OpenGLShaders.cpp b/Engine/Source/Runtime/OpenGLDrv/Private/OpenGLShaders.cpp
--- a/Engine/Source/Runtime/OpenGLDrv/Private/OpenGLShaders.cpp
+++ b/Engine/Source/Runtime/OpenGLDrv/Private/OpenGLShaders.cpp
@@ -962,6 +962,7 @@ void OPENGLDRV_API GetCurrentOpenGLShaderDeviceCapabilities(FOpenGLShaderDeviceC
 		Capabilities.MaxVaryingVectors = FOpenGL::GetMaxVaryingVectors();
 		Capabilities.bRequiresTexture2DPrecisionHack = FOpenGL::RequiresTexture2DPrecisionHack();
 		Capabilities.bRequiresRoundFunctionHack = FOpenGL::RequiresRoundFunctionHack();
+		Capabilities.bRequiresDisabledEarlyFragmentTests = FOpenGL::RequiresDisabledEarlyFragmentTests();
 	#endif // PLATFORM_LUMINGL4
 #elif PLATFORM_IOS
 	Capabilities.TargetPlatform = EOpenGLShaderTargetPlatform::OGLSTP_iOS;
@@ -1039,7 +1040,7 @@ void OPENGLDRV_API GLSLToDeviceCompatibleGLSL(FAnsiCharArray& GlslCodeOriginal,
 		AppendCString(GlslCode, "#endif\n");
 	}
 
-	if (TypeEnum == GL_FRAGMENT_SHADER)
+	if (TypeEnum == GL_FRAGMENT_SHADER && Capabilities.bRequiresDisabledEarlyFragmentTests)
 	{
 		ReplaceCString(GlslCodeOriginal, "layout(early_fragment_tests) in;", "");
 	}
diff --git a/Engine/Source/Runtime/OpenGLDrv/Public/OpenGLES.h b/Engine/Source/Runtime/OpenGLDrv/Public/OpenGLES.h
--- a/Engine/Source/Runtime/OpenGLDrv/Public/OpenGLES.h
+++ b/Engine/Source/Runtime/OpenGLDrv/Public/OpenGLES.h
@@ -695,6 +695,7 @@ struct FOpenGLES : public FOpenGLBase
 	static FORCEINLINE bool HasBinaryProgramRetrievalFailed() { return bBinaryProgramRetrievalFailed; }
 	static FORCEINLINE bool RequiresTexture2DPrecisionHack() { return bRequiresTexture2DPrecisionHack; }
 	static FORCEINLINE bool RequiresRoundFunctionHack() { return bRequiresRoundFunctionHack; }
+	static FORCEINLINE bool RequiresDisabledEarlyFragmentTests() { return bRequiresDisabledEarlyFragmentTests; }
 	static FORCEINLINE bool RequiresARMShaderFramebufferFetchDepthStencilUndef() { return bRequiresARMShaderFramebufferFetchDepthStencilUndef; }
 	static FORCEINLINE bool IsCheckingShaderCompilerHacks() { return bIsCheckingShaderCompilerHacks; }
 	static FORCEINLINE bool IsLimitingShaderCompileCount() { return bIsLimitingShaderCompileCount; }
@@ -1356,6 +1357,9 @@ public:
 	/* This is a hack to add a round() function when not available to a shader compiler */
 	static bool bRequiresRoundFunctionHack;
 
+	/* Some Mali devices do not work correctly with early_fragment_test enabled */
+	static bool bRequiresDisabledEarlyFragmentTests;
+		
 	/* This is to avoid a bug in Adreno drivers that define GL_ARM_shader_framebuffer_fetch_depth_stencil even when device does not support this extension  */
 	static bool bRequiresARMShaderFramebufferFetchDepthStencilUndef;
 
diff --git a/Engine/Source/Runtime/OpenGLDrv/Public/OpenGLShaders.h b/Engine/Source/Runtime/OpenGLDrv/Public/OpenGLShaders.h
--- a/Engine/Source/Runtime/OpenGLDrv/Public/OpenGLShaders.h
+++ b/Engine/Source/Runtime/OpenGLDrv/Public/OpenGLShaders.h
@@ -34,6 +34,7 @@ struct FOpenGLShaderDeviceCapabilities
 	GLint MaxVaryingVectors;
 	bool bRequiresTexture2DPrecisionHack;
 	bool bRequiresRoundFunctionHack;
+	bool bRequiresDisabledEarlyFragmentTests;
 };
 
 /**
diff --git a/Engine/Source/Runtime/Renderer/Private/VT/VirtualTexturePhysicalSpace.cpp b/Engine/Source/Runtime/Renderer/Private/VT/VirtualTexturePhysicalSpace.cpp
--- a/Engine/Source/Runtime/Renderer/Private/VT/VirtualTexturePhysicalSpace.cpp
+++ b/Engine/Source/Runtime/Renderer/Private/VT/VirtualTexturePhysicalSpace.cpp
@@ -94,14 +94,16 @@ void FVirtualTexturePhysicalSpace::InitRHI()
 		const EPixelFormat FormatSRV = RemapVirtualTexturePhysicalSpaceFormat(Description.Format[Layer]);
 		const EPixelFormat FormatUAV = GetUnorderedAccessViewFormat(FormatSRV);
 		const bool bCreateAliasedUAV = (FormatUAV != PF_Unknown) && (FormatUAV != FormatSRV);
-
+		// Not all mobile RHIs support sRGB texture views/aliasing, use only linear targets on mobile
+		uint32 VT_SRGB = GMaxRHIFeatureLevel > ERHIFeatureLevel::ES3_1 ? TexCreate_SRGB : TexCreate_None;
+		
 		// Allocate physical texture from the render target pool
 		const uint32 TextureSize = GetTextureSize();
 		const FPooledRenderTargetDesc Desc = FPooledRenderTargetDesc::Create2DDesc(
 			FIntPoint(TextureSize, TextureSize),
 			FormatSRV,
 			FClearValueBinding::None,
-			TexCreate_SRGB,
+			VT_SRGB,
 			bCreateAliasedUAV ? TexCreate_ShaderResource | TexCreate_UAV : TexCreate_ShaderResource,
 			false);
 
