Updating Graphics Resources for Indirect Commands - fixing Niagara for Opengl

diff --git a/Engine/Source/Runtime/OpenGLDrv/Private/OpenGLCommands.cpp b/Engine/Source/Runtime/OpenGLDrv/Private/OpenGLCommands.cpp
--- a/Engine/Source/Runtime/OpenGLDrv/Private/OpenGLCommands.cpp
+++ b/Engine/Source/Runtime/OpenGLDrv/Private/OpenGLCommands.cpp
@@ -2613,6 +2613,7 @@ void FOpenGLDynamicRHI::RHIDrawPrimitiveIndirect(FRHIVertexBuffer* ArgumentBuffe
 		UpdateRasterizerStateInOpenGLContext(ContextState);
 		UpdateDepthStencilStateInOpenGLContext(ContextState);
 		BindPendingShaderState(ContextState);
+		CommitGraphicsResourceTables();
 		SetupTexturesForDraw(ContextState);
 		SetupUAVsForDraw(ContextState);
 		CommitNonComputeShaderConstants();
@@ -2670,6 +2671,7 @@ void FOpenGLDynamicRHI::RHIDrawIndexedIndirect(FRHIIndexBuffer* IndexBufferRHI,
 		UpdateRasterizerStateInOpenGLContext(ContextState);
 		UpdateDepthStencilStateInOpenGLContext(ContextState);
 		BindPendingShaderState(ContextState);
+		CommitGraphicsResourceTables();
 		SetupTexturesForDraw(ContextState);
 		SetupUAVsForDraw(ContextState);
 		CommitNonComputeShaderConstants();
@@ -2832,6 +2834,7 @@ void FOpenGLDynamicRHI::RHIDrawIndexedPrimitiveIndirect(FRHIIndexBuffer* IndexBu
 		UpdateRasterizerStateInOpenGLContext(ContextState);
 		UpdateDepthStencilStateInOpenGLContext(ContextState);
 		BindPendingShaderState(ContextState);
+		CommitGraphicsResourceTables();
 		SetupTexturesForDraw(ContextState);
 		SetupUAVsForDraw(ContextState);
 		CommitNonComputeShaderConstants();
