Fix linux build failures due to the signature for PlatformBlitToViewport changing but the linux opengl impementation not being updated

diff --git a/Engine/Source/Runtime/OpenGLDrv/Private/Linux/OpenGLLinux.cpp b/Engine/Source/Runtime/OpenGLDrv/Private/Linux/OpenGLLinux.cpp
--- a/Engine/Source/Runtime/OpenGLDrv/Private/Linux/OpenGLLinux.cpp
+++ b/Engine/Source/Runtime/OpenGLDrv/Private/Linux/OpenGLLinux.cpp
@@ -427,14 +427,16 @@ bool PlatformBlitToViewport(FPlatformOpenGLDevice* Device,
 							uint32 BackbufferSizeX,
 							uint32 BackbufferSizeY,
 							bool bPresent,
-							bool bLockToVsync,
-							int32 SyncInterval )
+							bool bLockToVsync)
 {
 	FPlatformOpenGLContext* const Context = Viewport.GetGLContext();
 
 	check( Context && Context->hWnd );
 
 	FScopeLock ScopeLock( Device->ContextUsageGuard );
+
+	int32 SyncInterval = (int32)RHIGetSyncInterval();
+
 	{
 		FScopeContext ScopeContext( Context );
 
