Fixed: SunTemple loads into a black screen on some Android devices

diff --git a/Engine/Source/Runtime/OpenGLDrv/Private/OpenGLShaders.cpp b/Engine/Source/Runtime/OpenGLDrv/Private/OpenGLShaders.cpp
--- a/Engine/Source/Runtime/OpenGLDrv/Private/OpenGLShaders.cpp
+++ b/Engine/Source/Runtime/OpenGLDrv/Private/OpenGLShaders.cpp
@@ -2432,20 +2432,18 @@ void FOpenGLLinkedProgram::ConfigureShaderStage( int Stage, uint32 FirstUniformB
 		Name.Buffer[6] = 0;
 		for (uint8 UB = 0; UB < Config.Shaders[Stage].Bindings.NumUniformBuffers; ++UB)
 		{
+			const TArray<CrossCompiler::FPackedArrayInfo>& PackedInfo = Config.Shaders[Stage].Bindings.PackedUniformBuffers[UB];
+			
 			TArray<FPackedUniformInfo> PackedBuffers;
 			ANSICHAR* Str = SetIndex(Name.Buffer, 2, UB);
 			*Str++ = '_';
 			Str[1] = 0;
-			for (uint8 Index = 0; Index < CrossCompiler::PACKED_TYPEINDEX_MAX; ++Index)
+			for (uint8 Index = 0; Index < PackedInfo.Num(); ++Index)
 			{
-				uint8 ArrayIndexType = CrossCompiler::PackedTypeIndexToTypeName(Index);
-				Str[0] = ArrayIndexType;
-				GLint Location = glGetUniformLocation(StageProgram, Name.Buffer);
-				if ((int32)Location != -1)
-				{
-					FPackedUniformInfo Info = {Location, ArrayIndexType, Index};
-					PackedBuffers.Add(Info);
-				}
+				Str[0] = PackedInfo[Index].TypeName;
+				GLint Location = glGetUniformLocation(StageProgram, Name.Buffer); // This could be -1 if optimized out
+				FPackedUniformInfo Info = {Location, PackedInfo[Index].TypeName,  PackedInfo[Index].TypeIndex};
+				PackedBuffers.Add(Info);
 			}
 
 			StagePackedUniformInfo[Stage].PackedUniformBufferInfos.Add(PackedBuffers);
@@ -4148,6 +4146,12 @@ void FOpenGLShaderParameterCache::CommitPackedUniformBuffers(FOpenGLLinkedProgra
 				{
 					auto& UBInfo = Bindings.PackedUniformBuffers[BufferIndex];
 					const auto& UniformInfo = UniformBufferUploadInfoList[InfoIndex];
+					if (UniformInfo.Location < 0)
+					{
+						// Optimized out
+						continue;
+					}
+					
 					const void* RESTRICT UniformData = PackedUniformsScratch[UniformInfo.Index];
 					int32 NumVectors = UBInfo[InfoIndex].Size / SizeOfFloat4;
 					check(UniformInfo.ArrayType == UBInfo[InfoIndex].TypeName);
