Add new texture aliasing APIs for OpenGL. The missing implementation was causing an assertion failed and crash initializing mobile VR devices using OpenGL.

diff --git a/Engine/Source/Runtime/OpenGLDrv/Private/OpenGLTexture.cpp b/Engine/Source/Runtime/OpenGLDrv/Private/OpenGLTexture.cpp
--- a/Engine/Source/Runtime/OpenGLDrv/Private/OpenGLTexture.cpp
+++ b/Engine/Source/Runtime/OpenGLDrv/Private/OpenGLTexture.cpp
@@ -2698,6 +2698,18 @@ FTextureCubeRHIRef FOpenGLDynamicRHI::RHICreateTextureCubeFromResource(EPixelFor
 
 PRAGMA_DISABLE_DEPRECATION_WARNINGS
 
+void FOpenGLDynamicRHI::RHIAliasTextureResources(FTextureRHIRef& DestRHITexture, FTextureRHIRef& SrcRHITexture)
+{
+	// @todo: Move the raw-pointer implementation down here when it's deprecation is completed.
+	RHIAliasTextureResources((FRHITexture*)DestRHITexture, (FRHITexture*)SrcRHITexture);
+}
+
+FTextureRHIRef FOpenGLDynamicRHI::RHICreateAliasedTexture(FTextureRHIRef& SourceTexture)
+{
+	// @todo: Move the raw-pointer implementation down here when it's deprecation is completed.
+	return RHICreateAliasedTexture((FRHITexture*)SourceTexture);
+}
+
 void FOpenGLDynamicRHI::RHIAliasTextureResources(FRHITexture* DestRHITexture, FRHITexture* SrcRHITexture)
 {
 	FOpenGLTextureBase* DestTexture = GetOpenGLTextureFromRHITexture(DestRHITexture);
diff --git a/Engine/Source/Runtime/OpenGLDrv/Public/OpenGLDrv.h b/Engine/Source/Runtime/OpenGLDrv/Public/OpenGLDrv.h
--- a/Engine/Source/Runtime/OpenGLDrv/Public/OpenGLDrv.h
+++ b/Engine/Source/Runtime/OpenGLDrv/Public/OpenGLDrv.h
@@ -406,6 +406,8 @@ public:
 	virtual FRenderQueryRHIRef RHICreateRenderQuery(ERenderQueryType QueryType) final override;
 	virtual bool RHIGetRenderQueryResult(FRHIRenderQuery* RenderQuery, uint64& OutResult, bool bWait, uint32 GPUIndex = INDEX_NONE) final override;
 	virtual FTexture2DRHIRef RHIGetViewportBackBuffer(FRHIViewport* Viewport) final override;
+	virtual void RHIAliasTextureResources(FTextureRHIRef& DestTextureRHI, FTextureRHIRef& SrcTextureRHI) final override;
+	virtual FTextureRHIRef RHICreateAliasedTexture(FTextureRHIRef& SourceTexture) final override;
 	UE_DEPRECATED(4.25, "RHIAliasTextureResources now takes references to FTextureRHIRef objects as parameters")
 	virtual void RHIAliasTextureResources(FRHITexture* DestTexture, FRHITexture* SrcTexture) final override;
 	UE_DEPRECATED(4.25, "RHICreateAliasedTexture now takes a reference to an FTextureRHIRef object")
