Remove duplicate SyncInterval introduced by merge causing Linux compile errors
#rnx

diff --git a/Engine/Source/Runtime/OpenGLDrv/Private/Linux/OpenGLLinux.cpp b/Engine/Source/Runtime/OpenGLDrv/Private/Linux/OpenGLLinux.cpp
--- a/Engine/Source/Runtime/OpenGLDrv/Private/Linux/OpenGLLinux.cpp
+++ b/Engine/Source/Runtime/OpenGLDrv/Private/Linux/OpenGLLinux.cpp
@@ -430,7 +430,6 @@ bool PlatformBlitToViewport(FPlatformOpenGLDevice* Device,
 							bool bPresent,
 							bool bLockToVsync)
 {
-	int32 SyncInterval = FGenericPlatformRHIFramePacer::GetFramePace();
 	FPlatformOpenGLContext* const Context = Viewport.GetGLContext();
 
 	check( Context && Context->hWnd );
