Temp Disable GpuSorting on Opengl

diff --git a/Engine/Source/Runtime/Engine/Private/Particles/FXSystem.cpp b/Engine/Source/Runtime/Engine/Private/Particles/FXSystem.cpp
--- a/Engine/Source/Runtime/Engine/Private/Particles/FXSystem.cpp
+++ b/Engine/Source/Runtime/Engine/Private/Particles/FXSystem.cpp
@@ -117,7 +117,7 @@ namespace FXConsoleVariables
 		TEXT("FX.AllowGPUSorting"),
 		bAllowGPUSorting,
 		TEXT("Allow particles to be sorted on the GPU."),
-		ECVF_Cheat
+		ECVF_ReadOnly
 		);
 	FAutoConsoleVariableRef CVarFreezeGPUSimulation(
 		TEXT("FX.FreezeGPUSimulation"),
diff --git a/Engine/Source/Runtime/OpenGLDrv/Private/OpenGLDevice.cpp b/Engine/Source/Runtime/OpenGLDrv/Private/OpenGLDevice.cpp
--- a/Engine/Source/Runtime/OpenGLDrv/Private/OpenGLDevice.cpp
+++ b/Engine/Source/Runtime/OpenGLDrv/Private/OpenGLDevice.cpp
@@ -1171,6 +1171,15 @@ FOpenGLDynamicRHI::FOpenGLDynamicRHI()
 		}
 	}
 
+	{
+		// Temp disable gpusorting for Opengl because of issues on Adreno and Mali devices
+		auto* CVar = IConsoleManager::Get().FindConsoleVariable(TEXT("FX.AllowGPUSorting"));
+		if (CVar)
+		{
+			CVar->Set(false);
+		}
+	}
+
 	PrivateOpenGLDevicePtr = this;
 	GlobalUniformBuffers.AddZeroed(FUniformBufferStaticSlotRegistry::Get().GetSlotCount());
 }
