Fixed: crash with a real UB path on OpenGL

diff --git a/Engine/Source/Runtime/OpenGLDrv/Private/OpenGLShaders.cpp b/Engine/Source/Runtime/OpenGLDrv/Private/OpenGLShaders.cpp
--- a/Engine/Source/Runtime/OpenGLDrv/Private/OpenGLShaders.cpp
+++ b/Engine/Source/Runtime/OpenGLDrv/Private/OpenGLShaders.cpp
@@ -2430,11 +2430,17 @@ void FOpenGLLinkedProgram::ConfigureShaderStage( int Stage, uint32 FirstUniformB
 		Name.Buffer[4] = 0;
 		Name.Buffer[5] = 0;
 		Name.Buffer[6] = 0;
-		for (uint8 UB = 0; UB < Config.Shaders[Stage].Bindings.NumUniformBuffers; ++UB)
+
+		int32 NumUniformBuffers = Config.Shaders[Stage].Bindings.NumUniformBuffers;
+		StagePackedUniformInfo[Stage].PackedUniformBufferInfos.SetNum(NumUniformBuffers);
+		int32 NumPackedUniformBuffers = Config.Shaders[Stage].Bindings.PackedUniformBuffers.Num();
+		check(NumPackedUniformBuffers <= NumUniformBuffers);
+
+		for (int32 UB = 0; UB < NumPackedUniformBuffers; ++UB)
 		{
 			const TArray<CrossCompiler::FPackedArrayInfo>& PackedInfo = Config.Shaders[Stage].Bindings.PackedUniformBuffers[UB];
-			
-			TArray<FPackedUniformInfo> PackedBuffers;
+			TArray<FPackedUniformInfo>& PackedBuffers = StagePackedUniformInfo[Stage].PackedUniformBufferInfos[UB];
+
 			ANSICHAR* Str = SetIndex(Name.Buffer, 2, UB);
 			*Str++ = '_';
 			Str[1] = 0;
@@ -2445,8 +2451,6 @@ void FOpenGLLinkedProgram::ConfigureShaderStage( int Stage, uint32 FirstUniformB
 				FPackedUniformInfo Info = {Location, PackedInfo[Index].TypeName,  PackedInfo[Index].TypeIndex};
 				PackedBuffers.Add(Info);
 			}
-
-			StagePackedUniformInfo[Stage].PackedUniformBufferInfos.Add(PackedBuffers);
 		}
 	}
 
