Fixed: Runtime Virtual Textures crash on Mobile

diff --git a/Engine/Source/Runtime/OpenGLDrv/Private/OpenGLES.cpp b/Engine/Source/Runtime/OpenGLDrv/Private/OpenGLES.cpp
--- a/Engine/Source/Runtime/OpenGLDrv/Private/OpenGLES.cpp
+++ b/Engine/Source/Runtime/OpenGLDrv/Private/OpenGLES.cpp
@@ -85,6 +85,7 @@ PFNGLCLEARBUFFERIVPROC					glClearBufferiv = NULL;
 PFNGLCLEARBUFFERUIVPROC					glClearBufferuiv = NULL;
 PFNGLREADBUFFERPROC						glReadBuffer = NULL;
 PFNGLDRAWBUFFERSPROC					glDrawBuffers = NULL;
+PFNGLCOLORMASKIEXTPROC					glColorMaskiEXT = NULL;
 PFNGLTEXBUFFEREXTPROC					glTexBufferEXT = NULL;
 PFNGLTEXBUFFERRANGEEXTPROC				glTexBufferRangeEXT = NULL;
 PFNGLCOPYIMAGESUBDATAPROC				glCopyImageSubData = nullptr;
@@ -537,6 +538,13 @@ void FOpenGLES::ProcessExtensions(const FString& ExtensionsString)
 		glDispatchCompute = (PFNGLDISPATCHCOMPUTEPROC)((void*)eglGetProcAddress("glDispatchCompute"));
 		glDispatchComputeIndirect = (PFNGLDISPATCHCOMPUTEINDIRECTPROC)((void*)eglGetProcAddress("glDispatchComputeIndirect"));
 		glBindImageTexture = (PFNGLBINDIMAGETEXTUREPROC)((void*)eglGetProcAddress("glBindImageTexture"));
+
+		// ES 3.2
+		glColorMaskiEXT = (PFNGLCOLORMASKIEXTPROC)((void*)eglGetProcAddress("glColorMaski"));
+		if (glColorMaskiEXT == nullptr)
+		{
+			glColorMaskiEXT = (PFNGLCOLORMASKIEXTPROC)((void*)eglGetProcAddress("glColorMaskiEXT"));
+		}
 	}
 
 	// test for glCopyImageSubData functionality
diff --git a/Engine/Source/Runtime/OpenGLDrv/Private/OpenGLUAV.cpp b/Engine/Source/Runtime/OpenGLDrv/Private/OpenGLUAV.cpp
--- a/Engine/Source/Runtime/OpenGLDrv/Private/OpenGLUAV.cpp
+++ b/Engine/Source/Runtime/OpenGLDrv/Private/OpenGLUAV.cpp
@@ -5,6 +5,7 @@
 #include "RenderUtils.h"
 #include "OpenGLDrv.h"
 #include "OpenGLDrvPrivate.h"
+#include "ClearReplacementShaders.h"
 
 FShaderResourceViewRHIRef FOpenGLDynamicRHI::RHICreateShaderResourceView(FRHIVertexBuffer* VertexBufferRHI, uint32 Stride, uint8 Format)
 {
@@ -211,6 +212,7 @@ FOpenGLTextureUnorderedAccessView::FOpenGLTextureUnorderedAccessView(FRHITexture
 
 	this->Resource = Texture->Resource;
 	this->Format = GLFormat.InternalFormat[0];
+	this->UnrealFormat = TextureRHI->GetFormat();
 	this->bLayered = (Texture->Target == GL_TEXTURE_3D);
 }
 
@@ -239,6 +241,8 @@ FOpenGLVertexBufferUnorderedAccessView::FOpenGLVertexBufferUnorderedAccessView(
 	this->Resource = TextureID;
 	this->BufferResource = InVertexBuffer->Resource;
 	this->Format = GLFormat.InternalFormat[0];
+	this->UnrealFormat = Format;
+	
 }
 
 uint32 FOpenGLVertexBufferUnorderedAccessView::GetBufferSize()
@@ -279,25 +283,68 @@ void FOpenGLDynamicRHI::RHIClearUAVFloat(FRHIUnorderedAccessView* UnorderedAcces
 	FOpenGLUnorderedAccessView* Texture = ResourceCast(UnorderedAccessViewRHI);
 
 #if OPENGL_GL4 || PLATFORM_LUMINGL4
-	glBindBuffer(GL_TEXTURE_BUFFER, Texture->BufferResource);
-	FOpenGL::ClearBufferData(GL_TEXTURE_BUFFER, Texture->Format, GL_RGBA_INTEGER, GL_FLOAT, reinterpret_cast<const uint32*>(&Values));
-	GPUProfilingData.RegisterGPUWork(1);
-#else
-	UE_LOG(LogRHI, Fatal, TEXT("Only OpenGL4 supports RHIClearUAVFloat."));
+	if (GMaxRHIFeatureLevel >= ERHIFeatureLevel::SM5)
+	{
+		glBindBuffer(GL_TEXTURE_BUFFER, Texture->BufferResource);
+		FOpenGL::ClearBufferData(GL_TEXTURE_BUFFER, Texture->Format, GL_RGBA_INTEGER, GL_FLOAT, reinterpret_cast<const uint32*>(&Values));
+		GPUProfilingData.RegisterGPUWork(1);
+		return;
+	}
 #endif
+	// Use compute on ES3.1
+	check(Texture->BufferResource);
+	{
+		TRHICommandList_RecursiveHazardous<FOpenGLDynamicRHI> RHICmdList(this);
+		
+		int32 NumComponents = GPixelFormats[Texture->UnrealFormat].NumComponents;
+		uint32 NumElements = Texture->GetBufferSize() / GPixelFormats[Texture->UnrealFormat].BlockBytes;
+					
+		switch (NumComponents)
+		{
+		case 1:
+			ClearUAVShader_T<EClearReplacementResourceType::Buffer, EClearReplacementValueType::Float, 1, false>(RHICmdList, UnorderedAccessViewRHI, 1, 1, 1, *reinterpret_cast<const float(*)[1]>(&Values));
+			break;
+		case 4:
+			ClearUAVShader_T<EClearReplacementResourceType::Buffer, EClearReplacementValueType::Float, 4, false>(RHICmdList, UnorderedAccessViewRHI, 4, 1, 1, *reinterpret_cast<const float(*)[4]>(&Values));
+			break;
+		default:
+			check(false);
+		};
+	}
 }
 
 void FOpenGLDynamicRHI::RHIClearUAVUint(FRHIUnorderedAccessView* UnorderedAccessViewRHI, const FUintVector4& Values)
 {
 	FOpenGLUnorderedAccessView* Texture = ResourceCast(UnorderedAccessViewRHI);
-
 #if OPENGL_GL4 || PLATFORM_LUMINGL4
-	glBindBuffer(GL_TEXTURE_BUFFER, Texture->BufferResource);
-	FOpenGL::ClearBufferData(GL_TEXTURE_BUFFER, Texture->Format, GL_RGBA_INTEGER, GL_UNSIGNED_INT, reinterpret_cast<const uint32*>(&Values));
-	GPUProfilingData.RegisterGPUWork(1);
-#else
-	UE_LOG(LogRHI, Fatal, TEXT("Only OpenGL4 supports RHIClearUAVUint."));
+	if (GMaxRHIFeatureLevel >= ERHIFeatureLevel::SM5)
+	{
+		glBindBuffer(GL_TEXTURE_BUFFER, Texture->BufferResource);
+		FOpenGL::ClearBufferData(GL_TEXTURE_BUFFER, Texture->Format, GL_RGBA_INTEGER, GL_UNSIGNED_INT, reinterpret_cast<const uint32*>(&Values));
+		GPUProfilingData.RegisterGPUWork(1);
+		return;
+	}
 #endif
+	// Use compute on ES3.1
+	check(Texture->BufferResource);
+	{
+		TRHICommandList_RecursiveHazardous<FOpenGLDynamicRHI> RHICmdList(this);
+
+		int32 NumComponents = GPixelFormats[Texture->UnrealFormat].NumComponents;
+		uint32 NumElements = Texture->GetBufferSize() / GPixelFormats[Texture->UnrealFormat].BlockBytes;
+				
+		switch (NumComponents)
+		{
+		case 1:
+			ClearUAVShader_T<EClearReplacementResourceType::Buffer, EClearReplacementValueType::Uint32, 1, false>(RHICmdList, UnorderedAccessViewRHI, 1, 1, 1, *reinterpret_cast<const uint32(*)[1]>(&Values));
+			break;
+		case 4:
+			ClearUAVShader_T<EClearReplacementResourceType::Buffer, EClearReplacementValueType::Uint32, 4, false>(RHICmdList, UnorderedAccessViewRHI, 4, 1, 1, *reinterpret_cast<const uint32(*)[4]>(&Values));
+			break;
+		default:
+			check(false);
+		};
+	}
 }
 
 FOpenGLStructuredBufferUnorderedAccessView::FOpenGLStructuredBufferUnorderedAccessView(FOpenGLDynamicRHI* InOpenGLRHI, FRHIStructuredBuffer* InStructuredBufferRHI, uint8 InFormat)
@@ -322,6 +369,7 @@ FOpenGLStructuredBufferUnorderedAccessView::FOpenGLStructuredBufferUnorderedAcce
 	this->Resource = TextureID;
 	this->BufferResource = InStructuredBuffer->Resource;
 	this->Format = GLFormat.InternalFormat[0];
+	this->UnrealFormat = Format;
 }
 
 uint32 FOpenGLStructuredBufferUnorderedAccessView::GetBufferSize()
diff --git a/Engine/Source/Runtime/OpenGLDrv/Public/OpenGLES.h b/Engine/Source/Runtime/OpenGLDrv/Public/OpenGLES.h
--- a/Engine/Source/Runtime/OpenGLDrv/Public/OpenGLES.h
+++ b/Engine/Source/Runtime/OpenGLDrv/Public/OpenGLES.h
@@ -498,6 +498,7 @@ typedef void (GL_APIENTRYP PFNGLPOPGROUPMARKEREXTPROC) (void);
 typedef void (GL_APIENTRYP PFNGLFRAMEBUFFERTEXTURE2DMULTISAMPLEEXTPROC) (GLenum target, GLenum attachment, GLenum textarget, GLuint texture, GLint level, GLsizei samples);
 typedef void (GL_APIENTRYP PFNGLRENDERBUFFERSTORAGEMULTISAMPLEEXTPROC) (GLenum target, GLsizei samples, GLenum internalformat, GLsizei width, GLsizei height);
 typedef void (GL_APIENTRYP PFNGLFRAMEBUFFERTEXTURELAYERPROC) (GLenum target, GLenum attachment, GLuint texture, GLint level, GLint layer);
+typedef void (GL_APIENTRYP PFNGLCOLORMASKIEXTPROC) (GLuint index, GLboolean r, GLboolean g, GLboolean b, GLboolean a);
 
 /** from ES 3.0 but can be called on certain Adreno devices */
 typedef void (GL_APIENTRYP PFNGLTEXSTORAGE2DPROC) (GLenum target, GLsizei levels, GLenum internalformat, GLsizei width, GLsizei height);
@@ -560,6 +561,7 @@ extern PFNGLCLEARBUFFERIVPROC			glClearBufferiv;
 extern PFNGLCLEARBUFFERUIVPROC			glClearBufferuiv;
 extern PFNGLREADBUFFERPROC				glReadBuffer;
 extern PFNGLDRAWBUFFERSPROC				glDrawBuffers;
+extern PFNGLCOLORMASKIEXTPROC			glColorMaskiEXT;
 extern PFNGLTEXIMAGE3DPROC				glTexImage3D;
 extern PFNGLTEXSUBIMAGE3DPROC			glTexSubImage3D;
 extern PFNGLCOMPRESSEDTEXIMAGE3DPROC    glCompressedTexImage3D;
@@ -882,8 +884,15 @@ struct FOpenGLES : public FOpenGLBase
 
 	static FORCEINLINE void ColorMaskIndexed(GLuint Index, GLboolean Red, GLboolean Green, GLboolean Blue, GLboolean Alpha)
 	{
-		check(Index == 0 || SupportsMultipleRenderTargets());
-		glColorMask(Red, Green, Blue, Alpha);
+		if (Index > 0 && glColorMaskiEXT)
+		{
+			// ext or OpenGL ES 3.2
+			glColorMaskiEXT(Index, Red, Green, Blue, Alpha);
+		}
+		else
+		{
+			glColorMask(Red, Green, Blue, Alpha);
+		}
 	}
 
 	static FORCEINLINE void TexBuffer(GLenum Target, GLenum InternalFormat, GLuint Buffer)
diff --git a/Engine/Source/Runtime/OpenGLDrv/Public/OpenGLResources.h b/Engine/Source/Runtime/OpenGLDrv/Public/OpenGLResources.h
--- a/Engine/Source/Runtime/OpenGLDrv/Public/OpenGLResources.h
+++ b/Engine/Source/Runtime/OpenGLDrv/Public/OpenGLResources.h
@@ -1848,7 +1848,8 @@ public:
 	FOpenGLUnorderedAccessView():
 		Resource(0),
 		BufferResource(0),
-		Format(0)
+		Format(0),
+		UnrealFormat(0)
 	{
 
 	}
@@ -1856,6 +1857,7 @@ public:
 	GLuint	Resource;
 	GLuint	BufferResource;
 	GLenum	Format;
+	uint8	UnrealFormat;
 
 	virtual uint32 GetBufferSize()
 	{
diff --git a/Engine/Source/Runtime/RenderCore/Private/ClearReplacementShaders.cpp b/Engine/Source/Runtime/RenderCore/Private/ClearReplacementShaders.cpp
--- a/Engine/Source/Runtime/RenderCore/Private/ClearReplacementShaders.cpp
+++ b/Engine/Source/Runtime/RenderCore/Private/ClearReplacementShaders.cpp
@@ -14,6 +14,7 @@ IMPLEMENT_SHADER_TYPE4_WITH_TEMPLATE_PREFIX(template<>, RENDERCORE_API, FClearRe
 IMPLEMENT_SHADER_TYPE4_WITH_TEMPLATE_PREFIX(template<>, RENDERCORE_API, FClearReplacementCS_Buffer_Uint_Bounds          , SF_Compute);
 IMPLEMENT_SHADER_TYPE4_WITH_TEMPLATE_PREFIX(template<>, RENDERCORE_API, FClearReplacementCS_Buffer_Uint4_Bounds         , SF_Compute);
 IMPLEMENT_SHADER_TYPE4_WITH_TEMPLATE_PREFIX(template<>, RENDERCORE_API, FClearReplacementCS_Buffer_Uint4         		, SF_Compute);
+IMPLEMENT_SHADER_TYPE4_WITH_TEMPLATE_PREFIX(template<>, RENDERCORE_API, FClearReplacementCS_Buffer_Float_Bounds         , SF_Compute);
 IMPLEMENT_SHADER_TYPE4_WITH_TEMPLATE_PREFIX(template<>, RENDERCORE_API, FClearReplacementCS_Buffer_Float4_Bounds        , SF_Compute);
 IMPLEMENT_SHADER_TYPE4_WITH_TEMPLATE_PREFIX(template<>, RENDERCORE_API, FClearReplacementCS_Texture3D_Float4_Bounds     , SF_Compute);
 IMPLEMENT_SHADER_TYPE4_WITH_TEMPLATE_PREFIX(template<>, RENDERCORE_API, FClearReplacementCS_Texture2DArray_Float4_Bounds, SF_Compute);
diff --git a/Engine/Source/Runtime/RenderCore/Public/ClearReplacementShaders.h b/Engine/Source/Runtime/RenderCore/Public/ClearReplacementShaders.h
--- a/Engine/Source/Runtime/RenderCore/Public/ClearReplacementShaders.h
+++ b/Engine/Source/Runtime/RenderCore/Public/ClearReplacementShaders.h
@@ -253,7 +253,8 @@ typedef TClearReplacementBase<EClearReplacementValueType::Float , 4, true , fals
 typedef TClearReplacementBase<EClearReplacementValueType::Float , 4, true , true > FClearReplacementBase_Float4_Zero_Bounds;
 typedef TClearReplacementBase<EClearReplacementValueType::Uint32, 1, false, true > FClearReplacementBase_Uint_Bounds;
 typedef TClearReplacementBase<EClearReplacementValueType::Uint32, 4, false, true > FClearReplacementBase_Uint4_Bounds;
-typedef TClearReplacementBase<EClearReplacementValueType:: Int32, 4, false, true > FClearReplacementBase_Sint4_Bounds;
+typedef TClearReplacementBase<EClearReplacementValueType::Int32,  4, false, true > FClearReplacementBase_Sint4_Bounds;
+typedef TClearReplacementBase<EClearReplacementValueType::Float , 1, false, true > FClearReplacementBase_Float_Bounds;
 typedef TClearReplacementBase<EClearReplacementValueType::Float , 4, false, true > FClearReplacementBase_Float4_Bounds;
 
 // Simple vertex shaders for generating screen quads. Optionally with a min/max bounds in NDC space, and depth value.
@@ -269,6 +270,7 @@ typedef TClearReplacementPS<false, FClearReplacementBase_Float4_Zero>        FCl
 	
 // Compute shaders for clearing each resource type, with a min/max bounds enabled.
 typedef TClearReplacementCS<EClearReplacementResourceType::Buffer,         FClearReplacementBase_Uint_Bounds>   FClearReplacementCS_Buffer_Uint_Bounds;
+typedef TClearReplacementCS<EClearReplacementResourceType::Buffer,         FClearReplacementBase_Float_Bounds> FClearReplacementCS_Buffer_Float_Bounds;
 typedef TClearReplacementCS<EClearReplacementResourceType::Buffer,         FClearReplacementBase_Float4_Bounds> FClearReplacementCS_Buffer_Float4_Bounds;
 typedef TClearReplacementCS<EClearReplacementResourceType::Texture3D,      FClearReplacementBase_Float4_Bounds> FClearReplacementCS_Texture3D_Float4_Bounds;
 typedef TClearReplacementCS<EClearReplacementResourceType::Texture2D,      FClearReplacementBase_Float4_Bounds> FClearReplacementCS_Texture2D_Float4_Bounds;
diff --git a/Engine/Source/Runtime/Renderer/Private/MobileShadingRenderer.cpp b/Engine/Source/Runtime/Renderer/Private/MobileShadingRenderer.cpp
--- a/Engine/Source/Runtime/Renderer/Private/MobileShadingRenderer.cpp
+++ b/Engine/Source/Runtime/Renderer/Private/MobileShadingRenderer.cpp
@@ -468,11 +468,6 @@ void FMobileSceneRenderer::Render(FRHICommandListImmediate& RHICmdList)
 
 	// Find the visible primitives.
 	InitViews(RHICmdList);
-
-	if (bUseVirtualTexturing)
-	{
-		FVirtualTextureSystem::Get().Update(RHICmdList, ViewFeatureLevel, Scene);
-	}
 	
 	if (GRHINeedsExtraDeletionLatency || !GRHICommandList.Bypass())
 	{
@@ -494,13 +489,17 @@ void FMobileSceneRenderer::Render(FRHICommandListImmediate& RHICmdList)
 	RHICmdList.SetCurrentStat(GET_STATID(STAT_CLMM_SceneSim));
 
 	// Generate the Sky/Atmosphere look up tables
-	// Do compute work first, before any graphics work
 	const bool bShouldRenderSkyAtmosphere = ShouldRenderSkyAtmosphere(Scene, ViewFamily.EngineShowFlags);
 	if (bShouldRenderSkyAtmosphere)
 	{
 		RenderSkyAtmosphereLookUpTables(RHICmdList);
 	}
 
+	if (bUseVirtualTexturing)
+	{
+		FVirtualTextureSystem::Get().Update(RHICmdList, ViewFeatureLevel, Scene);
+	}
+
 	// Notify the FX system that the scene is about to be rendered.
 	if (Scene->FXSystem && ViewFamily.EngineShowFlags.Particles)
 	{
