Fix a bug of clear render targets unnecessarily for pc on OpenglRHI.

diff --git a/Engine/Source/Runtime/OpenGLDrv/Private/OpenGLCommands.cpp b/Engine/Source/Runtime/OpenGLDrv/Private/OpenGLCommands.cpp
--- a/Engine/Source/Runtime/OpenGLDrv/Private/OpenGLCommands.cpp
+++ b/Engine/Source/Runtime/OpenGLDrv/Private/OpenGLCommands.cpp
@@ -1682,8 +1682,11 @@ void FOpenGLDynamicRHI::RHISetRenderTargetsAndClear(const FRHISetRenderTargetsIn
 		&RenderTargetsInfo.DepthStencilRenderTarget);
 
 	/**
-	 * Convert all load action from NoAction to Clear for OpenGL platform, especially for mobile to avoid an unnecessary load action.
+	 * Convert all load action from NoAction to Clear for tiled GPU on OpenGL platform to avoid an unnecessary load action.
 	 */
+
+	bool bIsTiledGPU = RHIHasTiledGPU(GetFeatureLevelShaderPlatform(FOpenGL::GetFeatureLevel()));
+
 	bool bClearColor = RenderTargetsInfo.bClearColor;
 	bool bClearStencil = RenderTargetsInfo.bClearStencil;
 	bool bClearDepth = RenderTargetsInfo.bClearDepth;
@@ -1696,24 +1699,42 @@ void FOpenGLDynamicRHI::RHISetRenderTargetsAndClear(const FRHISetRenderTargetsIn
 	{
 		if (RenderTargetsInfo.ColorRenderTarget[i].Texture != nullptr)
 		{
-			bClearColor |= RenderTargetsInfo.ColorRenderTarget[i].LoadAction == ERenderTargetLoadAction::ENoAction;
-
 			const FClearValueBinding& ClearValue = RenderTargetsInfo.ColorRenderTarget[i].Texture->GetClearBinding();
-			
-			ClearColors[i] = ClearValue.ColorBinding == EClearBinding::EColorBound ? ClearValue.GetClearColor() : FLinearColor::Black;
+
+			if (bIsTiledGPU)
+			{
+				bClearColor |= RenderTargetsInfo.ColorRenderTarget[i].LoadAction == ERenderTargetLoadAction::ENoAction;
+
+				ClearColors[i] = ClearValue.ColorBinding == EClearBinding::EColorBound ? ClearValue.GetClearColor() : FLinearColor::Black;
+			}
+			else if(bClearColor)
+			{
+				checkf(ClearValue.ColorBinding == EClearBinding::EColorBound, TEXT("Texture: %s does not have a color bound for fast clears"), *RenderTargetsInfo.ColorRenderTarget[i].Texture->GetName().GetPlainNameString());
+
+				ClearColors[i] = ClearValue.GetClearColor();
+			}
 		}
 	}
 
 	if (RenderTargetsInfo.DepthStencilRenderTarget.Texture != nullptr)
 	{
-		bClearStencil |= RenderTargetsInfo.DepthStencilRenderTarget.StencilLoadAction == ERenderTargetLoadAction::ENoAction;
+		const FClearValueBinding& ClearValue = RenderTargetsInfo.DepthStencilRenderTarget.Texture->GetClearBinding();
 
-		bClearDepth |= RenderTargetsInfo.DepthStencilRenderTarget.DepthLoadAction == ERenderTargetLoadAction::ENoAction;
+		if (bIsTiledGPU)
+		{
+			bClearStencil |= RenderTargetsInfo.DepthStencilRenderTarget.StencilLoadAction == ERenderTargetLoadAction::ENoAction;
 
-		const FClearValueBinding& ClearValue = RenderTargetsInfo.DepthStencilRenderTarget.Texture->GetClearBinding();
+			bClearDepth |= RenderTargetsInfo.DepthStencilRenderTarget.DepthLoadAction == ERenderTargetLoadAction::ENoAction;
 
-		if (ClearValue.ColorBinding == EClearBinding::EDepthStencilBound)
+			if (ClearValue.ColorBinding == EClearBinding::EDepthStencilBound)
+			{
+				ClearValue.GetDepthStencil(DepthClear, StencilClear);
+			}
+		}
+		else if (bClearDepth || bClearStencil)
 		{
+			checkf(ClearValue.ColorBinding == EClearBinding::EDepthStencilBound, TEXT("Texture: %s does not have a DS value bound for fast clears"), *RenderTargetsInfo.DepthStencilRenderTarget.Texture->GetName().GetPlainNameString());
+
 			ClearValue.GetDepthStencil(DepthClear, StencilClear);
 		}
 	}
