Disable support for QCOM_shader_framebuffer_fetch_noncoherent extension on Adreno5xx GPUs, as it causes rendering artifacts with MSAA enabled

diff --git a/Engine/Config/BaseDeviceProfiles.ini b/Engine/Config/BaseDeviceProfiles.ini
--- a/Engine/Config/BaseDeviceProfiles.ini
+++ b/Engine/Config/BaseDeviceProfiles.ini
@@ -535,10 +535,14 @@ BaseProfileName=Android_Low
 [Android_Adreno5xx_Low DeviceProfile]
 DeviceType=Android
 BaseProfileName=Android_Low
+; disable QCOM_shader_framebuffer_fetch_noncoherent on Adreno5xx, as it causes rendering artifacts with enabled MSAA
++CVars=r.Android.DisableFBFNonCoherent=1
 
 [Android_Adreno5xx DeviceProfile]
 DeviceType=Android
 BaseProfileName=Android_Mid
+; disable QCOM_shader_framebuffer_fetch_noncoherent on Adreno5xx, as it causes rendering artifacts with enabled MSAA
++CVars=r.Android.DisableFBFNonCoherent=1
 
 [Android_Adreno6xx DeviceProfile]
 DeviceType=Android
diff --git a/Engine/Source/Runtime/OpenGLDrv/Private/Android/AndroidOpenGL.cpp b/Engine/Source/Runtime/OpenGLDrv/Private/Android/AndroidOpenGL.cpp
--- a/Engine/Source/Runtime/OpenGLDrv/Private/Android/AndroidOpenGL.cpp
+++ b/Engine/Source/Runtime/OpenGLDrv/Private/Android/AndroidOpenGL.cpp
@@ -45,6 +45,12 @@ static TAutoConsoleVariable<int32> CVarDisableEarlyFragmentTests(
 	TEXT("Whether to disable early_fragment_tests if any \n"),
 	ECVF_ReadOnly);
 
+static TAutoConsoleVariable<int32> CVarDisableFBFNonCoherent(
+	TEXT("r.Android.DisableFBFNonCoherent"),
+	0,
+	TEXT("Whether to disable usage of QCOM_shader_framebuffer_fetch_noncoherent extension\n"),
+	ECVF_ReadOnly);
+
 struct FPlatformOpenGLDevice
 {
 
@@ -754,7 +760,8 @@ extern FString AndroidThunkCpp_GetMetaDataString(const FString& Key);
 void FAndroidOpenGL::SetupDefaultGLContextState(const FString& ExtensionsString)
 {
 	// Enable QCOM non-coherent framebuffer fetch if supported
-	if (ExtensionsString.Contains(TEXT("GL_QCOM_shader_framebuffer_fetch_noncoherent")) && 
+	if (CVarDisableFBFNonCoherent.GetValueOnAnyThread() == 0 &&
+		ExtensionsString.Contains(TEXT("GL_QCOM_shader_framebuffer_fetch_noncoherent")) && 
 		ExtensionsString.Contains(TEXT("GL_EXT_shader_framebuffer_fetch")))
 	{
 		glEnable(GL_FRAMEBUFFER_FETCH_NONCOHERENT_QCOM);
@@ -997,7 +1004,8 @@ void FAndroidOpenGL::ProcessExtensions(const FString& ExtensionsString)
 	}
 
 	// Qualcomm non-coherent framebuffer_fetch
-	if (ExtensionsString.Contains(TEXT("GL_QCOM_shader_framebuffer_fetch_noncoherent")) && 
+	if (CVarDisableFBFNonCoherent.GetValueOnAnyThread() == 0 &&
+		ExtensionsString.Contains(TEXT("GL_QCOM_shader_framebuffer_fetch_noncoherent")) && 
 		ExtensionsString.Contains(TEXT("GL_EXT_shader_framebuffer_fetch")))
 	{
 		glFramebufferFetchBarrierQCOM = (PFNGLFRAMEBUFFERFETCHBARRIERQCOMPROC)((void*)eglGetProcAddress("glFramebufferFetchBarrierQCOM"));
