Add GL RHI support for creating VB/IBs WithoutNativeResources and switching of VB/IB underlying resources.
required to support mesh lod streaming.

diff --git a/Engine/Source/Runtime/Core/Public/Android/AndroidPlatformProperties.h b/Engine/Source/Runtime/Core/Public/Android/AndroidPlatformProperties.h
--- a/Engine/Source/Runtime/Core/Public/Android/AndroidPlatformProperties.h
+++ b/Engine/Source/Runtime/Core/Public/Android/AndroidPlatformProperties.h
@@ -117,6 +117,11 @@ struct FAndroidPlatformProperties
 	{
 		return true;
 	}
+
+	static FORCEINLINE bool SupportsMeshLODStreaming()
+	{
+		return true;
+	}
 };
 
 #ifdef PROPERTY_HEADER_SHOULD_DEFINE_TYPE
diff --git a/Engine/Source/Runtime/OpenGLDrv/Private/OpenGLIndexBuffer.cpp b/Engine/Source/Runtime/OpenGLDrv/Private/OpenGLIndexBuffer.cpp
--- a/Engine/Source/Runtime/OpenGLDrv/Private/OpenGLIndexBuffer.cpp
+++ b/Engine/Source/Runtime/OpenGLDrv/Private/OpenGLIndexBuffer.cpp
@@ -10,6 +10,11 @@
 
 FIndexBufferRHIRef FOpenGLDynamicRHI::RHICreateIndexBuffer(uint32 Stride,uint32 Size, uint32 InUsage, FRHIResourceCreateInfo& CreateInfo)
 {
+	if (CreateInfo.bWithoutNativeResource)
+	{
+		return new FOpenGLIndexBuffer();
+	}
+
 	const void *Data = NULL;
 
 	// If a resource array was provided for the resource, create the resource pre-populated
@@ -51,3 +56,20 @@ void FOpenGLDynamicRHI::UnlockIndexBuffer_BottomOfPipe(FRHICommandListImmediate&
 	IndexBuffer->Unlock();
 	RHITHREAD_GLCOMMAND_EPILOGUE();
 }
+
+void FOpenGLDynamicRHI::RHITransferIndexBufferUnderlyingResource(FRHIIndexBuffer* DestIndexBuffer, FRHIIndexBuffer* SrcIndexBuffer)
+{
+	VERIFY_GL_SCOPE();
+	check(DestIndexBuffer);
+	FOpenGLIndexBuffer* Dest = ResourceCast(DestIndexBuffer);
+	if (!SrcIndexBuffer)
+	{
+		TRefCountPtr<FOpenGLIndexBuffer> Src = new FOpenGLIndexBuffer();
+		Dest->Swap(*Src);
+	}
+	else
+	{
+		FOpenGLIndexBuffer* Src = ResourceCast(SrcIndexBuffer);
+		Dest->Swap(*Src);
+	}
+}
diff --git a/Engine/Source/Runtime/OpenGLDrv/Private/OpenGLUAV.cpp b/Engine/Source/Runtime/OpenGLDrv/Private/OpenGLUAV.cpp
--- a/Engine/Source/Runtime/OpenGLDrv/Private/OpenGLUAV.cpp
+++ b/Engine/Source/Runtime/OpenGLDrv/Private/OpenGLUAV.cpp
@@ -51,23 +51,21 @@ FShaderResourceViewRHIRef FOpenGLDynamicRHI::RHICreateShaderResourceView(const F
 				GLuint TextureID = 0;
 				if (FOpenGL::SupportsResourceView())
 				{
+					FOpenGL::GenTextures(1, &TextureID);
 					UE_CLOG(!GPixelFormats[Format].Supported, LogRHI, Error, TEXT("Unsupported EPixelFormat %d"), Format);
+					if (VertexBufferRHI)
+					{
+						FOpenGLVertexBuffer* VertexBuffer = FOpenGLDynamicRHI::ResourceCast(VertexBufferRHI);
 
-					FOpenGLVertexBuffer* VertexBuffer = FOpenGLDynamicRHI::ResourceCast(VertexBufferRHI);
+						const uint32 FormatBPP = GPixelFormats[Format].BlockBytes;
 
-					const uint32 FormatBPP = GPixelFormats[Format].BlockBytes;
+						const FOpenGLTextureFormat& GLFormat = GOpenGLTextureFormats[Format];
 
-					const FOpenGLTextureFormat& GLFormat = GOpenGLTextureFormats[Format];
-					FOpenGL::GenTextures(1, &TextureID);
-
-					// Use a texture stage that's not likely to be used for draws, to avoid waiting
-					if (VertexBuffer)
-					{
+						// Use a texture stage that's not likely to be used for draws, to avoid waiting
 						OGLRHI->CachedSetupTextureStage(OGLRHI->GetContextStateForCurrentContext(), FOpenGL::GetMaxCombinedTextureImageUnits() - 1, GL_TEXTURE_BUFFER, TextureID, -1, 1);
 						BindGLTexBufferRange(GL_TEXTURE_BUFFER, GLFormat.InternalFormat[0], VertexBuffer->Resource, Desc.StartOffsetBytes, Desc.NumElements, FormatBPP);
 					}
 				}
-
 				// No need to restore texture stage; leave it like this,
 				// and the next draw will take care of cleaning it up; or
 				// next operation that needs the stage will switch something else in on it.
@@ -110,15 +108,17 @@ FShaderResourceViewRHIRef FOpenGLDynamicRHI::RHICreateShaderResourceView(const F
 				GLuint TextureID = 0;
 				if (FOpenGL::SupportsResourceView())
 				{
-					FOpenGLIndexBuffer* IndexBuffer = ResourceCast(IndexBufferRHI);
 					FOpenGL::GenTextures(1, &TextureID);
-					CachedSetupTextureStage(GetContextStateForCurrentContext(), FOpenGL::GetMaxCombinedTextureImageUnits() - 1, GL_TEXTURE_BUFFER, TextureID, -1, 1);
-					uint32 Stride = IndexBufferRHI->GetStride();
-					GLenum Format = (Stride == 2) ? GL_R16UI : GL_R32UI;
-					BindGLTexBufferRange(GL_TEXTURE_BUFFER, Format, IndexBuffer->Resource, Desc.StartOffsetBytes, Desc.NumElements, Stride);
+					if (IndexBufferRHI)
+					{
+						FOpenGLIndexBuffer* IndexBuffer = ResourceCast(IndexBufferRHI);
+						CachedSetupTextureStage(GetContextStateForCurrentContext(), FOpenGL::GetMaxCombinedTextureImageUnits() - 1, GL_TEXTURE_BUFFER, TextureID, -1, 1);
+						uint32 Stride = IndexBufferRHI->GetStride();
+						GLenum Format = (Stride == 2) ? GL_R16UI : GL_R32UI;
+						BindGLTexBufferRange(GL_TEXTURE_BUFFER, Format, IndexBuffer->Resource, Desc.StartOffsetBytes, Desc.NumElements, Stride);
+					}
 				}
-
-				return new FOpenGLShaderResourceView(this, TextureID, GL_TEXTURE_BUFFER);
+				return new FOpenGLShaderResourceView(this, TextureID, GL_TEXTURE_BUFFER, IndexBufferRHI);
 			});
 		}
 
@@ -130,18 +130,56 @@ FShaderResourceViewRHIRef FOpenGLDynamicRHI::RHICreateShaderResourceView(const F
 	}
 }
 
+void FOpenGLDynamicRHI::RHIUpdateShaderResourceView(FRHIShaderResourceView* SRV, FRHIIndexBuffer* IndexBuffer)
+{
+	if (!FOpenGL::SupportsResourceView())
+	{
+		return;
+	}
+	VERIFY_GL_SCOPE();
+
+	FOpenGLShaderResourceView* SRVGL = FOpenGLDynamicRHI::ResourceCast(SRV);
+	FOpenGLIndexBuffer* IBGL = FOpenGLDynamicRHI::ResourceCast(IndexBuffer);
+
+	check(SRVGL);
+	check(!SRVGL->VertexBuffer);
+	GLuint TextureID = SRVGL->Resource;
+	CachedSetupTextureStage(GetContextStateForCurrentContext(), FOpenGL::GetMaxCombinedTextureImageUnits() - 1, GL_TEXTURE_BUFFER, TextureID, -1, 1);
+
+	if (!IBGL)
+	{
+		FOpenGL::TexBuffer(GL_TEXTURE_BUFFER, GL_R16UI, 0); // format ignored here since we're detaching.
+		SRVGL->IndexBuffer = nullptr;
+		SRVGL->ModificationVersion = 0;
+	}
+	else
+	{
+		uint32 Stride = IndexBuffer->GetStride();
+		GLenum Format = (Stride == 2) ? GL_R16UI : GL_R32UI;
+		const FOpenGLTextureFormat& GLFormat = GOpenGLTextureFormats[Format];
+		check(SRVGL->Target == GL_TEXTURE_BUFFER);
+
+		uint32 NumElements = IndexBuffer->GetSize() / Stride;
+		BindGLTexBufferRange(GL_TEXTURE_BUFFER, Format, IBGL->Resource, 0, NumElements, Stride);
+		SRVGL->IndexBuffer = IndexBuffer;
+		SRVGL->ModificationVersion = IBGL->ModificationCount;
+	}
+}
+
 void FOpenGLDynamicRHI::RHIUpdateShaderResourceView(FRHIShaderResourceView* SRV, FRHIVertexBuffer* VertexBuffer, uint32 Stride, uint8 Format)
 {
 	if (!FOpenGL::SupportsResourceView())
 	{
 		return;
 	}
-	
+	VERIFY_GL_SCOPE();
+
 	FOpenGLShaderResourceView* SRVGL = FOpenGLDynamicRHI::ResourceCast(SRV);
 	FOpenGLVertexBuffer* VBGL = FOpenGLDynamicRHI::ResourceCast(VertexBuffer);
 	const FOpenGLTextureFormat& GLFormat = GOpenGLTextureFormats[Format];
 	
 	check(SRVGL);
+	check(!SRVGL->IndexBuffer);
 	GLuint TextureID = SRVGL->Resource;
 	CachedSetupTextureStage(GetContextStateForCurrentContext(), FOpenGL::GetMaxCombinedTextureImageUnits() - 1, GL_TEXTURE_BUFFER, TextureID, -1, 1);
 	
diff --git a/Engine/Source/Runtime/OpenGLDrv/Private/OpenGLVertexBuffer.cpp b/Engine/Source/Runtime/OpenGLDrv/Private/OpenGLVertexBuffer.cpp
--- a/Engine/Source/Runtime/OpenGLDrv/Private/OpenGLVertexBuffer.cpp
+++ b/Engine/Source/Runtime/OpenGLDrv/Private/OpenGLVertexBuffer.cpp
@@ -173,6 +173,11 @@ void BeginFrame_VertexBufferCleanup()
 
 FVertexBufferRHIRef FOpenGLDynamicRHI::RHICreateVertexBuffer(uint32 Size, uint32 InUsage, FRHIResourceCreateInfo& CreateInfo)
 {
+	if (CreateInfo.bWithoutNativeResource)
+	{
+		return new FOpenGLVertexBuffer();
+	}
+
 	const void *Data = NULL;
 
 	// If a resource array was provided for the resource, create the resource pre-populated
@@ -244,12 +249,13 @@ void FOpenGLDynamicRHI::RHICopyVertexBuffer(FRHIVertexBuffer* SourceBufferRHI, F
 
 void FOpenGLDynamicRHI::RHITransferVertexBufferUnderlyingResource(FRHIVertexBuffer* DestVertexBuffer, FRHIVertexBuffer* SrcVertexBuffer)
 {
+	VERIFY_GL_SCOPE();
 	check(DestVertexBuffer);
 	FOpenGLVertexBuffer* Dest = ResourceCast(DestVertexBuffer);
 	if (!SrcVertexBuffer)
 	{
-		// Not implemented yet
-		checkNoEntry();
+		TRefCountPtr<FOpenGLVertexBuffer> Src = new FOpenGLVertexBuffer();
+		Dest->Swap(*Src);
 	}
 	else
 	{
diff --git a/Engine/Source/Runtime/OpenGLDrv/Public/OpenGLDrv.h b/Engine/Source/Runtime/OpenGLDrv/Public/OpenGLDrv.h
--- a/Engine/Source/Runtime/OpenGLDrv/Public/OpenGLDrv.h
+++ b/Engine/Source/Runtime/OpenGLDrv/Public/OpenGLDrv.h
@@ -354,6 +354,7 @@ public:
 	virtual void UnlockVertexBuffer_BottomOfPipe(FRHICommandListImmediate& RHICmdList, FRHIVertexBuffer* VertexBuffer) final override;
 	virtual void RHICopyVertexBuffer(FRHIVertexBuffer* SourceBuffer, FRHIVertexBuffer* DestBuffer) final override;
 	virtual void RHITransferVertexBufferUnderlyingResource(FRHIVertexBuffer* DestVertexBuffer, FRHIVertexBuffer* SrcVertexBuffer) final override;
+	virtual void RHITransferIndexBufferUnderlyingResource(FRHIIndexBuffer* DestIndexBuffer, FRHIIndexBuffer* SrcIndexBuffer) final override;
 	virtual FStructuredBufferRHIRef RHICreateStructuredBuffer(uint32 Stride, uint32 Size, uint32 InUsage, FRHIResourceCreateInfo& CreateInfo) final override;
 	virtual void* LockStructuredBuffer_BottomOfPipe(FRHICommandListImmediate& RHICmdList, FRHIStructuredBuffer* StructuredBuffer, uint32 Offset, uint32 SizeRHI, EResourceLockMode LockMode) final override;
 	virtual void UnlockStructuredBuffer_BottomOfPipe(FRHICommandListImmediate& RHICmdList, FRHIStructuredBuffer* StructuredBuffer) final override;
@@ -366,6 +367,7 @@ public:
 	virtual FShaderResourceViewRHIRef RHICreateShaderResourceView(const FShaderResourceViewInitializer& Initializer) final override;
 	virtual FShaderResourceViewRHIRef RHICreateShaderResourceView(FRHIIndexBuffer* Buffer) final override;
 	virtual void RHIUpdateShaderResourceView(FRHIShaderResourceView* SRV, FRHIVertexBuffer* VertexBuffer, uint32 Stride, uint8 Format) final override;
+	virtual void RHIUpdateShaderResourceView(FRHIShaderResourceView* SRV, FRHIIndexBuffer* IndexBuffer) final override;
 	virtual uint64 RHICalcTexture2DPlatformSize(uint32 SizeX, uint32 SizeY, uint8 Format, uint32 NumMips, uint32 NumSamples, uint32 Flags, const FRHIResourceCreateInfo& CreateInfo, uint32& OutAlign) final override;
 	virtual uint64 RHICalcTexture3DPlatformSize(uint32 SizeX, uint32 SizeY, uint32 SizeZ, uint8 Format, uint32 NumMips, uint32 Flags, const FRHIResourceCreateInfo& CreateInfo, uint32& OutAlign) final override;
 	virtual uint64 RHICalcTextureCubePlatformSize(uint32 Size, uint8 Format, uint32 NumMips, uint32 Flags, const FRHIResourceCreateInfo& CreateInfo, uint32& OutAlign) final override;
diff --git a/Engine/Source/Runtime/OpenGLDrv/Public/OpenGLResources.h b/Engine/Source/Runtime/OpenGLDrv/Public/OpenGLResources.h
--- a/Engine/Source/Runtime/OpenGLDrv/Public/OpenGLResources.h
+++ b/Engine/Source/Runtime/OpenGLDrv/Public/OpenGLResources.h
@@ -420,6 +420,19 @@ public:
 	/** Needed on OS X to force a rebind of the texture buffer to the texture name to workaround radr://18379338 */
 	uint64 ModificationCount;
 
+	TOpenGLBuffer()
+		: Resource(0)
+		, ModificationCount(0)
+		, bIsLocked(false)
+		, bIsLockReadOnly(false)
+		, bStreamDraw(false)
+		, bLockBufferWasAllocated(false)
+		, LockSize(0)
+		, LockOffset(0)
+		, LockBuffer(NULL)
+		, RealSize(0)
+	{ }
+
 	TOpenGLBuffer(uint32 InStride,uint32 InSize,uint32 InUsage,
 		const void *InData = NULL, bool bStreamedDraw = false, GLuint ResourceToUse = 0, uint32 ResourceSize = 0)
 	: BaseType(InStride,InSize,InUsage)
@@ -883,6 +896,9 @@ private:
 class FOpenGLBaseVertexBuffer : public FRHIVertexBuffer
 {
 public:
+	FOpenGLBaseVertexBuffer() : ZeroStrideVertexBuffer(0)
+	{}
+
 	FOpenGLBaseVertexBuffer(uint32 InStride,uint32 InSize,uint32 InUsage): FRHIVertexBuffer(InSize,InUsage), ZeroStrideVertexBuffer(0)
 	{
 		if(!(FOpenGL::SupportsVertexAttribBinding() && OpenGLConsoleVariables::bUseVAB) && InUsage & BUF_ZeroStride )
@@ -1003,6 +1019,9 @@ public:
 class FOpenGLBaseIndexBuffer : public FRHIIndexBuffer
 {
 public:
+	FOpenGLBaseIndexBuffer()
+	{}
+
 	FOpenGLBaseIndexBuffer(uint32 InStride,uint32 InSize,uint32 InUsage): FRHIIndexBuffer(InStride,InSize,InUsage)
 	{
 #if ENABLE_LOW_LEVEL_MEM_TRACKER
@@ -1943,6 +1962,7 @@ public:
 
 	/** Needed on OS X to force a rebind of the texture buffer to the texture name to workaround radr://18379338 */
 	FVertexBufferRHIRef VertexBuffer;
+	FIndexBufferRHIRef IndexBuffer;
 	uint64 ModificationVersion;
 	uint8 Format;
 
@@ -1956,6 +1976,23 @@ public:
 	,	OwnsResource(true)
 	{}
 
+	FOpenGLShaderResourceView(FOpenGLDynamicRHI* InOpenGLRHI, GLuint InResource, GLenum InTarget, FRHIIndexBuffer* InIndexBuffer)
+		: Resource(InResource)
+		, Target(InTarget)
+		, LimitMip(-1)
+		, IndexBuffer(InIndexBuffer)
+		, ModificationVersion(0)
+		, Format(0)
+		, OpenGLRHI(InOpenGLRHI)
+		, OwnsResource(true)
+	{
+		if (IndexBuffer)
+		{
+			FOpenGLIndexBuffer* IB = (FOpenGLIndexBuffer*)IndexBuffer.GetReference();
+			ModificationVersion = IB->ModificationCount;
+		}
+	}
+
 	FOpenGLShaderResourceView( FOpenGLDynamicRHI* InOpenGLRHI, GLuint InResource, GLenum InTarget, FRHIVertexBuffer* InVertexBuffer, uint8 InFormat )
 	:	Resource(InResource)
 	,	Target(InTarget)
@@ -1985,12 +2022,6 @@ public:
 
 	virtual ~FOpenGLShaderResourceView( void );
 
-	void SetGLParameters(GLuint InResource, GLenum InTarget)
-	{
-		Target = InTarget;
-		Resource = InResource;
-	}
-
 protected:
 	FOpenGLDynamicRHI* OpenGLRHI;
 	bool OwnsResource;
