OculusQuest Bug
will keep the DefaultVAO for Oculus

diff --git a/Engine/Source/Runtime/OpenGLDrv/Private/Android/AndroidOpenGL.cpp b/Engine/Source/Runtime/OpenGLDrv/Private/Android/AndroidOpenGL.cpp
--- a/Engine/Source/Runtime/OpenGLDrv/Private/Android/AndroidOpenGL.cpp
+++ b/Engine/Source/Runtime/OpenGLDrv/Private/Android/AndroidOpenGL.cpp
@@ -410,7 +410,7 @@ void FPlatformOpenGLDevice::SetupCurrentContext()
 		return;
 	}
 	
-	if (*DefaultVao == 0)
+	if (*DefaultVao == 0 && !AndroidThunkCpp_IsOculusMobileApplication())
 	{
 		glGenVertexArrays(1, DefaultVao);
 		glBindVertexArray(*DefaultVao);
