Expand generic FGenericPlatformRHIFramePacer to get/set frame pace rate rate mapping to rhi.syncinterval
Add Android OpenGL FPlatformRHIFramePacer support optionally using swappy and relocate code from AndroidEGL
Add FIOSPlatformRHIFramePacer get/set frame pace functions to set rhi.syncinterval taking into account GetMaxRefreshRate
Deprecate UGameUserSettings::GetSyncInterval() and add GetFramePace()

diff --git a/Engine/Source/Runtime/ApplicationCore/Private/HAL/PlatformFramePacer.cpp b/Engine/Source/Runtime/ApplicationCore/Private/HAL/PlatformFramePacer.cpp
--- a/Engine/Source/Runtime/ApplicationCore/Private/HAL/PlatformFramePacer.cpp
+++ b/Engine/Source/Runtime/ApplicationCore/Private/HAL/PlatformFramePacer.cpp
@@ -22,4 +22,4 @@ static FAutoConsoleCommand CFramePace(
 		UE_LOG(LogFramePacer, Display, TEXT("r.FramePace : current frame pace is %d"), FPlatformRHIFramePacer::GetFramePace());
 	}
 }),
-ECVF_Cheat);
\ No newline at end of file
+ECVF_Default);
diff --git a/Engine/Source/Runtime/OpenGLDrv/OpenGLDrv.Build.cs b/Engine/Source/Runtime/OpenGLDrv/OpenGLDrv.Build.cs
--- a/Engine/Source/Runtime/OpenGLDrv/OpenGLDrv.Build.cs
+++ b/Engine/Source/Runtime/OpenGLDrv/OpenGLDrv.Build.cs
@@ -55,8 +55,8 @@ public class OpenGLDrv : ModuleRules
         if (Target.Platform == UnrealTargetPlatform.Android)
         {
             // for Swappy
-            PublicDefinitions.Add("USE_ANDROID_OPENGL_SWAPPY=0");
-/*
+            PublicDefinitions.Add("USE_ANDROID_OPENGL_SWAPPY=1");
+
             PrivateDependencyModuleNames.AddRange(
                 new string[]
                 {
@@ -64,7 +64,6 @@ public class OpenGLDrv : ModuleRules
                 "GoogleGameSDK"
                 }
             );
-*/
         }
 
         if (Target.Platform != UnrealTargetPlatform.Win32 && Target.Platform != UnrealTargetPlatform.Win64
diff --git a/Engine/Source/Runtime/OpenGLDrv/Private/Android/AndroidOpenGLFramePacer.cpp b/Engine/Source/Runtime/OpenGLDrv/Private/Android/AndroidOpenGLFramePacer.cpp
--- a/Engine/Source/Runtime/OpenGLDrv/Private/Android/AndroidOpenGLFramePacer.cpp
+++ b/Engine/Source/Runtime/OpenGLDrv/Private/Android/AndroidOpenGLFramePacer.cpp
@@ -9,12 +9,13 @@
  *******************************************************************/
 
 #include "OpenGLDrvPrivate.h"
-
+#include "Math/UnrealMathUtility.h"
 #if USE_ANDROID_OPENGL_SWAPPY
 #include "Android/AndroidJNI.h"
 #include "Android/AndroidApplication.h"
 #include "swappyGL.h"
 #include "swappyGL_extra.h"
+#include "swappy_common.h"
 #endif
 
 #include "AndroidEGL.h"
@@ -40,8 +41,6 @@ void FAndroidOpenGLFramePacer::Init()
 			}
 		}));
 	}
-#else
-	bSwappyInit = false;
 #endif
 }
 
@@ -115,10 +114,31 @@ bool FAndroidOpenGLFramePacer::SupportsFramePace(int32 QueryFramePace)
 #if USE_ANDROID_OPENGL_SWAPPY
 	if (FAndroidPlatformRHIFramePacer::CVarUseSwappyForFramePacing.GetValueOnAnyThread() == 1)
 	{
-		// temporary
-		if (QueryFramePace == 90 || QueryFramePace == 45)
+		int NumRates = Swappy_getSupportedRefreshRates(nullptr, 0);
+		TArray<uint64> RefreshRatesNs;
+		RefreshRatesNs.AddZeroed(NumRates);
+		Swappy_getSupportedRefreshRates((uint64_t*)RefreshRatesNs.GetData(), NumRates);
+		TArray<int32> RefreshRates;
+		RefreshRates.Empty(NumRates);
+		FString DebugString = TEXT("Supported Refresh Rates:");
+		for (uint64 RateNs : RefreshRatesNs)
 		{
-			return true;
+			if (RateNs > 0)
+			{
+				int32 RefreshRate = FMath::DivideAndRoundNearest(1000000000ull, RateNs);
+				RefreshRates.Add(RefreshRate);
+				DebugString += FString::Printf(TEXT(" %d (%ld ns)"), RefreshRate, RateNs);
+			}
+		}
+		UE_LOG(LogRHI, Log, TEXT("%s"), *DebugString);
+
+		for (uint64 Rate : RefreshRates)
+		{
+			if ((Rate % QueryFramePace) == 0)
+			{
+				UE_LOG(LogRHI, Log, TEXT("Using Refresh rate %d with sync interval %d"), Rate, Rate / QueryFramePace);
+				return true;
+			}
 		}
 	}
 #endif
@@ -144,7 +164,7 @@ bool FAndroidOpenGLFramePacer::SwapBuffers(bool bLockToVsync)
 	bool bPrintMethod = false;
 
 #if USE_ANDROID_OPENGL_SWAPPY
-	if (FAndroidPlatformRHIFramePacer::CVarUseSwappyForFramePacing.GetValueOnRenderThread() != 0 && ensure(bSwappyInit))
+	if (FAndroidPlatformRHIFramePacer::CVarUseSwappyForFramePacing.GetValueOnRenderThread() != 0 && bSwappyInit)
 	{
 		int64 DesiredFrameNS = (1000000000L) / (int64)FAndroidPlatformRHIFramePacer::GetFramePace();
 		SwappyGL_setSwapIntervalNS(DesiredFrameNS);
