remove unused (and incorrect) GL format line, fixes static analysis report.
UE-91483

diff --git a/Engine/Source/Runtime/OpenGLDrv/Private/OpenGLUAV.cpp b/Engine/Source/Runtime/OpenGLDrv/Private/OpenGLUAV.cpp
--- a/Engine/Source/Runtime/OpenGLDrv/Private/OpenGLUAV.cpp
+++ b/Engine/Source/Runtime/OpenGLDrv/Private/OpenGLUAV.cpp
@@ -156,7 +156,6 @@ void FOpenGLDynamicRHI::RHIUpdateShaderResourceView(FRHIShaderResourceView* SRV,
 	{
 		uint32 Stride = IndexBuffer->GetStride();
 		GLenum Format = (Stride == 2) ? GL_R16UI : GL_R32UI;
-		const FOpenGLTextureFormat& GLFormat = GOpenGLTextureFormats[Format];
 		check(SRVGL->Target == GL_TEXTURE_BUFFER);
 
 		uint32 NumElements = IndexBuffer->GetSize() / Stride;
