Fix from Oculus for the VAO overwrite issue causing black screen on Quest in OpenGL.

diff --git a/Engine/Source/Runtime/OpenGLDrv/Private/Android/AndroidOpenGL.cpp b/Engine/Source/Runtime/OpenGLDrv/Private/Android/AndroidOpenGL.cpp
--- a/Engine/Source/Runtime/OpenGLDrv/Private/Android/AndroidOpenGL.cpp
+++ b/Engine/Source/Runtime/OpenGLDrv/Private/Android/AndroidOpenGL.cpp
@@ -410,7 +410,7 @@ void FPlatformOpenGLDevice::SetupCurrentContext()
 		return;
 	}
 	
-	if (*DefaultVao == 0 && !AndroidThunkCpp_IsOculusMobileApplication())
+	if (*DefaultVao == 0)
 	{
 		glGenVertexArrays(1, DefaultVao);
 		glBindVertexArray(*DefaultVao);
