Fixed: TM-RuntimeVirtualTextures with Mobile Virtual Texture Support enabled not rendering correctly

diff --git a/Engine/Source/Runtime/OpenGLDrv/Private/OpenGLUAV.cpp b/Engine/Source/Runtime/OpenGLDrv/Private/OpenGLUAV.cpp
--- a/Engine/Source/Runtime/OpenGLDrv/Private/OpenGLUAV.cpp
+++ b/Engine/Source/Runtime/OpenGLDrv/Private/OpenGLUAV.cpp
@@ -345,10 +345,10 @@ void FOpenGLDynamicRHI::RHIClearUAVFloat(FRHIUnorderedAccessView* UnorderedAcces
 		switch (NumComponents)
 		{
 		case 1:
-			ClearUAVShader_T<EClearReplacementResourceType::Buffer, EClearReplacementValueType::Float, 1, false>(RHICmdList, UnorderedAccessViewRHI, 1, 1, 1, *reinterpret_cast<const float(*)[1]>(&Values));
+			ClearUAVShader_T<EClearReplacementResourceType::Buffer, EClearReplacementValueType::Float, 1, false>(RHICmdList, UnorderedAccessViewRHI, NumElements, 1, 1, *reinterpret_cast<const float(*)[1]>(&Values));
 			break;
 		case 4:
-			ClearUAVShader_T<EClearReplacementResourceType::Buffer, EClearReplacementValueType::Float, 4, false>(RHICmdList, UnorderedAccessViewRHI, 4, 1, 1, *reinterpret_cast<const float(*)[4]>(&Values));
+			ClearUAVShader_T<EClearReplacementResourceType::Buffer, EClearReplacementValueType::Float, 4, false>(RHICmdList, UnorderedAccessViewRHI, NumElements, 1, 1, *reinterpret_cast<const float(*)[4]>(&Values));
 			break;
 		default:
 			check(false);
@@ -379,10 +379,10 @@ void FOpenGLDynamicRHI::RHIClearUAVUint(FRHIUnorderedAccessView* UnorderedAccess
 		switch (NumComponents)
 		{
 		case 1:
-			ClearUAVShader_T<EClearReplacementResourceType::Buffer, EClearReplacementValueType::Uint32, 1, false>(RHICmdList, UnorderedAccessViewRHI, 1, 1, 1, *reinterpret_cast<const uint32(*)[1]>(&Values));
+			ClearUAVShader_T<EClearReplacementResourceType::Buffer, EClearReplacementValueType::Uint32, 1, false>(RHICmdList, UnorderedAccessViewRHI, NumElements, 1, 1, *reinterpret_cast<const uint32(*)[1]>(&Values));
 			break;
 		case 4:
-			ClearUAVShader_T<EClearReplacementResourceType::Buffer, EClearReplacementValueType::Uint32, 4, false>(RHICmdList, UnorderedAccessViewRHI, 4, 1, 1, *reinterpret_cast<const uint32(*)[4]>(&Values));
+			ClearUAVShader_T<EClearReplacementResourceType::Buffer, EClearReplacementValueType::Uint32, 4, false>(RHICmdList, UnorderedAccessViewRHI, NumElements, 1, 1, *reinterpret_cast<const uint32(*)[4]>(&Values));
 			break;
 		default:
 			check(false);
diff --git a/Engine/Source/Runtime/VulkanRHI/Private/VulkanUAV.cpp b/Engine/Source/Runtime/VulkanRHI/Private/VulkanUAV.cpp
--- a/Engine/Source/Runtime/VulkanRHI/Private/VulkanUAV.cpp
+++ b/Engine/Source/Runtime/VulkanRHI/Private/VulkanUAV.cpp
@@ -75,6 +75,7 @@ void FVulkanShaderResourceView::UpdateView()
 	// update the buffer view for dynamic backed buffers (or if it was never set)
 	if (SourceBuffer != nullptr)
 	{
+		uint32 CurrentViewSize = Size;
 		if (SourceBuffer->IsVolatile() && VolatileLockCounter != SourceBuffer->GetVolatileLockCounter())
 		{
 			VkBuffer SourceVolatileBufferHandle = SourceBuffer->GetHandle();
@@ -82,12 +83,12 @@ void FVulkanShaderResourceView::UpdateView()
 			// If the volatile buffer shrinks, make sure our size doesn't exceed the new limit.
 			uint32 AvailableSize = SourceBuffer->GetVolatileLockSize();
 			AvailableSize = Offset < AvailableSize ? AvailableSize - Offset : 0;
-			Size = FMath::Min(Size, AvailableSize);
+			CurrentViewSize = FMath::Min(CurrentViewSize, AvailableSize);
 
 			// We might end up with the same BufferView, so do not recreate in that case
 			if (!BufferViews[0] 
 				|| BufferViews[0]->Offset != (SourceBuffer->GetOffset() + Offset)
-				|| BufferViews[0]->Size != Size
+				|| BufferViews[0]->Size != CurrentViewSize
 				|| VolatileBufferHandle != SourceVolatileBufferHandle)
 			{
 				BufferViews[0] = nullptr;
@@ -104,7 +105,7 @@ void FVulkanShaderResourceView::UpdateView()
 		if (!BufferViews[BufferIndex])
 		{
 			BufferViews[BufferIndex] = new FVulkanBufferView(Device);
-			BufferViews[BufferIndex]->Create(SourceBuffer, BufferViewFormat, SourceBuffer->GetOffset() + Offset, Size);
+			BufferViews[BufferIndex]->Create(SourceBuffer, BufferViewFormat, SourceBuffer->GetOffset() + Offset, CurrentViewSize);
 		}
 	}
 	else if (SourceStructuredBuffer)
