Fixing PVS 7.7 Issues:

warning V1013: Suspicious subexpression in a sequence of similar comparisons.

diff --git a/Engine/Source/Runtime/OpenGLDrv/Private/OpenGLCommands.cpp b/Engine/Source/Runtime/OpenGLDrv/Private/OpenGLCommands.cpp
--- a/Engine/Source/Runtime/OpenGLDrv/Private/OpenGLCommands.cpp
+++ b/Engine/Source/Runtime/OpenGLDrv/Private/OpenGLCommands.cpp
@@ -1992,7 +1992,7 @@ void FOpenGLDynamicRHI::SetupVertexArraysVAB(FOpenGLContextState& ContextState,
 					KnowsDivisor[StreamIndex] = true;
 					Divisor[StreamIndex] = VertexElement.Divisor;
 
-					if ((Attr.StreamOffset != VertexElement.Offset) ||
+					if ((Attr.StreamOffset != VertexElement.Offset) || //-V1013
 						(Attr.Size != VertexElement.Size) ||
 						(Attr.Type != VertexElement.Type) ||
 						(Attr.bNormalized != VertexElement.bNormalized))
