Fixed Ensure condition failed: oldValue == newValue launching on Android [ThreadingBase.cpp] when app loses focus on Android

diff --git a/Engine/Source/Runtime/OpenGLDrv/Private/Android/AndroidEGL.cpp b/Engine/Source/Runtime/OpenGLDrv/Private/Android/AndroidEGL.cpp
--- a/Engine/Source/Runtime/OpenGLDrv/Private/Android/AndroidEGL.cpp
+++ b/Engine/Source/Runtime/OpenGLDrv/Private/Android/AndroidEGL.cpp
@@ -1183,6 +1183,7 @@ void BlockRendering()
 
 	FGraphEventRef RTBlockTask = FFunctionGraphTask::CreateAndDispatchWhenReady([RTBlockedTrigger]()
 	{
+		FTaskTagScope TagScope(ETaskTag::ERenderingThread);
 		BlockOnLostWindowRenderCommand(RTBlockedTrigger);
 	}, TStatId(), NULL, ENamedThreads::GetRenderThread());
 
diff --git a/Engine/Source/Runtime/OpenGLDrv/Private/OpenGLQuery.cpp b/Engine/Source/Runtime/OpenGLDrv/Private/OpenGLQuery.cpp
--- a/Engine/Source/Runtime/OpenGLDrv/Private/OpenGLQuery.cpp
+++ b/Engine/Source/Runtime/OpenGLDrv/Private/OpenGLQuery.cpp
@@ -509,6 +509,7 @@ public:
 
 	void DoTask(ENamedThreads::Type CurrentThread, const FGraphEventRef& MyCompletionGraphEvent)
 	{
+		FTaskTagScope TagScope(ETaskTag::ERhiThread);
 		check(IsRunningRHIInDedicatedThread() && IsInRHIThread()); // this should never be used on a platform that doesn't support the RHI thread, and it can't quite work when running the RHI stuff on task threads
 		if (bWait)
 		{
