Remove Texture2D as a storage for GPUScene primitives, OpenGL will use StructuredBuffers supported through SSBOs

diff --git a/Engine/Plugins/Runtime/HairStrands/Source/HairStrandsCore/Private/HairCardsVertexFactory.cpp b/Engine/Plugins/Runtime/HairStrands/Source/HairStrandsCore/Private/HairCardsVertexFactory.cpp
--- a/Engine/Plugins/Runtime/HairStrands/Source/HairStrandsCore/Private/HairCardsVertexFactory.cpp
+++ b/Engine/Plugins/Runtime/HairStrands/Source/HairStrandsCore/Private/HairCardsVertexFactory.cpp
@@ -218,7 +218,6 @@ void FHairCardsVertexFactory::ModifyCompilationEnvironment(const FVertexFactoryS
 	const bool bUseGPUSceneAndPrimitiveIdStream = VF_CARDS_SUPPORT_GPU_SCENE && Parameters.VertexFactoryType->SupportsPrimitiveIdStream() && UseGPUScene(Parameters.Platform, GetMaxSupportedFeatureLevel(Parameters.Platform));
 	OutEnvironment.SetDefine(TEXT("VF_SUPPORTS_PRIMITIVE_SCENE_DATA"), bUseGPUSceneAndPrimitiveIdStream);
 	OutEnvironment.SetDefine(TEXT("HAIR_CARD_MESH_FACTORY"), TEXT("1"));
-	OutEnvironment.SetDefine(TEXT("VF_GPU_SCENE_TEXTURE"), bUseGPUSceneAndPrimitiveIdStream && GPUSceneUseTexture2D(Parameters.Platform));
 	OutEnvironment.SetDefine(TEXT("MANUAL_VERTEX_FETCH"), RHISupportsManualVertexFetch(Parameters.Platform) ? TEXT("1") : TEXT("0"));	
 }
 
diff --git a/Engine/Plugins/Runtime/HairStrands/Source/HairStrandsCore/Private/HairStrandsVertexFactory.cpp b/Engine/Plugins/Runtime/HairStrands/Source/HairStrandsCore/Private/HairStrandsVertexFactory.cpp
--- a/Engine/Plugins/Runtime/HairStrands/Source/HairStrandsCore/Private/HairStrandsVertexFactory.cpp
+++ b/Engine/Plugins/Runtime/HairStrands/Source/HairStrandsCore/Private/HairStrandsVertexFactory.cpp
@@ -177,7 +177,6 @@ void FHairStrandsVertexFactory::ModifyCompilationEnvironment(const FVertexFactor
 	const bool bUseGPUSceneAndPrimitiveIdStream = Parameters.VertexFactoryType->SupportsPrimitiveIdStream() && UseGPUScene(Parameters.Platform, GetMaxSupportedFeatureLevel(Parameters.Platform));
 	OutEnvironment.SetDefine(TEXT("VF_SUPPORTS_PRIMITIVE_SCENE_DATA"), bUseGPUSceneAndPrimitiveIdStream);
 	OutEnvironment.SetDefine(TEXT("HAIR_STRAND_MESH_FACTORY"), TEXT("1"));
-	OutEnvironment.SetDefine(TEXT("VF_GPU_SCENE_TEXTURE"), bUseGPUSceneAndPrimitiveIdStream && GPUSceneUseTexture2D(Parameters.Platform));
 }
 
 void FHairStrandsVertexFactory::ValidateCompiledResult(const FVertexFactoryType* Type, EShaderPlatform Platform, const FShaderParameterMap& ParameterMap, TArray<FString>& OutErrors)
diff --git a/Engine/Shaders/Private/SceneData.ush b/Engine/Shaders/Private/SceneData.ush
--- a/Engine/Shaders/Private/SceneData.ush
+++ b/Engine/Shaders/Private/SceneData.ush
@@ -157,44 +157,18 @@ struct FPrimitiveSceneData
 // Stride of a single primitive's data in float4's, must match C++
 #define PRIMITIVE_SCENE_DATA_STRIDE 37
 
-struct FPrimitiveIndex
-{
-#if VF_GPU_SCENE_TEXTURE
-	uint IndY;
-	uint IndX;
-#else
-	uint BaseOffset;
-#endif
-};
-
-FPrimitiveIndex SetupPrimitiveIndexes(uint PrimitiveId)
-{
-	FPrimitiveIndex PrimitiveIndex;
-#if VF_GPU_SCENE_TEXTURE
-	PrimitiveIndex.IndY = (PrimitiveId & 0xFFFF0000) >> 16;
-	PrimitiveIndex.IndX = (PrimitiveId & 0xFFFF) * PRIMITIVE_SCENE_DATA_STRIDE;
-#else
-	PrimitiveIndex.BaseOffset = PrimitiveId * PRIMITIVE_SCENE_DATA_STRIDE;
-#endif
-	return PrimitiveIndex;
-}
-
 /*void UnpackLightmapDataAndUVIndex(uint Packed, out uint LightmapDataIndex, out uint LightmapUVIndex)
 {
 	LightmapDataIndex = (Packed & 0xFFFFFF00) >> 8u;
 	LightmapUVIndex   = (Packed & 0x000000FF);
 }*/
 
-float4 LoadPrimitivePrimitiveSceneDataElement(FPrimitiveIndex PrimitiveIndex, uint ItemIndex)
+float4 LoadPrimitivePrimitiveSceneDataElement(uint PrimitiveIndex, uint ItemIndex)
 {
 #if USE_GLOBAL_GPU_SCENE_DATA
-	return GPUScenePrimitiveSceneData[PrimitiveIndex.BaseOffset + ItemIndex];
+	return GPUScenePrimitiveSceneData[PrimitiveIndex + ItemIndex];
 #else
-	#if VF_GPU_SCENE_TEXTURE
-		return View.PrimitiveSceneDataTexture.Load(int3(PrimitiveIndex.IndX + ItemIndex, PrimitiveIndex.IndY, 0));
-	#else
-		return View.PrimitiveSceneData[PrimitiveIndex.BaseOffset + ItemIndex];
-	#endif
+	return View.PrimitiveSceneData[PrimitiveIndex + ItemIndex];
 #endif
 }
 
@@ -203,7 +177,7 @@ FPrimitiveSceneData GetPrimitiveData(uint PrimitiveId)
 {
 	// Note: layout must match FPrimitiveSceneShaderData in C++
 	// Relying on optimizer to remove unused loads
-	FPrimitiveIndex PrimitiveIndex = SetupPrimitiveIndexes(PrimitiveId);
+	uint PrimitiveIndex = PrimitiveId * PRIMITIVE_SCENE_DATA_STRIDE;
 
 	FPrimitiveSceneData PrimitiveData;
 
diff --git a/Engine/Source/Developer/ShaderFormatOpenGL/Private/GlslBackend.cpp b/Engine/Source/Developer/ShaderFormatOpenGL/Private/GlslBackend.cpp
--- a/Engine/Source/Developer/ShaderFormatOpenGL/Private/GlslBackend.cpp
+++ b/Engine/Source/Developer/ShaderFormatOpenGL/Private/GlslBackend.cpp
@@ -1160,14 +1160,28 @@ class ir_gen_glsl_visitor : public ir_visitor
 			}
 			else if (var->type->is_image())
 			{
+				const int UAV_stage_first_unit[] = 
+				{
+					0, //vertex_shader, must match FOpenGL::GetFirstVertexUAVUnit()
+					0,
+					4, //fragment_shader, must match FOpenGL::GetFirstPixelUAVUnit()
+					0,
+					0,
+					0, //compute_shader
+				};
+						
 				if (var->type->HlslName && (!strncmp(var->type->HlslName, "RWStructuredBuffer<", 19) || !strncmp(var->type->HlslName, "StructuredBuffer<", 17)))
 				{
 					if (bGenerateLayoutLocations && var->explicit_location)
 					{
+						const char * const readonly_str[] = { "", "readonly " };
+						const int readonly = strncmp(var->type->HlslName, "StructuredBuffer<", 17)==0 ? 1 : 0;
+						
 						ralloc_asprintf_append(
 							buffer,
-							"layout(std430,binding=%d) buffer ",
-							var->location
+							"layout(std430,binding=%d) %s buffer ",
+							var->location + UAV_stage_first_unit[ShaderTarget],
+							readonly_str[readonly]
 							);
 					}
 					else
@@ -1207,7 +1221,7 @@ class ir_gen_glsl_visitor : public ir_visitor
 							"layout(%s%s,binding=%d) ",
 							comp_str,
 							type_str[var->type->inner_type->base_type],
-							var->location
+							var->location + UAV_stage_first_unit[ShaderTarget]
 							);
 					}
 					else
diff --git a/Engine/Source/Developer/ShaderFormatOpenGL/Private/ShaderFormatOpenGL.cpp b/Engine/Source/Developer/ShaderFormatOpenGL/Private/ShaderFormatOpenGL.cpp
--- a/Engine/Source/Developer/ShaderFormatOpenGL/Private/ShaderFormatOpenGL.cpp
+++ b/Engine/Source/Developer/ShaderFormatOpenGL/Private/ShaderFormatOpenGL.cpp
@@ -17,7 +17,7 @@ class FShaderFormatGLSL : public IShaderFormat
 	enum
 	{
 		/** Version for shader format, this becomes part of the DDC key. */
-		UE_SHADER_GLSL_VER = 96,
+		UE_SHADER_GLSL_VER = 97,
 	};
 
 	void CheckFormat(FName Format) const
diff --git a/Engine/Source/Developer/ShaderFormatOpenGL/Public/GlslBackend.h b/Engine/Source/Developer/ShaderFormatOpenGL/Public/GlslBackend.h
--- a/Engine/Source/Developer/ShaderFormatOpenGL/Public/GlslBackend.h
+++ b/Engine/Source/Developer/ShaderFormatOpenGL/Public/GlslBackend.h
@@ -40,7 +40,7 @@ public:
 
 	virtual bool AllowsImageLoadsForNonScalar() const { return true; }
 
-	virtual bool EmulateStructuredWithTypedBuffers() const override { return true; }
+	virtual bool EmulateStructuredWithTypedBuffers() const override { return false; }
 };
 
 class ir_variable;
diff --git a/Engine/Source/Runtime/Core/Private/HAL/ConsoleManager.cpp b/Engine/Source/Runtime/Core/Private/HAL/ConsoleManager.cpp
--- a/Engine/Source/Runtime/Core/Private/HAL/ConsoleManager.cpp
+++ b/Engine/Source/Runtime/Core/Private/HAL/ConsoleManager.cpp
@@ -2188,13 +2188,6 @@ static TAutoConsoleVariable<int32> CVarMobileSupportGPUScene(
 	ECVF_ReadOnly | ECVF_RenderThreadSafe
 );
 
-static TAutoConsoleVariable<int32> CVarMobileGPUSceneUseTexture2D(
-	TEXT("r.Mobile.UseGPUSceneTexture"),
-	0,
-	TEXT("Use a Texture2D instead of TextureBuffer for GPUScene.\n"),
-	ECVF_RenderThreadSafe
-);
-
 static TAutoConsoleVariable<int32> CVarSetClearSceneMethod(
 	TEXT("r.ClearSceneMethod"),
 	1,
diff --git a/Engine/Source/Runtime/Engine/Private/GPUSkinVertexFactory.cpp b/Engine/Source/Runtime/Engine/Private/GPUSkinVertexFactory.cpp
--- a/Engine/Source/Runtime/Engine/Private/GPUSkinVertexFactory.cpp
+++ b/Engine/Source/Runtime/Engine/Private/GPUSkinVertexFactory.cpp
@@ -423,7 +423,6 @@ void TGPUSkinVertexFactory<BoneInfluenceType>::ModifyCompilationEnvironment(cons
 	OutEnvironment.SetDefine(TEXT("GPU_SKINNED_MESH_FACTORY"), 1);
 
 	OutEnvironment.SetDefine(TEXT("VF_SUPPORTS_PRIMITIVE_SCENE_DATA"), bSupportsPrimitiveIdStream && bUseGPUScene);
-	OutEnvironment.SetDefine(TEXT("VF_GPU_SCENE_TEXTURE"), bSupportsPrimitiveIdStream && bUseGPUScene && GPUSceneUseTexture2D(Parameters.Platform));
 }
 
 template<GPUSkinBoneInfluenceType BoneInfluenceType>
diff --git a/Engine/Source/Runtime/Engine/Private/LocalVertexFactory.cpp b/Engine/Source/Runtime/Engine/Private/LocalVertexFactory.cpp
--- a/Engine/Source/Runtime/Engine/Private/LocalVertexFactory.cpp
+++ b/Engine/Source/Runtime/Engine/Private/LocalVertexFactory.cpp
@@ -209,7 +209,6 @@ void FLocalVertexFactory::ModifyCompilationEnvironment(const FVertexFactoryShade
 	}
 
 	OutEnvironment.SetDefine(TEXT("VF_SUPPORTS_PRIMITIVE_SCENE_DATA"), Parameters.VertexFactoryType->SupportsPrimitiveIdStream() && UseGPUScene(Parameters.Platform, GetMaxSupportedFeatureLevel(Parameters.Platform)));
-	OutEnvironment.SetDefine(TEXT("VF_GPU_SCENE_TEXTURE"), Parameters.VertexFactoryType->SupportsPrimitiveIdStream() && UseGPUScene(Parameters.Platform, GetMaxSupportedFeatureLevel(Parameters.Platform)) && GPUSceneUseTexture2D(Parameters.Platform));
 }
 
 void FLocalVertexFactory::ValidateCompiledResult(const FVertexFactoryType* Type, EShaderPlatform Platform, const FShaderParameterMap& ParameterMap, TArray<FString>& OutErrors)
diff --git a/Engine/Source/Runtime/Engine/Private/PrimitiveUniformShaderParameters.cpp b/Engine/Source/Runtime/Engine/Private/PrimitiveUniformShaderParameters.cpp
--- a/Engine/Source/Runtime/Engine/Private/PrimitiveUniformShaderParameters.cpp
+++ b/Engine/Source/Runtime/Engine/Private/PrimitiveUniformShaderParameters.cpp
@@ -43,20 +43,9 @@ void FSinglePrimitiveStructured::UploadToGPU()
 	{
 		void* LockedData = nullptr;
 
-		EShaderPlatform SafeShaderPlatform = ShaderPlatform < SP_NumPlatforms ? ShaderPlatform : GMaxRHIShaderPlatform;
-		if (!GPUSceneUseTexture2D(SafeShaderPlatform))
-		{
-			LockedData = RHILockBuffer(PrimitiveSceneDataBufferRHI, 0, FPrimitiveSceneShaderData::PrimitiveDataStrideInFloat4s * sizeof(FVector4), RLM_WriteOnly);
-			FPlatformMemory::Memcpy(LockedData, PrimitiveSceneData.Data, FPrimitiveSceneShaderData::PrimitiveDataStrideInFloat4s * sizeof(FVector4));
-			RHIUnlockBuffer(PrimitiveSceneDataBufferRHI);
-		}
-		else
-		{
-			uint32 SrcStride;
-			LockedData = RHILockTexture2D(PrimitiveSceneDataTextureRHI, 0, RLM_WriteOnly, SrcStride, false);
-			FPlatformMemory::Memcpy(LockedData, PrimitiveSceneData.Data, FPrimitiveSceneShaderData::PrimitiveDataStrideInFloat4s * sizeof(FVector4));
-			RHIUnlockTexture2D(PrimitiveSceneDataTextureRHI, 0, false);
-		}
+		LockedData = RHILockBuffer(PrimitiveSceneDataBufferRHI, 0, FPrimitiveSceneShaderData::PrimitiveDataStrideInFloat4s * sizeof(FVector4), RLM_WriteOnly);
+		FPlatformMemory::Memcpy(LockedData, PrimitiveSceneData.Data, FPrimitiveSceneShaderData::PrimitiveDataStrideInFloat4s * sizeof(FVector4));
+		RHIUnlockBuffer(PrimitiveSceneDataBufferRHI);
 
 		LockedData = RHILockBuffer(LightmapSceneDataBufferRHI, 0, FLightmapSceneShaderData::LightmapDataStrideInFloat4s * sizeof(FVector4), RLM_WriteOnly);
 		FPlatformMemory::Memcpy(LockedData, LightmapSceneData.Data, FLightmapSceneShaderData::LightmapDataStrideInFloat4s * sizeof(FVector4));
diff --git a/Engine/Source/Runtime/Engine/Private/SceneManagement.cpp b/Engine/Source/Runtime/Engine/Private/SceneManagement.cpp
--- a/Engine/Source/Runtime/Engine/Private/SceneManagement.cpp
+++ b/Engine/Source/Runtime/Engine/Private/SceneManagement.cpp
@@ -790,7 +790,6 @@ FViewUniformShaderParameters::FViewUniformShaderParameters()
 	CameraAerialPerspectiveVolume = GBlackAlpha1VolumeTexture->TextureRHI;
 	CameraAerialPerspectiveVolumeSampler = TStaticSamplerState<SF_Bilinear>::GetRHI();
 
-	PrimitiveSceneDataTexture = OrBlack2DIfNull(GIdentityPrimitiveBuffer.PrimitiveSceneDataTextureRHI);
 	PrimitiveSceneData = GIdentityPrimitiveBuffer.PrimitiveSceneDataBufferSRV;
 	InstanceSceneData = GIdentityPrimitiveBuffer.InstanceSceneDataBufferSRV;
 	LightmapSceneData = GIdentityPrimitiveBuffer.LightmapSceneDataBufferSRV;
diff --git a/Engine/Source/Runtime/Engine/Public/SceneView.h b/Engine/Source/Runtime/Engine/Public/SceneView.h
--- a/Engine/Source/Runtime/Engine/Public/SceneView.h
+++ b/Engine/Source/Runtime/Engine/Public/SceneView.h
@@ -857,7 +857,6 @@ BEGIN_GLOBAL_SHADER_PARAMETER_STRUCT_WITH_CONSTRUCTOR(FViewUniformShaderParamete
 	SHADER_PARAMETER_SAMPLER(SamplerState, PreIntegratedBRDFSampler)
 	SHADER_PARAMETER_SRV(StructuredBuffer<float4>, PrimitiveSceneData)
 	SHADER_PARAMETER_SRV(StructuredBuffer<float4>, InstanceSceneData)
-	SHADER_PARAMETER_TEXTURE(Texture2D<float4>, PrimitiveSceneDataTexture)
 	SHADER_PARAMETER_SRV(StructuredBuffer<float4>, LightmapSceneData)
 	SHADER_PARAMETER_SRV(StructuredBuffer<float4>, SkyIrradianceEnvironmentMap)
 
diff --git a/Engine/Source/Runtime/OpenGLDrv/Private/Android/AndroidOpenGL.h b/Engine/Source/Runtime/OpenGLDrv/Private/Android/AndroidOpenGL.h
--- a/Engine/Source/Runtime/OpenGLDrv/Private/Android/AndroidOpenGL.h
+++ b/Engine/Source/Runtime/OpenGLDrv/Private/Android/AndroidOpenGL.h
@@ -157,7 +157,8 @@ struct FAndroidOpenGL : public FOpenGLES
 	static FORCEINLINE GLint GetMaxComputeUniformComponents() { check(MaxComputeUniformComponents != -1); return MaxComputeUniformComponents; }
 	static FORCEINLINE GLint GetFirstComputeUAVUnit()			{ return 0; }
 	static FORCEINLINE GLint GetMaxComputeUAVUnits()			{ check(MaxComputeUAVUnits != -1); return MaxComputeUAVUnits; }
-	static FORCEINLINE GLint GetFirstPixelUAVUnit()				{ return 0; }
+	static FORCEINLINE GLint GetFirstVertexUAVUnit()			{ return 0; }
+	static FORCEINLINE GLint GetFirstPixelUAVUnit()				{ return 4; }
 	static FORCEINLINE GLint GetMaxPixelUAVUnits()				{ check(MaxPixelUAVUnits != -1); return MaxPixelUAVUnits; }
 	static FORCEINLINE GLint GetMaxCombinedUAVUnits()			{ return MaxCombinedUAVUnits; }
 		
diff --git a/Engine/Source/Runtime/OpenGLDrv/Private/OpenGL4.cpp b/Engine/Source/Runtime/OpenGLDrv/Private/OpenGL4.cpp
--- a/Engine/Source/Runtime/OpenGLDrv/Private/OpenGL4.cpp
+++ b/Engine/Source/Runtime/OpenGLDrv/Private/OpenGL4.cpp
@@ -16,7 +16,6 @@ GLint FOpenGL4::MaxCombinedUAVUnits = 0;
 GLint FOpenGL4::MaxComputeUAVUnits = -1;
 GLint FOpenGL4::MaxPixelUAVUnits = -1;
 
-bool FOpenGL4::bSupportsComputeShaders = true;
 bool FOpenGL4::bSupportsGPUMemoryInfo = false;
 bool FOpenGL4::bSupportsVertexAttribBinding = true;
 bool FOpenGL4::bSupportsTextureView = true;
@@ -27,17 +26,15 @@ void FOpenGL4::ProcessQueryGLInt()
 	GET_GL_INT(GL_MAX_FRAGMENT_IMAGE_UNIFORMS, 0, MaxPixelUAVUnits);
 	GET_GL_INT(GL_TEXTURE_BUFFER_OFFSET_ALIGNMENT, 0, TextureBufferAlignment);
 
-	if (bSupportsComputeShaders)
-	{
-		GET_GL_INT(GL_MAX_COMPUTE_TEXTURE_IMAGE_UNITS, 0, MaxComputeTextureImageUnits);
-		GET_GL_INT(GL_MAX_COMPUTE_UNIFORM_COMPONENTS, 0, MaxComputeUniformComponents);
-		GET_GL_INT(GL_MAX_COMPUTE_IMAGE_UNIFORMS, 0, MaxComputeUAVUnits);
-	}
+	GET_GL_INT(GL_MAX_COMPUTE_TEXTURE_IMAGE_UNITS, 0, MaxComputeTextureImageUnits);
+	GET_GL_INT(GL_MAX_COMPUTE_UNIFORM_COMPONENTS, 0, MaxComputeUniformComponents);
+	GET_GL_INT(GL_MAX_COMPUTE_IMAGE_UNIFORMS, 0, MaxComputeUAVUnits);
 	
 	// clamp UAV units to a sensible limit
 	MaxCombinedUAVUnits = FMath::Min(MaxCombinedUAVUnits, 8);
-	MaxComputeUAVUnits = FMath::Min(MaxComputeUAVUnits, MaxCombinedUAVUnits);
-	MaxPixelUAVUnits = FMath::Min(MaxPixelUAVUnits, MaxCombinedUAVUnits);
+	MaxComputeUAVUnits = FMath::Min(MaxComputeUAVUnits, 16);
+	// this is split between VS and PS, 4 to each stage
+	MaxPixelUAVUnits = FMath::Min(MaxPixelUAVUnits, 4);
 }
 
 void FOpenGL4::ProcessExtensions( const FString& ExtensionsString )
@@ -56,7 +53,6 @@ void FOpenGL4::ProcessExtensions( const FString& ExtensionsString )
 
 
 	bSupportsGPUMemoryInfo = ExtensionsString.Contains(TEXT("GL_NVX_gpu_memory_info"));
-	bSupportsComputeShaders = ExtensionsString.Contains(TEXT("GL_ARB_compute_shader")) || (MajorVersion ==4 && MinorVersion >= 3) || (MajorVersion > 4);
 	bSupportsVertexAttribBinding = ExtensionsString.Contains(TEXT("GL_ARB_vertex_attrib_binding")) || (MajorVersion == 4 && MinorVersion >= 3) || (MajorVersion > 4);
 	bSupportsTextureView = ExtensionsString.Contains(TEXT("GL_ARB_texture_view")) || (MajorVersion == 4 && MinorVersion >= 3) || (MajorVersion > 4);
 
diff --git a/Engine/Source/Runtime/OpenGLDrv/Private/OpenGLBuffer.cpp b/Engine/Source/Runtime/OpenGLDrv/Private/OpenGLBuffer.cpp
--- a/Engine/Source/Runtime/OpenGLDrv/Private/OpenGLBuffer.cpp
+++ b/Engine/Source/Runtime/OpenGLDrv/Private/OpenGLBuffer.cpp
@@ -192,7 +192,17 @@ FBufferRHIRef FOpenGLDynamicRHI::RHICreateBuffer(uint32 Size, EBufferUsageFlags
 		Data = CreateInfo.ResourceArray->GetResourceData();
 	}
 
-	TRefCountPtr<FOpenGLBuffer> Buffer = new FOpenGLBuffer((Usage & BUF_IndexBuffer)? GL_ARRAY_BUFFER : GL_ELEMENT_ARRAY_BUFFER, Stride, Size, Usage, Data);
+	GLenum BufferType = GL_ARRAY_BUFFER;
+	if (Usage & BUF_StructuredBuffer)
+	{
+		BufferType = GL_SHADER_STORAGE_BUFFER;
+	}
+	else if (Usage & BUF_IndexBuffer)
+	{
+		BufferType = GL_ELEMENT_ARRAY_BUFFER;
+	}
+
+	TRefCountPtr<FOpenGLBuffer> Buffer = new FOpenGLBuffer(BufferType, Stride, Size, Usage, Data);
 	
 	if (CreateInfo.ResourceArray)
 	{
diff --git a/Engine/Source/Runtime/OpenGLDrv/Private/OpenGLCommands.cpp b/Engine/Source/Runtime/OpenGLDrv/Private/OpenGLCommands.cpp
--- a/Engine/Source/Runtime/OpenGLDrv/Private/OpenGLCommands.cpp
+++ b/Engine/Source/Runtime/OpenGLDrv/Private/OpenGLCommands.cpp
@@ -560,21 +560,32 @@ void FOpenGLDynamicRHI::RHISetUAVParameter(FRHIPixelShader* PixelShaderRHI, uint
 }
 
 
-void FOpenGLDynamicRHI::RHISetUAVParameter(FRHIComputeShader* ComputeShaderRHI,uint32 UAVIndex, FRHIUnorderedAccessView* UnorderedAccessViewRHI)
+void FOpenGLDynamicRHI::RHISetUAVParameter(FRHIComputeShader* ComputeShaderRHI, uint32 InUAVIndex, FRHIUnorderedAccessView* UnorderedAccessViewRHI)
 {
-	check(FOpenGL::SupportsComputeShaders());
-	
 	VERIFY_GL_SCOPE();
-	if(UnorderedAccessViewRHI)
+
+	GLint UAVIndex = FOpenGL::GetFirstComputeUAVUnit() + InUAVIndex;
+	
+	if (UnorderedAccessViewRHI == nullptr)
+	{
+		InternalSetShaderBufferUAV(UAVIndex, 0);
+		return;
+	}
+
+	FOpenGLUnorderedAccessView* UnorderedAccessView = ResourceCast(UnorderedAccessViewRHI);
+	if (UnorderedAccessView->Resource)
 	{
-		FOpenGLUnorderedAccessView* UnorderedAccessView = ResourceCast(UnorderedAccessViewRHI);
-		bool bLayered = UnorderedAccessView->IsLayered();
-		GLint Layer = UnorderedAccessView->GetLayer();
-		InternalSetShaderUAV(FOpenGL::GetFirstComputeUAVUnit() + UAVIndex, UnorderedAccessView->Format, UnorderedAccessView->Resource, bLayered, Layer, GL_READ_WRITE);
+		GLuint	Resource = UnorderedAccessView->Resource;
+		GLenum	Format = UnorderedAccessView->Format;
+		bool	bLayered = UnorderedAccessView->IsLayered();
+		GLint	Layer = UnorderedAccessView->GetLayer();
+		GLenum	Access = GL_READ_WRITE;
+		InternalSetShaderImageUAV(UAVIndex, Format, Resource, bLayered, Layer, Access);
 	}
 	else
 	{
-		InternalSetShaderUAV(FOpenGL::GetFirstComputeUAVUnit() + UAVIndex, GL_R32F, 0, false, 0, GL_READ_WRITE);
+		GLuint Resource = UnorderedAccessView->BufferResource;
+		InternalSetShaderBufferUAV(UAVIndex, Resource);
 	}
 }
 
@@ -832,31 +843,41 @@ void FOpenGLDynamicRHI::SetupTexturesForDraw( FOpenGLContextState& ContextState
 	SetupTexturesForDraw(ContextState, PendingState.BoundShaderState, FOpenGL::GetMaxCombinedTextureImageUnits());
 }
 
-void FOpenGLDynamicRHI::InternalSetShaderUAV(GLint UAVIndex, GLenum Format, GLuint Resource, bool bLayered, GLint Layer, GLenum Access)
+void FOpenGLDynamicRHI::InternalSetShaderImageUAV(GLint UAVIndex, GLenum Format, GLuint Resource, bool bLayered, GLint Layer, GLenum Access)
 {
 	VERIFY_GL_SCOPE();
 	PendingState.UAVs[UAVIndex].Format = Format;
 	PendingState.UAVs[UAVIndex].Resource = Resource;
+	PendingState.UAVs[UAVIndex].Access = Access;
 	PendingState.UAVs[UAVIndex].Layer = Layer;
 	PendingState.UAVs[UAVIndex].bLayered = bLayered;
-	PendingState.UAVs[UAVIndex].Access = Access;
+}
+
+void FOpenGLDynamicRHI::InternalSetShaderBufferUAV(GLint UAVIndex, GLuint Resource)
+{
+	VERIFY_GL_SCOPE();
+	PendingState.UAVs[UAVIndex].Format = 0;
+	PendingState.UAVs[UAVIndex].Resource = Resource;
+	PendingState.UAVs[UAVIndex].Access = GL_READ_WRITE;
+	PendingState.UAVs[UAVIndex].Layer = 0;
+	PendingState.UAVs[UAVIndex].bLayered = false;
 }
 
 void FOpenGLDynamicRHI::SetupUAVsForDraw(FOpenGLContextState& ContextState)
 {
 	int32 MaxUAVUnitUsed = 0;
 	const TBitArray<>& NeededBits = PendingState.BoundShaderState->GetUAVNeeds(MaxUAVUnitUsed);
-	SetupUAVsForProgram(ContextState, NeededBits, MaxUAVUnitUsed, FOpenGL::GetMaxPixelUAVUnits());
+	SetupUAVsForProgram(ContextState, NeededBits, MaxUAVUnitUsed);
 }
 
 void FOpenGLDynamicRHI::SetupUAVsForCompute(FOpenGLContextState& ContextState, const FOpenGLComputeShader* ComputeShader)
 {
 	int32 MaxUAVUnitUsed = 0;
 	const TBitArray<>& NeededBits = ComputeShader->GetUAVNeeds(MaxUAVUnitUsed);
-	SetupUAVsForProgram(ContextState, NeededBits, MaxUAVUnitUsed, FOpenGL::GetMaxComputeUAVUnits());
+	SetupUAVsForProgram(ContextState, NeededBits, MaxUAVUnitUsed);
 }
 
-void FOpenGLDynamicRHI::SetupUAVsForProgram(FOpenGLContextState& ContextState, const TBitArray<>& NeededBits, int32 MaxUAVUnitUsed, int32 MaxUAVUnits)
+void FOpenGLDynamicRHI::SetupUAVsForProgram(FOpenGLContextState& ContextState, const TBitArray<>& NeededBits, int32 MaxUAVUnitUsed)
 {
 	if (MaxUAVUnitUsed < 0 && ContextState.ActiveUAVMask == 0)
 	{
@@ -868,7 +889,7 @@ void FOpenGLDynamicRHI::SetupUAVsForProgram(FOpenGLContextState& ContextState, c
 	{
 		if (!NeededBits[UAVStageIndex])
 		{
-			CachedSetupUAVStage(ContextState, UAVStageIndex, GL_R32F, 0, false, 0, GL_READ_WRITE);
+			CachedSetupUAVStage(ContextState, UAVStageIndex, 0, 0, false, 0, GL_READ_WRITE);
 		}
 		else
 		{
@@ -878,9 +899,14 @@ void FOpenGLDynamicRHI::SetupUAVsForProgram(FOpenGLContextState& ContextState, c
 	}
 
 	// clear rest of the units
-	for (int32 UAVStageIndex = (MaxUAVUnitUsed+1); UAVStageIndex < MaxUAVUnits; ++UAVStageIndex)
+	int32 UAVStageIndex = (MaxUAVUnitUsed + 1);
+	if ((ContextState.ActiveUAVMask >> UAVStageIndex) != 0)
 	{
-		CachedSetupUAVStage(ContextState, UAVStageIndex, GL_R32F, 0, false, 0, GL_READ_WRITE);
+		const int32 NumUAVs = ContextState.UAVs.Num();
+		for (; UAVStageIndex < NumUAVs; ++UAVStageIndex)
+		{
+			CachedSetupUAVStage(ContextState, UAVStageIndex, 0, 0, false, 0, GL_READ_WRITE);
+		}
 	}
 }
 
@@ -914,6 +940,8 @@ static bool IsImageTextureFormatSupported(GLenum Format)
 
 void FOpenGLDynamicRHI::CachedSetupUAVStage( FOpenGLContextState& ContextState, GLint UAVIndex, GLenum Format, GLuint Resource, bool bLayered, GLint Layer, GLenum Access)
 {
+	VERIFY_GL_SCOPE();
+	
 	FUAVStage& UAVStage = ContextState.UAVs[UAVIndex];	
 
 	if (UAVStage.Format == Format && 
@@ -925,19 +953,71 @@ void FOpenGLDynamicRHI::CachedSetupUAVStage( FOpenGLContextState& ContextState,
 		// Nothing's changed, no need to update
 		return;
 	}
+
+	// unbind any SSBO or Image in this slot
+	if (Resource == 0)
+	{
+		if (UAVStage.Resource != 0)
+		{
+			// SSBO
+			if (UAVStage.Format == 0)
+			{
+				FOpenGL::BindBufferBase(GL_SHADER_STORAGE_BUFFER, UAVIndex, 0);
+				ContextState.StorageBufferBound = 0;
+			}
+			else // Image
+			{
+				FOpenGL::BindImageTexture(UAVIndex, 0, 0, GL_FALSE, 0, GL_READ_WRITE, GL_R32F);
+			}
+
+			UAVStage.Format = 0;
+			UAVStage.Resource = 0;
+			UAVStage.Access = GL_READ_WRITE;
+			UAVStage.Layer = 0;
+			UAVStage.bLayered = false;
+		}
+	}
+	else
+	{
+		// SSBO
+		if (Format == 0)
+		{
+			// make sure we dont end up binding both SSBO and Image to the same UAV slot
+			if (UAVStage.Resource != 0 && UAVStage.Format != 0)
+			{
+				FOpenGL::BindImageTexture(UAVIndex, 0, 0, GL_FALSE, 0, GL_READ_WRITE, GL_R32F);
+			}
+
+			FOpenGL::BindBufferBase(GL_SHADER_STORAGE_BUFFER, UAVIndex, Resource);
+			
+			UAVStage.Format = 0;
+			UAVStage.Resource = Resource;
+			UAVStage.Access = GL_READ_WRITE;
+			UAVStage.Layer = 0;
+			UAVStage.bLayered = false;
+			ContextState.StorageBufferBound = Resource;
+		}
+		else // Image
+		{
+			// make sure we dont end up binding both SSBO and Image to the same UAV slot
+			if (UAVStage.Resource != 0 && UAVStage.Format == 0)
+			{
+				FOpenGL::BindBufferBase(GL_SHADER_STORAGE_BUFFER, UAVIndex, 0);
+				ContextState.StorageBufferBound = 0;
+			}
+			
+			check(IsImageTextureFormatSupported(Format));
 	
-	check(IsImageTextureFormatSupported(Format));
-		
-	VERIFY_GL_SCOPE();
+			FOpenGL::BindImageTexture(UAVIndex, Resource, 0, bLayered ? GL_TRUE : GL_FALSE, Layer, Access, Format);
 	
-	FOpenGL::BindImageTexture(UAVIndex, Resource, 0, bLayered ? GL_TRUE : GL_FALSE, Layer, Access, Format);
+			UAVStage.Format = Format;
+			UAVStage.Resource = Resource;
+			UAVStage.Access = Access;
+			UAVStage.Layer = Layer;
+			UAVStage.bLayered = bLayered;
+		}
+	}
 	
-	UAVStage.Format = Format;
-	UAVStage.Resource = Resource;
-	UAVStage.Access = Access;
-	UAVStage.Layer = Layer;
-	UAVStage.bLayered = bLayered;
-
 	uint32 UAVBit = 1 << UAVIndex;
 	if (Resource != 0)
 	{
@@ -1006,29 +1086,34 @@ void FOpenGLDynamicRHI::RHISetShaderResourceViewParameter(FRHIGraphicsShader* Sh
 {
 	VERIFY_GL_SCOPE();
 
-#ifndef __EMSCRIPTEN__
-	// TODO: On WebGL1 builds, the following assert() always comes out false, however when simply ignoring this check,
-	// everything seems to be working fine. That said, I don't know what should change here, shader resource views are a D3D abstraction,
-	// but I think InternalSetShaderTexture() and RHISetShaderSampler() calls below need to occur even on GLES2.
-	check(FOpenGL::SupportsResourceView());
-#endif
-
 	GLint Index = 0;
 	GLint MaxUnits = 0;
 	GetShaderStageIndexAndMaxUnits(ShaderRHI, Index, MaxUnits);
 
 	ensureMsgf((int32)TextureIndex < MaxUnits, TEXT("Using more texture units (%d) than allowed (%d) on Frequency %d!"), TextureIndex, MaxUnits, (int32)ShaderRHI->GetFrequency());
 	FOpenGLShaderResourceView* SRV = ResourceCast(SRVRHI);
+	
 	GLuint Resource = 0;
 	GLenum Target = GL_TEXTURE_BUFFER;
 	int32 LimitMip = -1;
-	if( SRV )
+	if (SRV)
 	{
-		Resource = SRV->Resource;
 		Target = SRV->Target;
-		LimitMip = SRV->LimitMip;
-		UpdateSRV(SRV);
+		
+		if (Target == GL_SHADER_STORAGE_BUFFER)
+		{
+			Index = (ShaderRHI->GetFrequency() == SF_Pixel) ? FOpenGL::GetFirstPixelUAVUnit() :  FOpenGL::GetFirstVertexUAVUnit();
+			InternalSetShaderBufferUAV(Index + TextureIndex, SRV->Resource);
+			return;
+		}
+		else
+		{
+			Resource = SRV->Resource;
+			LimitMip = SRV->LimitMip;
+			UpdateSRV(SRV);
+		}
 	}
+
 	ensureMsgf((int32)TextureIndex < MaxUnits, TEXT("Using more textures (%d) than allowed (%d)!"), TextureIndex, MaxUnits);
 	InternalSetShaderTexture(NULL, SRV, Index + TextureIndex, Target, Resource, 0, LimitMip);
 	RHISetShaderSampler(ShaderRHI, TextureIndex, PointSamplerState);
@@ -1037,19 +1122,27 @@ void FOpenGLDynamicRHI::RHISetShaderResourceViewParameter(FRHIGraphicsShader* Sh
 void FOpenGLDynamicRHI::RHISetShaderResourceViewParameter(FRHIComputeShader* ComputeShaderRHI,uint32 TextureIndex, FRHIShaderResourceView* SRVRHI)
 {
 	VERIFY_GL_SCOPE();
-	check(FOpenGL::SupportsComputeShaders());
-	check(FOpenGL::SupportsResourceView());
 	FOpenGLShaderResourceView* SRV = ResourceCast(SRVRHI);
 	GLuint Resource = 0;
 	GLenum Target = GL_TEXTURE_BUFFER;
 	int32 LimitMip = -1;
-	if( SRV )
+	if (SRV)
 	{
-		Resource = SRV->Resource;
 		Target = SRV->Target;
-		LimitMip = SRV->LimitMip;
-		UpdateSRV(SRV);
+		
+		if (Target == GL_SHADER_STORAGE_BUFFER)
+		{
+			InternalSetShaderBufferUAV(FOpenGL::GetFirstComputeUAVUnit() + TextureIndex, SRV->Resource);
+			return;
+		}
+		else
+		{
+			Resource = SRV->Resource;
+			LimitMip = SRV->LimitMip;
+			UpdateSRV(SRV);
+		}
 	}
+
 	ensureMsgf((int32)TextureIndex < FOpenGL::GetMaxComputeTextureImageUnits(), TEXT("Using more compute texture units (%d) than allowed (%d)!"), TextureIndex, FOpenGL::GetMaxComputeTextureImageUnits());
 	InternalSetShaderTexture(NULL, SRV, FOpenGL::GetFirstComputeTextureUnit() + TextureIndex, Target, Resource, 0, LimitMip);
 	RHISetShaderSampler(ComputeShaderRHI,TextureIndex,PointSamplerState);
@@ -1089,7 +1182,6 @@ void FOpenGLDynamicRHI::RHISetShaderSampler(FRHIGraphicsShader* ShaderRHI,uint32
 
 void FOpenGLDynamicRHI::RHISetShaderTexture(FRHIComputeShader* ComputeShaderRHI,uint32 TextureIndex, FRHITexture* NewTextureRHI)
 {
-	check(FOpenGL::SupportsComputeShaders());
 	VERIFY_GL_SCOPE();
 	FOpenGLTextureBase* NewTexture = GetOpenGLTextureFromRHITexture(NewTextureRHI);
 	ensureMsgf((int32)TextureIndex < FOpenGL::GetMaxComputeTextureImageUnits(), TEXT("Using more compute texture units (%d) than allowed (%d)!"), TextureIndex, FOpenGL::GetMaxComputeTextureImageUnits());
@@ -1127,7 +1219,6 @@ void FOpenGLDynamicRHI::RHISetShaderUniformBuffer(FRHIGraphicsShader* ShaderRHI,
 
 void FOpenGLDynamicRHI::RHISetShaderSampler(FRHIComputeShader* ComputeShaderRHI,uint32 SamplerIndex, FRHISamplerState* NewStateRHI)
 {
-	check(FOpenGL::SupportsComputeShaders());
 	VERIFY_GL_SCOPE();
 	FOpenGLSamplerState* NewState = ResourceCast(NewStateRHI);
 	InternalSetSamplerStates(FOpenGL::GetFirstComputeTextureUnit() + SamplerIndex, NewState);
@@ -1136,7 +1227,6 @@ void FOpenGLDynamicRHI::RHISetShaderSampler(FRHIComputeShader* ComputeShaderRHI,
 void FOpenGLDynamicRHI::RHISetShaderUniformBuffer(FRHIComputeShader* ComputeShaderRHI,uint32 BufferIndex, FRHIUniformBuffer* BufferRHI)
 {
 	VERIFY_GL_SCOPE();
-	check(FOpenGL::SupportsComputeShaders());
 	PendingState.BoundUniformBuffers[SF_Compute][BufferIndex] = BufferRHI;
 	PendingState.DirtyUniformBuffers[SF_Compute] |= 1 << BufferIndex;
 }
@@ -1155,7 +1245,6 @@ void FOpenGLDynamicRHI::RHISetShaderParameter(FRHIGraphicsShader* ShaderRHI,uint
 void FOpenGLDynamicRHI::RHISetShaderParameter(FRHIComputeShader* ComputeShaderRHI,uint32 BufferIndex,uint32 BaseIndex,uint32 NumBytes,const void* NewValue)
 { 
 	VERIFY_GL_SCOPE();
-	check(FOpenGL::SupportsComputeShaders());
 	PendingState.ShaderParameters[CrossCompiler::SHADER_STAGE_COMPUTE].Set(BufferIndex, BaseIndex, NumBytes, NewValue);
 	PendingState.LinkedProgramAndDirtyFlag = nullptr;
 }
@@ -2203,6 +2292,16 @@ void FOpenGLDynamicRHI::OnBufferDeletion( GLuint BufferResource )
 		RenderingContextState.ArrayBufferBound = -1;	// will force refresh
 	}
 
+	if (SharedContextState.StorageBufferBound == BufferResource)
+	{
+		SharedContextState.StorageBufferBound = -1;	// will force refresh
+	}
+
+	if (RenderingContextState.StorageBufferBound == BufferResource)
+	{
+		RenderingContextState.StorageBufferBound = -1;	// will force refresh
+	}
+
 	if (FOpenGL::SupportsVertexAttribBinding() && OpenGLConsoleVariables::bUseVAB)
 	{
 		// loop through active streams
@@ -2257,6 +2356,18 @@ void FOpenGLDynamicRHI::OnBufferDeletion( GLuint BufferResource )
 			}
 		}
 	}
+	
+	// Storage buffer
+	{
+		int32 NumStorageBuffers = PendingState.UAVs.Num();
+		for (int32 StorageBufferIndex = 0; StorageBufferIndex < NumStorageBuffers; StorageBufferIndex++)
+		{
+			if (PendingState.UAVs[StorageBufferIndex].Format == 0 && PendingState.UAVs[StorageBufferIndex].Resource == BufferResource)
+			{
+				PendingState.UAVs[StorageBufferIndex].Resource = 0;
+			}
+		}
+	}
 
 	if (SharedContextState.ElementArrayBufferBound == BufferResource)
 	{
@@ -2375,7 +2486,6 @@ void FOpenGLDynamicRHI::CommitNonComputeShaderConstantsFastPath(FOpenGLLinkedPro
 void FOpenGLDynamicRHI::CommitComputeShaderConstants(FOpenGLComputeShader* ComputeShader)
 {
 	VERIFY_GL_SCOPE();
-	check(FOpenGL::SupportsComputeShaders());
 
 	const int32 Stage = CrossCompiler::SHADER_STAGE_COMPUTE;
 	FOpenGLShaderParameterCache& StageShaderParameters = PendingState.ShaderParameters[Stage];
@@ -2410,6 +2520,7 @@ template <> FORCEINLINE uint32 GetNumTextureUnits<SF_Compute>() { return FOpenGL
 
 template <EShaderFrequency Frequency>
 uint32 GetFirstUAVUnit() { return 0; }
+template <> FORCEINLINE uint32 GetFirstUAVUnit<SF_Vertex>() { return FOpenGL::GetFirstVertexUAVUnit(); }
 template <> FORCEINLINE uint32 GetFirstUAVUnit<SF_Pixel>() { return FOpenGL::GetFirstPixelUAVUnit(); }
 template <> FORCEINLINE uint32 GetFirstUAVUnit<SF_Compute>() { return FOpenGL::GetFirstComputeUAVUnit(); }
 
@@ -2417,6 +2528,7 @@ template <EShaderFrequency Frequency>
 uint32 GetNumUAVUnits() { return 0; }
 template <> FORCEINLINE uint32 GetNumUAVUnits<SF_Compute>()	{ return FOpenGL::GetMaxComputeUAVUnits(); }
 template <> FORCEINLINE uint32 GetNumUAVUnits<SF_Pixel>()	{ return FOpenGL::GetMaxPixelUAVUnits(); }
+template <> FORCEINLINE uint32 GetNumUAVUnits<SF_Vertex>()	{ return FOpenGL::GetMaxPixelUAVUnits(); }
 
 template <EShaderFrequency Frequency>
 void SetResource(FOpenGLDynamicRHI* RESTRICT OpenGLRHI, uint32 BindIndex, FRHITexture* RESTRICT TextureRHI)
@@ -2451,8 +2563,15 @@ FORCEINLINE void SetResource(FOpenGLDynamicRHI* RESTRICT OpenGLRHI, uint32 BindI
 	ensureMsgf(BindIndex < GetNumTextureUnits<Frequency>(), TEXT("Using more %s texture units (%d) than allowed (%d) on a shader unit!"), GetShaderFrequencyString(Frequency, false), BindIndex, GetNumTextureUnits<Frequency>());	VERIFY_GL_SCOPE();
 #endif
 	auto SRV = FOpenGLDynamicRHI::ResourceCast(RHISRV);
-	OpenGLRHI->InternalSetShaderTexture(NULL, SRV, GetFirstTextureUnit<Frequency>() + BindIndex, SRV->Target, SRV->Resource, 0, SRV->LimitMip);
-	SetResource<Frequency>(OpenGLRHI,BindIndex,OpenGLRHI->GetPointSamplerState());
+	if (SRV->Target == GL_SHADER_STORAGE_BUFFER)
+	{
+		OpenGLRHI->InternalSetShaderBufferUAV(GetFirstUAVUnit<Frequency>() + BindIndex, SRV->Resource);
+	}
+	else
+	{
+		OpenGLRHI->InternalSetShaderTexture(NULL, SRV, GetFirstTextureUnit<Frequency>() + BindIndex, SRV->Target, SRV->Resource, 0, SRV->LimitMip);
+		SetResource<Frequency>(OpenGLRHI,BindIndex,OpenGLRHI->GetPointSamplerState());
+	}
 }
 
 template <EShaderFrequency Frequency>
@@ -2462,11 +2581,20 @@ FORCEINLINE void SetResource(FOpenGLDynamicRHI* RESTRICT OpenGLRHI, uint32 BindI
 	ensureMsgf(BindIndex < GetNumUAVUnits<Frequency>(), TEXT("Using more %s image units (%d) than allowed (%d) on a shader unit!"), GetShaderFrequencyString(Frequency, false), BindIndex, GetNumUAVUnits<Frequency>());	VERIFY_GL_SCOPE();
 #endif
 	auto UAV = FOpenGLDynamicRHI::ResourceCast(RHIUAV);
-	GLenum Access = (Frequency == SF_Compute) ? GL_READ_WRITE : GL_WRITE_ONLY;
-	// TODO: This must be true for 3D textures
-	bool bLayered = false;
-	GLint Layer = 0;
-	OpenGLRHI->InternalSetShaderUAV(GetFirstUAVUnit<Frequency>() + BindIndex, UAV->Format, UAV->Resource, bLayered, Layer, Access);
+	GLint UAVIndex = GetFirstUAVUnit<Frequency>() + BindIndex;
+	
+	if (UAV->Resource)
+	{
+		GLenum Access = (Frequency == SF_Compute) ? GL_READ_WRITE : GL_WRITE_ONLY;
+		// TODO: This must be true for 3D textures
+		bool bLayered = false;
+		GLint Layer = 0;
+		OpenGLRHI->InternalSetShaderImageUAV(UAVIndex, UAV->Format, UAV->Resource, bLayered, Layer, Access);
+	}
+	else
+	{
+		OpenGLRHI->InternalSetShaderBufferUAV(UAVIndex, UAV->BufferResource);
+	}
 }
 
 template <class GLResourceType, EShaderFrequency ShaderFrequency>
@@ -3185,14 +3313,7 @@ void FOpenGLDynamicRHI::RHISetComputeShader(FRHIComputeShader* ComputeShaderRHI)
 		return;
 	}
 
-	if (FOpenGL::SupportsComputeShaders())
-	{
-		PendingState.CurrentComputeShader = ComputeShaderRHI;
-	}
-	else
-	{
-		UE_LOG(LogRHI,Fatal,TEXT("Platform doesn't support SM5 for OpenGL but set feature level to SM5"));
-	}
+	PendingState.CurrentComputeShader = ComputeShaderRHI;
 
 	ApplyStaticUniformBuffers(ComputeShaderRHI, ResourceCast(ComputeShaderRHI));
 }
@@ -3204,83 +3325,69 @@ void FOpenGLDynamicRHI::RHIDispatchComputeShader(uint32 ThreadGroupCountX, uint3
 		return;
 	}
 
-	if (FOpenGL::SupportsComputeShaders())
-	{
-		VERIFY_GL_SCOPE();
+
+	VERIFY_GL_SCOPE();
 		
-		FRHIComputeShader* ComputeShaderRHI = PendingState.CurrentComputeShader;
-		check(ComputeShaderRHI);
+	FRHIComputeShader* ComputeShaderRHI = PendingState.CurrentComputeShader;
+	check(ComputeShaderRHI);
 
-		FOpenGLComputeShader* ComputeShader = ResourceCast(ComputeShaderRHI);
+	FOpenGLComputeShader* ComputeShader = ResourceCast(ComputeShaderRHI);
 
-		if (ComputeShader->LinkedProgram == nullptr)
-		{
-			ComputeShader->LinkedProgram = GetLinkedComputeProgram(ComputeShaderRHI);
-		}
-		FOpenGLContextState& ContextState = GetContextStateForCurrentContext();
+	if (ComputeShader->LinkedProgram == nullptr)
+	{
+		ComputeShader->LinkedProgram = GetLinkedComputeProgram(ComputeShaderRHI);
+	}
+	FOpenGLContextState& ContextState = GetContextStateForCurrentContext();
 
-		GPUProfilingData.RegisterGPUDispatch(FIntVector(ThreadGroupCountX, ThreadGroupCountY, ThreadGroupCountZ));	
+	GPUProfilingData.RegisterGPUDispatch(FIntVector(ThreadGroupCountX, ThreadGroupCountY, ThreadGroupCountZ));	
 
-		BindPendingComputeShaderState(ContextState, ComputeShader);
-		CommitComputeResourceTables(ComputeShader);
-		SetupTexturesForDraw(ContextState, ComputeShader, FOpenGL::GetMaxComputeTextureImageUnits());
-		SetupUAVsForCompute(ContextState, ComputeShader);
-		CommitComputeShaderConstants(ComputeShader);
+	BindPendingComputeShaderState(ContextState, ComputeShader);
+	CommitComputeResourceTables(ComputeShader);
+	SetupTexturesForDraw(ContextState, ComputeShader, FOpenGL::GetMaxComputeTextureImageUnits());
+	SetupUAVsForCompute(ContextState, ComputeShader);
+	CommitComputeShaderConstants(ComputeShader);
 	
-		FOpenGL::MemoryBarrier(GL_ALL_BARRIER_BITS);
-		FOpenGL::DispatchCompute(ThreadGroupCountX, ThreadGroupCountY, ThreadGroupCountZ);
-		FOpenGL::MemoryBarrier(GL_ALL_BARRIER_BITS);
-	}
-	else
-	{
-		UE_LOG(LogRHI,Fatal,TEXT("Platform doesn't support SM5 for OpenGL but set feature level to SM5"));
-	}
+	FOpenGL::MemoryBarrier(GL_ALL_BARRIER_BITS);
+	FOpenGL::DispatchCompute(ThreadGroupCountX, ThreadGroupCountY, ThreadGroupCountZ);
+	FOpenGL::MemoryBarrier(GL_ALL_BARRIER_BITS);
 }
 
 void FOpenGLDynamicRHI::RHIDispatchIndirectComputeShader(FRHIBuffer* ArgumentBufferRHI, uint32 ArgumentOffset)
 {
-	if (FOpenGL::SupportsComputeShaders())
-	{
-		VERIFY_GL_SCOPE();
+	VERIFY_GL_SCOPE();
 		
-		FRHIComputeShader* ComputeShaderRHI = PendingState.CurrentComputeShader;
-		check(ComputeShaderRHI);
+	FRHIComputeShader* ComputeShaderRHI = PendingState.CurrentComputeShader;
+	check(ComputeShaderRHI);
 
-		FOpenGLComputeShader* ComputeShader = ResourceCast(ComputeShaderRHI);
-		if (ComputeShader->LinkedProgram == nullptr)
-		{
-			ComputeShader->LinkedProgram = GetLinkedComputeProgram(ComputeShaderRHI);
-		}
+	FOpenGLComputeShader* ComputeShader = ResourceCast(ComputeShaderRHI);
+	if (ComputeShader->LinkedProgram == nullptr)
+	{
+		ComputeShader->LinkedProgram = GetLinkedComputeProgram(ComputeShaderRHI);
+	}
 
-		FOpenGLBuffer* ArgumentBuffer = ResourceCast(ArgumentBufferRHI);
+	FOpenGLBuffer* ArgumentBuffer = ResourceCast(ArgumentBufferRHI);
 
-		FOpenGLContextState& ContextState = GetContextStateForCurrentContext();
+	FOpenGLContextState& ContextState = GetContextStateForCurrentContext();
 
-		GPUProfilingData.RegisterGPUDispatch(FIntVector(1, 1, 1));	
+	GPUProfilingData.RegisterGPUDispatch(FIntVector(1, 1, 1));	
 
-		BindPendingComputeShaderState(ContextState, ComputeShader);
+	BindPendingComputeShaderState(ContextState, ComputeShader);
 
-		SetupTexturesForDraw(ContextState, ComputeShader, FOpenGL::GetMaxComputeTextureImageUnits());
+	SetupTexturesForDraw(ContextState, ComputeShader, FOpenGL::GetMaxComputeTextureImageUnits());
 
-		SetupUAVsForCompute(ContextState, ComputeShader);
+	SetupUAVsForCompute(ContextState, ComputeShader);
 	
-		CommitComputeShaderConstants(ComputeShader);
+	CommitComputeShaderConstants(ComputeShader);
 	
-		FOpenGL::MemoryBarrier(GL_ALL_BARRIER_BITS);
+	FOpenGL::MemoryBarrier(GL_ALL_BARRIER_BITS);
 
-		glBindBuffer( GL_DISPATCH_INDIRECT_BUFFER, ArgumentBuffer->Resource);
+	glBindBuffer( GL_DISPATCH_INDIRECT_BUFFER, ArgumentBuffer->Resource);
 	
-		FOpenGL::DispatchComputeIndirect(ArgumentOffset);
+	FOpenGL::DispatchComputeIndirect(ArgumentOffset);
 
-		glBindBuffer( GL_DISPATCH_INDIRECT_BUFFER, 0);
+	glBindBuffer( GL_DISPATCH_INDIRECT_BUFFER, 0);
 	
-		FOpenGL::MemoryBarrier(GL_ALL_BARRIER_BITS);
-
-	}
-	else
-	{
-		UE_LOG(LogRHI,Fatal,TEXT("Platform doesn't support SM5 for OpenGL but set feature level to SM5"));
-	}
+	FOpenGL::MemoryBarrier(GL_ALL_BARRIER_BITS);
 }
 
 void FOpenGLDynamicRHI::RHISetMultipleViewports(uint32 Count, const FViewportBounds* Data)
@@ -3336,8 +3443,10 @@ void FOpenGLDynamicRHI::RHIInvalidateCachedState()
 	RenderingContextState = FOpenGLContextState();
 	SharedContextState = FOpenGLContextState();
 
-	RenderingContextState.InitializeResources(FOpenGL::GetMaxCombinedTextureImageUnits(), FOpenGL::GetMaxCombinedUAVUnits());
-	SharedContextState.InitializeResources(FOpenGL::GetMaxCombinedTextureImageUnits(), FOpenGL::GetMaxCombinedUAVUnits());
+	GLint NumUAVUnits = FMath::Max(FOpenGL::GetMaxCombinedUAVUnits(), FOpenGL::GetMaxComputeUAVUnits());
+
+	RenderingContextState.InitializeResources(FOpenGL::GetMaxCombinedTextureImageUnits(), NumUAVUnits);
+	SharedContextState.InitializeResources(FOpenGL::GetMaxCombinedTextureImageUnits(), NumUAVUnits);
 }
 
 void FOpenGLDynamicRHI::RHICopyToStagingBuffer(FRHIBuffer* SourceBufferRHI, FRHIStagingBuffer* DestinationStagingBufferRHI, uint32 InOffset, uint32 InNumBytes)
diff --git a/Engine/Source/Runtime/OpenGLDrv/Private/OpenGLDevice.cpp b/Engine/Source/Runtime/OpenGLDrv/Private/OpenGLDevice.cpp
--- a/Engine/Source/Runtime/OpenGLDrv/Private/OpenGLDevice.cpp
+++ b/Engine/Source/Runtime/OpenGLDrv/Private/OpenGLDevice.cpp
@@ -111,6 +111,10 @@ void CachedBindBuffer( GLenum Type, GLuint Buffer )
 	{
 		PrivateOpenGLDevicePtr->CachedBindElementArrayBuffer(PrivateOpenGLDevicePtr->GetContextStateForCurrentContext(), Buffer);
 	}
+	else if (Type == GL_SHADER_STORAGE_BUFFER)
+	{
+		PrivateOpenGLDevicePtr->CachedBindStorageBuffer(PrivateOpenGLDevicePtr->GetContextStateForCurrentContext(), Buffer);
+	}
 	else
 	{
 		checkNoEntry();
@@ -1216,7 +1220,7 @@ void FOpenGLDynamicRHI::Cleanup()
 	{
 		FOpenGL::DeleteBuffers(1, &PendingState.ZeroFilledDummyUniformBuffer);
 		PendingState.ZeroFilledDummyUniformBuffer = 0;
-		DecrementBufferMemory(GL_UNIFORM_BUFFER, false, ZERO_FILLED_DUMMY_UNIFORM_BUFFER_SIZE);
+		DecrementBufferMemory(GL_UNIFORM_BUFFER, ZERO_FILLED_DUMMY_UNIFORM_BUFFER_SIZE);
 	}
 
 	// Release pending shader
diff --git a/Engine/Source/Runtime/OpenGLDrv/Private/OpenGLES.cpp b/Engine/Source/Runtime/OpenGLDrv/Private/OpenGLES.cpp
--- a/Engine/Source/Runtime/OpenGLDrv/Private/OpenGLES.cpp
+++ b/Engine/Source/Runtime/OpenGLDrv/Private/OpenGLES.cpp
@@ -174,8 +174,9 @@ void FOpenGLES::ProcessQueryGLInt()
 
 	// clamp UAV units to a sensible limit
 	MaxCombinedUAVUnits = FMath::Min(MaxCombinedUAVUnits, 8);
-	MaxComputeUAVUnits = FMath::Min(MaxComputeUAVUnits, MaxCombinedUAVUnits);
-	MaxPixelUAVUnits = FMath::Min(MaxPixelUAVUnits, MaxCombinedUAVUnits);
+	MaxComputeUAVUnits = FMath::Min(MaxComputeUAVUnits, 16);
+	// this is split between VS and PS, 4 to each stage
+	MaxPixelUAVUnits = FMath::Min(MaxPixelUAVUnits, 4);
 
 	const GLint RequiredMaxVertexUniformComponents = 256;
 	if (MaxVertexUniformComponents < RequiredMaxVertexUniformComponents)
diff --git a/Engine/Source/Runtime/OpenGLDrv/Private/OpenGLShaders.cpp b/Engine/Source/Runtime/OpenGLDrv/Private/OpenGLShaders.cpp
--- a/Engine/Source/Runtime/OpenGLDrv/Private/OpenGLShaders.cpp
+++ b/Engine/Source/Runtime/OpenGLDrv/Private/OpenGLShaders.cpp
@@ -650,7 +650,7 @@ ShaderType* CompileOpenGLShader(TArrayView<const uint8> InShaderCode, const FSHA
 		|| (TypeEnum == GL_VERTEX_SHADER && Header.FrequencyMarker != 0x5653)
 		|| (TypeEnum == GL_FRAGMENT_SHADER && Header.FrequencyMarker != 0x5053)
 		|| (TypeEnum == GL_GEOMETRY_SHADER && Header.FrequencyMarker != 0x4753)
-		|| (TypeEnum == GL_COMPUTE_SHADER && Header.FrequencyMarker != 0x4353 && FOpenGL::SupportsComputeShaders())
+		|| (TypeEnum == GL_COMPUTE_SHADER && Header.FrequencyMarker != 0x4353)
 		|| (TypeEnum == GL_TESS_CONTROL_SHADER && Header.FrequencyMarker != 0x4853 && FOpenGL::SupportsTessellation()) /* hull shader*/
 		|| (TypeEnum == GL_TESS_EVALUATION_SHADER && Header.FrequencyMarker != 0x4453 && FOpenGL::SupportsTessellation()) /* domain shader*/
 		)
@@ -1140,7 +1140,7 @@ void FOpenGLDynamicRHI::BindUniformBufferBase(FOpenGLContextState& ContextState,
 				CachedBindUniformBuffer(ContextState,PendingState.ZeroFilledDummyUniformBuffer);
 				glBufferData(GL_UNIFORM_BUFFER, ZERO_FILLED_DUMMY_UNIFORM_BUFFER_SIZE, ZeroBuffer, GL_STATIC_DRAW);
 				FMemory::Free(ZeroBuffer);
-				IncrementBufferMemory(GL_UNIFORM_BUFFER, false, ZERO_FILLED_DUMMY_UNIFORM_BUFFER_SIZE);
+				IncrementBufferMemory(GL_UNIFORM_BUFFER, ZERO_FILLED_DUMMY_UNIFORM_BUFFER_SIZE);
 			}
 
 			Buffer = PendingState.ZeroFilledDummyUniformBuffer;
@@ -2237,7 +2237,7 @@ void FOpenGLLinkedProgram::ConfigureShaderStage( int Stage, uint32 FirstUniformB
 	};
 	static const GLint FirstUAVUnit[CrossCompiler::NUM_SHADER_STAGES] =
 	{
-		OGL_UAV_NOT_SUPPORTED_FOR_GRAPHICS_UNIT,
+		FOpenGL::GetFirstVertexUAVUnit(),
 		FOpenGL::GetFirstPixelUAVUnit(),
 		OGL_UAV_NOT_SUPPORTED_FOR_GRAPHICS_UNIT,
 		OGL_UAV_NOT_SUPPORTED_FOR_GRAPHICS_UNIT,
@@ -2374,8 +2374,19 @@ void FOpenGLLinkedProgram::ConfigureShaderStage( int Stage, uint32 FirstUniformB
 	int32 LastFoundUAVIndex = -1;
 	for (int32 UAVIndex = 0; UAVIndex < Config.Shaders[Stage].Bindings.NumUAVs; ++UAVIndex)
 	{
-		SetIndex(Name.Buffer, 2, UAVIndex);
+		ANSICHAR* Str = SetIndex(Name.Buffer, 2, UAVIndex);
 		GLint Location = glGetUniformLocation(StageProgram, Name.Buffer);
+		if (Location == -1)
+		{
+			// SSBO
+			Str[0] = '_';
+			Str[1] = 'V';
+			Str[2] = 'A';
+			Str[3] = 'R';
+			Str[4] = '\0';
+			Location = glGetProgramResourceIndex(StageProgram, GL_SHADER_STORAGE_BLOCK, Name.Buffer);
+		}
+
 		if (Location == -1)
 		{
 			if (LastFoundUAVIndex != -1)
@@ -2603,12 +2614,9 @@ static void VerifyUniformBufferLayouts(GLuint Program)
 #endif
 			}
 			
-			if (RHISupportsComputeShaders(GMaxRHIShaderPlatform))
-			{
 #ifdef GL_UNIFORM_BLOCK_REFERENCED_BY_COMPUTE_SHADER
-				glGetActiveUniformBlockiv(Program, BlockIndex, GL_UNIFORM_BLOCK_REFERENCED_BY_COMPUTE_SHADER, &ReferencedByCS);
+			glGetActiveUniformBlockiv(Program, BlockIndex, GL_UNIFORM_BLOCK_REFERENCED_BY_COMPUTE_SHADER, &ReferencedByCS);
 #endif
-			}
 
 			if(ReferencedByVS) {ReferencedBy += TEXT("V");}
 			if(ReferencedByHS) {ReferencedBy += TEXT("H");}
@@ -2860,7 +2868,6 @@ FOpenGLLinkedProgram* FOpenGLDynamicRHI::GetLinkedComputeProgram(FRHIComputeShad
 
 FComputeShaderRHIRef FOpenGLDynamicRHI::RHICreateComputeShader(TArrayView<const uint8> Code, const FSHAHash& Hash)
 {
-	check(RHISupportsComputeShaders(GMaxRHIShaderPlatform));
 	return CreateProxyShader<FRHIComputeShader, FOpenGLComputeShaderProxy>(Code, Hash);
 }
 
diff --git a/Engine/Source/Runtime/OpenGLDrv/Private/OpenGLState.cpp b/Engine/Source/Runtime/OpenGLDrv/Private/OpenGLState.cpp
--- a/Engine/Source/Runtime/OpenGLDrv/Private/OpenGLState.cpp
+++ b/Engine/Source/Runtime/OpenGLDrv/Private/OpenGLState.cpp
@@ -477,10 +477,7 @@ void FOpenGLRHIState::InitializeResources(int32 NumCombinedTextures, int32 NumCo
 	}
 
 	LinkedProgramAndDirtyFlag = nullptr;
-	if ( FOpenGL::SupportsComputeShaders() )
-	{
-		ShaderParameters[CrossCompiler::SHADER_STAGE_COMPUTE].InitializeResources(FOpenGL::GetMaxComputeUniformComponents() * 4 * sizeof(float));
-	}
+	ShaderParameters[CrossCompiler::SHADER_STAGE_COMPUTE].InitializeResources(FOpenGL::GetMaxComputeUniformComponents() * 4 * sizeof(float));
 
 	for (int32 Frequency = 0; Frequency < SF_NumStandardFrequencies; ++Frequency)
 	{
diff --git a/Engine/Source/Runtime/OpenGLDrv/Private/OpenGLTexture.cpp b/Engine/Source/Runtime/OpenGLDrv/Private/OpenGLTexture.cpp
--- a/Engine/Source/Runtime/OpenGLDrv/Private/OpenGLTexture.cpp
+++ b/Engine/Source/Runtime/OpenGLDrv/Private/OpenGLTexture.cpp
@@ -2420,7 +2420,9 @@ void FOpenGLDynamicRHI::InvalidateTextureResourceInCache(GLuint Resource)
 void FOpenGLDynamicRHI::InvalidateUAVResourceInCache(GLuint Resource)
 {
 	VERIFY_GL_SCOPE();
-	for (int32 UAVIndex = 0; UAVIndex < FOpenGL::GetMaxCombinedUAVUnits(); ++UAVIndex)
+	int32 NumUAVs = RenderingContextState.UAVs.Num();
+	
+	for (int32 UAVIndex = 0; UAVIndex < NumUAVs; ++UAVIndex)
 	{
 		if (SharedContextState.UAVs[UAVIndex].Resource == Resource)
 		{
diff --git a/Engine/Source/Runtime/OpenGLDrv/Private/OpenGLUAV.cpp b/Engine/Source/Runtime/OpenGLDrv/Private/OpenGLUAV.cpp
--- a/Engine/Source/Runtime/OpenGLDrv/Private/OpenGLUAV.cpp
+++ b/Engine/Source/Runtime/OpenGLDrv/Private/OpenGLUAV.cpp
@@ -82,18 +82,8 @@ FShaderResourceViewRHIRef FOpenGLDynamicRHI::RHICreateShaderResourceView(const F
 			return new FOpenGLShaderResourceViewProxy([=](FRHIShaderResourceView* OwnerRHI)
 			{
 				VERIFY_GL_SCOPE();
-				GLuint TextureID = 0;
-				if (FOpenGL::SupportsResourceView())
-				{
-					FOpenGLBuffer* StructuredBuffer = ResourceCast(BufferRHI);
-					FOpenGL::GenTextures(1, &TextureID);
-					CachedSetupTextureStage(GetContextStateForCurrentContext(), FOpenGL::GetMaxCombinedTextureImageUnits() - 1, GL_TEXTURE_BUFFER, TextureID, -1, 1);
-					uint32 Stride = StructuredBuffer->GetStride();
-					GLenum Format = (Stride == 4) ? GL_R32F : GL_RGBA32F;
-					BindGLTexBufferRange(GL_TEXTURE_BUFFER, Format, StructuredBuffer->Resource, Desc.StartOffsetBytes, Desc.NumElements, Stride);
-				}
-
-				return new FOpenGLShaderResourceView(this, TextureID, GL_TEXTURE_BUFFER);
+				FOpenGLBuffer* BufferGL = ResourceCast(BufferRHI);
+				return new FOpenGLShaderResourceView(this, BufferGL->Resource, GL_SHADER_STORAGE_BUFFER, BufferRHI);
 			});
 		}
 
@@ -127,70 +117,72 @@ FShaderResourceViewRHIRef FOpenGLDynamicRHI::RHICreateShaderResourceView(const F
 	}
 }
 
-void FOpenGLDynamicRHI::RHIUpdateShaderResourceView(FRHIShaderResourceView* SRV, FRHIBuffer* IndexBuffer)
+void FOpenGLDynamicRHI::RHIUpdateShaderResourceView(FRHIShaderResourceView* SRV, FRHIBuffer* InBuffer)
 {
-	if (!FOpenGL::SupportsResourceView())
-	{
-		return;
-	}
 	VERIFY_GL_SCOPE();
 
 	FOpenGLShaderResourceView* SRVGL = FOpenGLDynamicRHI::ResourceCast(SRV);
-	FOpenGLBuffer* IBGL = FOpenGLDynamicRHI::ResourceCast(IndexBuffer);
+	FOpenGLBuffer* BufferGL = FOpenGLDynamicRHI::ResourceCast(InBuffer);
 
 	check(SRVGL);
-	check(!SRVGL->VertexBuffer);
-	GLuint TextureID = SRVGL->Resource;
-	CachedSetupTextureStage(GetContextStateForCurrentContext(), FOpenGL::GetMaxCombinedTextureImageUnits() - 1, GL_TEXTURE_BUFFER, TextureID, -1, 1);
-
-	if (!IBGL)
+	if (SRVGL->Target == GL_TEXTURE_BUFFER)
 	{
-		FOpenGL::TexBuffer(GL_TEXTURE_BUFFER, GL_R16UI, 0); // format ignored here since we're detaching.
-		SRVGL->IndexBuffer = nullptr;
-		SRVGL->ModificationVersion = 0;
+		GLuint TextureID = SRVGL->Resource;
+		CachedSetupTextureStage(GetContextStateForCurrentContext(), FOpenGL::GetMaxCombinedTextureImageUnits() - 1, GL_TEXTURE_BUFFER, TextureID, -1, 1);
+		if (!BufferGL)
+		{
+			FOpenGL::TexBuffer(GL_TEXTURE_BUFFER, GL_R16UI, 0); // format ignored here since we're detaching.
+			SRVGL->BufferRHI = nullptr;
+			SRVGL->ModificationVersion = 0;
+		}
+		else
+		{
+			uint32 Stride = InBuffer->GetStride();
+			GLenum Format = (Stride == 2) ? GL_R16UI : GL_R32UI;
+			uint32 NumElements = InBuffer->GetSize() / Stride;
+			BindGLTexBufferRange(GL_TEXTURE_BUFFER, Format, BufferGL->Resource, 0, NumElements, Stride);
+			SRVGL->BufferRHI = InBuffer;
+			SRVGL->ModificationVersion = BufferGL->ModificationCount;
+		}
+	}
+	else if (SRVGL->Target == GL_SHADER_STORAGE_BUFFER)
+	{
+		SRVGL->Resource = BufferGL->Resource;
+		SRVGL->BufferRHI = InBuffer;
+		SRVGL->ModificationVersion = BufferGL->ModificationCount;
 	}
 	else
 	{
-		uint32 Stride = IndexBuffer->GetStride();
-		GLenum Format = (Stride == 2) ? GL_R16UI : GL_R32UI;
-		check(SRVGL->Target == GL_TEXTURE_BUFFER);
-
-		uint32 NumElements = IndexBuffer->GetSize() / Stride;
-		BindGLTexBufferRange(GL_TEXTURE_BUFFER, Format, IBGL->Resource, 0, NumElements, Stride);
-		SRVGL->IndexBuffer = IndexBuffer;
-		SRVGL->ModificationVersion = IBGL->ModificationCount;
+		checkNoEntry();
 	}
 }
 
-void FOpenGLDynamicRHI::RHIUpdateShaderResourceView(FRHIShaderResourceView* SRV, FRHIBuffer* VertexBuffer, uint32 Stride, uint8 Format)
+void FOpenGLDynamicRHI::RHIUpdateShaderResourceView(FRHIShaderResourceView* SRV, FRHIBuffer* InBuffer, uint32 Stride, uint8 Format)
 {
-	if (!FOpenGL::SupportsResourceView())
-	{
-		return;
-	}
 	VERIFY_GL_SCOPE();
 
 	FOpenGLShaderResourceView* SRVGL = FOpenGLDynamicRHI::ResourceCast(SRV);
-	FOpenGLBuffer* VBGL = FOpenGLDynamicRHI::ResourceCast(VertexBuffer);
+	FOpenGLBuffer* BufferGL = FOpenGLDynamicRHI::ResourceCast(InBuffer);
 	const FOpenGLTextureFormat& GLFormat = GOpenGLTextureFormats[Format];
 	
 	check(SRVGL);
-	check(!SRVGL->IndexBuffer);
+	check(SRVGL->Target == GL_TEXTURE_BUFFER); // add support for SSBO views?
+
 	GLuint TextureID = SRVGL->Resource;
 	CachedSetupTextureStage(GetContextStateForCurrentContext(), FOpenGL::GetMaxCombinedTextureImageUnits() - 1, GL_TEXTURE_BUFFER, TextureID, -1, 1);
 	
-	if (!VBGL)
+	if (!BufferGL)
 	{
 		FOpenGL::TexBuffer(GL_TEXTURE_BUFFER, GLFormat.InternalFormat[0], 0);
-		SRVGL->VertexBuffer = nullptr;
+		SRVGL->BufferRHI = nullptr;
 		SRVGL->ModificationVersion = 0;
 	}
 	else
 	{
 		check(SRVGL->Format == Format && SRVGL->Target == GL_TEXTURE_BUFFER);
-		FOpenGL::TexBuffer(GL_TEXTURE_BUFFER, GLFormat.InternalFormat[0], VBGL->Resource);
-		SRVGL->VertexBuffer = VertexBuffer;
-		SRVGL->ModificationVersion = VBGL->ModificationCount;
+		FOpenGL::TexBuffer(GL_TEXTURE_BUFFER, GLFormat.InternalFormat[0], BufferGL->Resource);
+		SRVGL->BufferRHI = InBuffer;
+		SRVGL->ModificationVersion = BufferGL->ModificationCount;
 	}
 }
 
@@ -212,29 +204,6 @@ FOpenGLShaderResourceView::~FOpenGLShaderResourceView()
 	}
 }
 
-FUnorderedAccessViewRHIRef FOpenGLDynamicRHI::RHICreateUnorderedAccessView(FRHIBuffer* BufferRHI, bool bUseUAVCounter, bool bAppendBuffer)
-{
-	FOpenGLBuffer* Buffer = ResourceCast(BufferRHI);
-	// emulate structured buffer of specific size as typed buffer
-	if (GMaxRHIFeatureLevel == ERHIFeatureLevel::ES3_1)
-	{
-		// ES3.1 cross-compiler converts StructuredBuffer<type4> into Buffer<type4> and StructuredBuffer<type> into Buffer<type>
-		// type can be float, int and uint
-		check(Buffer->GetStride() == 16 || Buffer->GetStride() == 4);
-		if (Buffer->GetStride() == 16)
-		{
-			return new FOpenGLStructuredBufferUnorderedAccessView(this, BufferRHI, PF_A32B32G32R32F);
-		}
-		else if (Buffer->GetStride() == 4)
-		{
-			return new FOpenGLStructuredBufferUnorderedAccessView(this, BufferRHI, PF_R32_FLOAT);
-		}
-	}
-	
-	UE_LOG(LogRHI, Fatal,TEXT("%s not implemented yet"),ANSI_TO_TCHAR(__FUNCTION__)); 
-	return new FOpenGLUnorderedAccessView();
-}
-
 FUnorderedAccessViewRHIRef FOpenGLDynamicRHI::RHICreateUnorderedAccessView(FRHITexture* TextureRHI, uint32 MipLevel)
 {
 	FOpenGLTexture* Texture = ResourceCast(TextureRHI);
@@ -242,7 +211,6 @@ FUnorderedAccessViewRHIRef FOpenGLDynamicRHI::RHICreateUnorderedAccessView(FRHIT
 	return new FOpenGLTextureUnorderedAccessView(TextureRHI);
 }
 
-
 FOpenGLTextureUnorderedAccessView::FOpenGLTextureUnorderedAccessView(FRHITexture* InTextureRHI):
 	TextureRHI(InTextureRHI)
 {
@@ -258,42 +226,44 @@ FOpenGLTextureUnorderedAccessView::FOpenGLTextureUnorderedAccessView(FRHITexture
 	this->bLayered = (Texture->Target == GL_TEXTURE_3D);
 }
 
+FUnorderedAccessViewRHIRef FOpenGLDynamicRHI::RHICreateUnorderedAccessView(FRHIBuffer* BufferRHI, uint8 Format)
+{
+	check(BufferRHI->GetUsage() & BUF_UnorderedAccess);
+	return new FOpenGLTexBufferUnorderedAccessView(this, BufferRHI, Format);
+}
 
-FOpenGLVertexBufferUnorderedAccessView::FOpenGLVertexBufferUnorderedAccessView(	FOpenGLDynamicRHI* InOpenGLRHI, FRHIBuffer* InVertexBufferRHI, uint8 Format):
-	VertexBufferRHI(InVertexBufferRHI),
+FOpenGLTexBufferUnorderedAccessView::FOpenGLTexBufferUnorderedAccessView(FOpenGLDynamicRHI* InOpenGLRHI, FRHIBuffer* InBufferRHI, uint8 Format):
+	BufferRHI(InBufferRHI),
 	OpenGLRHI(InOpenGLRHI)
 {
 	VERIFY_GL_SCOPE();
-	FOpenGLBuffer* InVertexBuffer = FOpenGLDynamicRHI::ResourceCast(InVertexBufferRHI);
-
+	FOpenGLBuffer* InBuffer = FOpenGLDynamicRHI::ResourceCast(InBufferRHI);
 
 	const FOpenGLTextureFormat& GLFormat = GOpenGLTextureFormats[Format];
 
 	GLuint TextureID = 0;
-	FOpenGL::GenTextures(1,&TextureID);
+	FOpenGL::GenTextures(1, &TextureID);
 
 	// Use a texture stage that's not likely to be used for draws, to avoid waiting
 	OpenGLRHI->CachedSetupTextureStage(OpenGLRHI->GetContextStateForCurrentContext(), FOpenGL::GetMaxCombinedTextureImageUnits() - 1, GL_TEXTURE_BUFFER, TextureID, -1, 1);
-	FOpenGL::TexBuffer(GL_TEXTURE_BUFFER, GLFormat.InternalFormat[0], InVertexBuffer->Resource);
+	FOpenGL::TexBuffer(GL_TEXTURE_BUFFER, GLFormat.InternalFormat[0], InBuffer->Resource);
 
 	// No need to restore texture stage; leave it like this,
 	// and the next draw will take care of cleaning it up; or
 	// next operation that needs the stage will switch something else in on it.
-
+	
 	this->Resource = TextureID;
-	this->BufferResource = InVertexBuffer->Resource;
+	this->BufferResource = InBuffer->Resource;
 	this->Format = GLFormat.InternalFormat[0];
 	this->UnrealFormat = Format;
-	
 }
 
-uint32 FOpenGLVertexBufferUnorderedAccessView::GetBufferSize()
+uint32 FOpenGLTexBufferUnorderedAccessView::GetBufferSize()
 {
-	FOpenGLBuffer* VertexBuffer = FOpenGLDynamicRHI::ResourceCast(VertexBufferRHI.GetReference());
-	return VertexBufferRHI->GetSize();
+	return BufferRHI->GetSize();
 }
 
-FOpenGLVertexBufferUnorderedAccessView::~FOpenGLVertexBufferUnorderedAccessView()
+FOpenGLTexBufferUnorderedAccessView::~FOpenGLTexBufferUnorderedAccessView()
 {
 	if (Resource)
 	{
@@ -306,10 +276,29 @@ FOpenGLVertexBufferUnorderedAccessView::~FOpenGLVertexBufferUnorderedAccessView(
 	}
 }
 
+FUnorderedAccessViewRHIRef FOpenGLDynamicRHI::RHICreateUnorderedAccessView(FRHIBuffer* BufferRHI, bool bUseUAVCounter, bool bAppendBuffer)
+{
+	check(BufferRHI->GetUsage() & BUF_UnorderedAccess);
+	return new FOpenGLBufferUnorderedAccessView(this, BufferRHI);
+}
 
-FUnorderedAccessViewRHIRef FOpenGLDynamicRHI::RHICreateUnorderedAccessView(FRHIBuffer* BufferRHI, uint8 Format)
+FOpenGLBufferUnorderedAccessView::FOpenGLBufferUnorderedAccessView(FOpenGLDynamicRHI* InOpenGLRHI, FRHIBuffer* InBufferRHI)
+	: BufferRHI(InBufferRHI)
+	, OpenGLRHI(InOpenGLRHI)
 {
-	return new FOpenGLVertexBufferUnorderedAccessView(this, BufferRHI, Format);
+	VERIFY_GL_SCOPE();
+	FOpenGLBuffer* Buffer = FOpenGLDynamicRHI::ResourceCast(BufferRHI.GetReference());
+	this->BufferResource = Buffer->Resource;
+}
+
+uint32 FOpenGLBufferUnorderedAccessView::GetBufferSize()
+{
+	return BufferRHI->GetSize();
+}
+
+FOpenGLBufferUnorderedAccessView::~FOpenGLBufferUnorderedAccessView()
+{
+	// not clearing cached state as SSBO can still be bound as SRV
 }
 
 void FOpenGLDynamicRHI::RHIClearUAVFloat(FRHIUnorderedAccessView* UnorderedAccessViewRHI, const FVector4& Values)
@@ -379,16 +368,7 @@ void FOpenGLDynamicRHI::RHIClearUAVFloat(FRHIUnorderedAccessView* UnorderedAcces
 void FOpenGLDynamicRHI::RHIClearUAVUint(FRHIUnorderedAccessView* UnorderedAccessViewRHI, const FUintVector4& Values)
 {
 	FOpenGLUnorderedAccessView* Texture = ResourceCast(UnorderedAccessViewRHI);
-#if OPENGL_GL4 || PLATFORM_LUMINGL4
-	if (GMaxRHIFeatureLevel >= ERHIFeatureLevel::SM5)
-	{
-		glBindBuffer(GL_TEXTURE_BUFFER, Texture->BufferResource);
-		FOpenGL::ClearBufferData(GL_TEXTURE_BUFFER, Texture->Format, GL_RGBA_INTEGER, GL_UNSIGNED_INT, reinterpret_cast<const uint32*>(&Values));
-		GPUProfilingData.RegisterGPUWork(1);
-		return;
-	}
-#endif
-	// Use compute on ES3.1
+	
 	TRHICommandList_RecursiveHazardous<FOpenGLDynamicRHI> RHICmdList(this);
 
 	if (Texture->GetBufferSize() == 0)
@@ -427,47 +407,3 @@ void FOpenGLDynamicRHI::RHIClearUAVUint(FRHIUnorderedAccessView* UnorderedAccess
 		}
 	}
 }
-
-FOpenGLStructuredBufferUnorderedAccessView::FOpenGLStructuredBufferUnorderedAccessView(FOpenGLDynamicRHI* InOpenGLRHI, FRHIBuffer* InStructuredBufferRHI, uint8 InFormat)
-	: StructuredBufferRHI(InStructuredBufferRHI)
-	, OpenGLRHI(InOpenGLRHI)
-{
-	VERIFY_GL_SCOPE();
-	FOpenGLBuffer* InStructuredBuffer = FOpenGLDynamicRHI::ResourceCast(InStructuredBufferRHI);
-	const FOpenGLTextureFormat& GLFormat = GOpenGLTextureFormats[InFormat];
-
-	GLuint TextureID = 0;
-	FOpenGL::GenTextures(1,&TextureID);
-
-	// Use a texture stage that's not likely to be used for draws, to avoid waiting
-	OpenGLRHI->CachedSetupTextureStage(OpenGLRHI->GetContextStateForCurrentContext(), FOpenGL::GetMaxCombinedTextureImageUnits() - 1, GL_TEXTURE_BUFFER, TextureID, -1, 1);
-	FOpenGL::TexBuffer(GL_TEXTURE_BUFFER, GLFormat.InternalFormat[0], InStructuredBuffer->Resource);
-
-	// No need to restore texture stage; leave it like this,
-	// and the next draw will take care of cleaning it up; or
-	// next operation that needs the stage will switch something else in on it.
-
-	this->Resource = TextureID;
-	this->BufferResource = InStructuredBuffer->Resource;
-	this->Format = GLFormat.InternalFormat[0];
-	this->UnrealFormat = InFormat;
-}
-
-uint32 FOpenGLStructuredBufferUnorderedAccessView::GetBufferSize()
-{
-	FOpenGLBuffer* StructuredBuffer = FOpenGLDynamicRHI::ResourceCast(StructuredBufferRHI.GetReference());
-	return StructuredBuffer->GetSize();
-}
-
-FOpenGLStructuredBufferUnorderedAccessView::~FOpenGLStructuredBufferUnorderedAccessView()
-{
-	if (Resource)
-	{
-		RunOnGLRenderContextThread([OpenGLRHI= OpenGLRHI, Resource = Resource]()
-		{
-			VERIFY_GL_SCOPE();
-			OpenGLRHI->InvalidateTextureResourceInCache(Resource);
-			FOpenGL::DeleteTextures(1, &Resource);
-		});
-	}
-}
\ No newline at end of file
diff --git a/Engine/Source/Runtime/OpenGLDrv/Private/OpenGLUniformBuffer.cpp b/Engine/Source/Runtime/OpenGLDrv/Private/OpenGLUniformBuffer.cpp
--- a/Engine/Source/Runtime/OpenGLDrv/Private/OpenGLUniformBuffer.cpp
+++ b/Engine/Source/Runtime/OpenGLDrv/Private/OpenGLUniformBuffer.cpp
@@ -231,7 +231,7 @@ static void ReleaseUniformBuffer(bool bEmulatedBufferData, GLuint Resource, uint
 
 		RunOnGLRenderContextThread(MoveTemp(DeleteGLBuffer));
 	}
-	DecrementBufferMemory(GL_UNIFORM_BUFFER, /*bIsStructuredBuffer=*/ false, AllocatedSize);
+	DecrementBufferMemory(GL_UNIFORM_BUFFER, AllocatedSize);
 }
 
 // Does per-frame global updating for the uniform buffer pool.
@@ -581,7 +581,7 @@ static FUniformBufferRHIRef CreateUniformBuffer(const void* Contents, const FRHI
 		}
 	}
 
-	IncrementBufferMemory(GL_UNIFORM_BUFFER, /*bIsStructuredBuffer=*/ false, AllocatedSize);
+	IncrementBufferMemory(GL_UNIFORM_BUFFER, AllocatedSize);
 
 	check(!GUseEmulatedUniformBuffers || (IsValidRef(EmulatedUniformDataRef) && (EmulatedUniformDataRef->Data.Num() * EmulatedUniformDataRef->Data.GetTypeSize() == AllocatedSize)));
 
diff --git a/Engine/Source/Runtime/OpenGLDrv/Private/OpenGLUtil.cpp b/Engine/Source/Runtime/OpenGLDrv/Private/OpenGLUtil.cpp
--- a/Engine/Source/Runtime/OpenGLDrv/Private/OpenGLUtil.cpp
+++ b/Engine/Source/Runtime/OpenGLDrv/Private/OpenGLUtil.cpp
@@ -102,11 +102,10 @@ DEFINE_STAT(STAT_OpenGLUniformBindTime);
 DEFINE_STAT(STAT_OpenGLVBOSetupTime);
 #endif
 
-void IncrementBufferMemory(GLenum Type, bool bStructuredBuffer, uint32 NumBytes)
+void IncrementBufferMemory(GLenum Type, uint32 NumBytes)
 {
-	if (bStructuredBuffer)
+	if (Type == GL_SHADER_STORAGE_BUFFER)
 	{
-		check(Type == GL_ARRAY_BUFFER);
 		INC_MEMORY_STAT_BY(STAT_StructuredBufferMemory,NumBytes);
 	}
 	else if (Type == GL_UNIFORM_BUFFER)
@@ -128,11 +127,10 @@ void IncrementBufferMemory(GLenum Type, bool bStructuredBuffer, uint32 NumBytes)
 	}
 }
 
-void DecrementBufferMemory(GLenum Type, bool bStructuredBuffer, uint32 NumBytes)
+void DecrementBufferMemory(GLenum Type, uint32 NumBytes)
 {
-	if (bStructuredBuffer)
+	if (Type == GL_SHADER_STORAGE_BUFFER)
 	{
-		check(Type == GL_ARRAY_BUFFER);
 		DEC_MEMORY_STAT_BY(STAT_StructuredBufferMemory,NumBytes);
 	}
 	else if (Type == GL_UNIFORM_BUFFER)
diff --git a/Engine/Source/Runtime/OpenGLDrv/Private/Windows/OpenGLWindows.h b/Engine/Source/Runtime/OpenGLDrv/Private/Windows/OpenGLWindows.h
--- a/Engine/Source/Runtime/OpenGLDrv/Private/Windows/OpenGLWindows.h
+++ b/Engine/Source/Runtime/OpenGLDrv/Private/Windows/OpenGLWindows.h
@@ -294,7 +294,8 @@
 	EnumMacro(PFNGLFENCESYNCPROC, glFenceSync)\
 	EnumMacro(PFNGLGETSYNCIVPROC, glGetSynciv)\
 	EnumMacro(PFNGLCLIENTWAITSYNCPROC, glClientWaitSync)\
-	EnumMacro(PFNGLBINDBUFFERRANGEPROC, glBindBufferRange)
+	EnumMacro(PFNGLBINDBUFFERRANGEPROC, glBindBufferRange) \
+	EnumMacro(PFNGLGETPROGRAMRESOURCEINDEXPROC, glGetProgramResourceIndex) \
 
 #define ENUM_GL_ENTRYPOINTS_OPTIONAL(EnumMacro) \
 	EnumMacro(PFNGLCLIPCONTROLPROC,glClipControl) \
diff --git a/Engine/Source/Runtime/OpenGLDrv/Public/OpenGL.h b/Engine/Source/Runtime/OpenGLDrv/Public/OpenGL.h
--- a/Engine/Source/Runtime/OpenGLDrv/Public/OpenGL.h
+++ b/Engine/Source/Runtime/OpenGLDrv/Public/OpenGL.h
@@ -126,7 +126,6 @@ public:
 	static FORCEINLINE bool SupportsTessellation()						{ return false; }
 	static FORCEINLINE void EnableSupportsClipControl()					{ bSupportsClipControl = true; }
 	static FORCEINLINE bool SupportsClipControl()						{ return bSupportsClipControl; }
-	static FORCEINLINE bool SupportsComputeShaders()					{ return false; }
 	static FORCEINLINE bool SupportsTextureView()						{ return false; }
 	static FORCEINLINE bool SupportsSeamlessCubeMap()					{ return false; }
 	static FORCEINLINE bool SupportsDrawIndirect()						{ return false; }
@@ -166,6 +165,7 @@ public:
 	// Image load/store units
 	static FORCEINLINE GLint GetFirstComputeUAVUnit()			{ return 0; }
 	static FORCEINLINE GLint GetMaxComputeUAVUnits()			{ return 0; }
+	static FORCEINLINE GLint GetFirstVertexUAVUnit()			{ return 0; }
 	static FORCEINLINE GLint GetFirstPixelUAVUnit()				{ return 0; }
 	static FORCEINLINE GLint GetMaxPixelUAVUnits()				{ return 0; }
 	static FORCEINLINE GLint GetMaxCombinedUAVUnits()			{ return 0; }
diff --git a/Engine/Source/Runtime/OpenGLDrv/Public/OpenGL4.h b/Engine/Source/Runtime/OpenGLDrv/Public/OpenGL4.h
--- a/Engine/Source/Runtime/OpenGLDrv/Public/OpenGL4.h
+++ b/Engine/Source/Runtime/OpenGLDrv/Public/OpenGL4.h
@@ -14,7 +14,6 @@
 
 struct FOpenGL4 : public FOpenGL3
 {
-	static FORCEINLINE bool SupportsComputeShaders()					{ return bSupportsComputeShaders; }
 	static FORCEINLINE bool SupportsDrawIndirect()						{ return true; }
 	static FORCEINLINE bool SupportsVertexAttribBinding()				{ return bSupportsVertexAttribBinding; }
 	static FORCEINLINE bool SupportsTextureView()						{ return bSupportsTextureView; }
@@ -108,7 +107,8 @@ struct FOpenGL4 : public FOpenGL3
 
 	static FORCEINLINE GLint GetFirstComputeUAVUnit()			{ return 0; }
 	static FORCEINLINE GLint GetMaxComputeUAVUnits()			{ check(MaxComputeUAVUnits != -1); return MaxComputeUAVUnits; }
-	static FORCEINLINE GLint GetFirstPixelUAVUnit()				{ return 0; }
+	static FORCEINLINE GLint GetFirstVertexUAVUnit()			{ return 0; }
+	static FORCEINLINE GLint GetFirstPixelUAVUnit()				{ return 4; }
 	static FORCEINLINE GLint GetMaxPixelUAVUnits()				{ check(MaxPixelUAVUnits != -1); return MaxPixelUAVUnits; }
 	static FORCEINLINE GLint GetMaxCombinedUAVUnits()			{ return MaxCombinedUAVUnits; }
 		
@@ -124,7 +124,6 @@ protected:
 	static GLint MaxComputeUAVUnits;
 	static GLint MaxPixelUAVUnits;
 
-	static bool bSupportsComputeShaders;
 	static bool bSupportsGPUMemoryInfo;
 	static bool bSupportsVertexAttribBinding;
 	static bool bSupportsTextureView;
diff --git a/Engine/Source/Runtime/OpenGLDrv/Public/OpenGLDrv.h b/Engine/Source/Runtime/OpenGLDrv/Public/OpenGLDrv.h
--- a/Engine/Source/Runtime/OpenGLDrv/Public/OpenGLDrv.h
+++ b/Engine/Source/Runtime/OpenGLDrv/Public/OpenGLDrv.h
@@ -570,6 +570,16 @@ public:
 			ContextState.ElementArrayBufferBound = Buffer;
 		}
 	}
+	
+	FORCEINLINE void CachedBindStorageBuffer( FOpenGLContextState& ContextState, GLuint Buffer )
+	{
+		VERIFY_GL_SCOPE();
+		if( ContextState.StorageBufferBound != Buffer )
+		{
+			glBindBuffer( GL_SHADER_STORAGE_BUFFER, Buffer );
+			ContextState.StorageBufferBound = Buffer;
+		}
+	}
 
 	void CachedBindPixelUnpackBuffer( FOpenGLContextState& ContextState, GLuint Buffer )
 	{
@@ -1140,14 +1150,15 @@ private:
 
 	void SetupUAVsForDraw(FOpenGLContextState& ContextState);
 	void SetupUAVsForCompute(FOpenGLContextState& ContextState, const FOpenGLComputeShader* ComputeShader);
-	void SetupUAVsForProgram(FOpenGLContextState& ContextState, const TBitArray<>& NeededBits, int32 MaxUAVUnitUsed, int32 MaxUAVUnits);
+	void SetupUAVsForProgram(FOpenGLContextState& ContextState, const TBitArray<>& NeededBits, int32 MaxUAVUnitUsed);
 
 	void RHIClearMRT(bool bClearColor, int32 NumClearColors, const FLinearColor* ColorArray, bool bClearDepth, float Depth, bool bClearStencil, uint32 Stencil);
 
 public:
 	/** Remember what RHI user wants set on a specific OpenGL texture stage, translating from Stage and TextureIndex for stage pair. */
 	void InternalSetShaderTexture(FOpenGLTextureBase* Texture, FOpenGLShaderResourceView* SRV, GLint TextureIndex, GLenum Target, GLuint Resource, int NumMips, int LimitMip);
-	void InternalSetShaderUAV(GLint UAVIndex, GLenum Format, GLuint Resource, bool bLayered, GLint Layer, GLenum Access);
+	void InternalSetShaderImageUAV(GLint UAVIndex, GLenum Format, GLuint Resource, bool bLayered, GLint Layer, GLenum Access);
+	void InternalSetShaderBufferUAV(GLint UAVIndex, GLuint Resource);
 	void InternalSetSamplerStates(GLint TextureIndex, FOpenGLSamplerState* SamplerState);
 	void InitializeGLTextureInternal(GLuint TextureID, FRHITexture* Texture, uint32 SizeX, const uint32 SizeY, const bool bCubeTexture, const bool bArrayTexture, const bool bIsExternal, const uint8 Format, const uint32 NumMips, const uint32 NumSamples, const uint32 ArraySize, const ETextureCreateFlags Flags, const FClearValueBinding& InClearValue, FResourceBulkDataInterface* BulkData);
 private:
diff --git a/Engine/Source/Runtime/OpenGLDrv/Public/OpenGLES.h b/Engine/Source/Runtime/OpenGLDrv/Public/OpenGLES.h
--- a/Engine/Source/Runtime/OpenGLDrv/Public/OpenGLES.h
+++ b/Engine/Source/Runtime/OpenGLDrv/Public/OpenGLES.h
@@ -214,7 +214,6 @@ struct FOpenGLES : public FOpenGLBase
 	static FORCEINLINE GLenum GetShadowDepthFormat() { return GL_DEPTH_COMPONENT16; }
 	static FORCEINLINE bool SupportsFramebufferSRGBEnable() { return false; }
 	static FORCEINLINE bool SupportsRGB10A2() { return bSupportsRGB10A2; }
-	static FORCEINLINE bool SupportsComputeShaders() { return true; }
 	static FORCEINLINE bool SupportsDrawIndirect() { return true; }
 	static FORCEINLINE bool SupportsVertexAttribBinding() { return true; }
 	static FORCEINLINE bool SupportsBufferStorage() { return bSupportsBufferStorage; }
diff --git a/Engine/Source/Runtime/OpenGLDrv/Public/OpenGLResources.h b/Engine/Source/Runtime/OpenGLDrv/Public/OpenGLResources.h
--- a/Engine/Source/Runtime/OpenGLDrv/Public/OpenGLResources.h
+++ b/Engine/Source/Runtime/OpenGLDrv/Public/OpenGLResources.h
@@ -56,8 +56,8 @@ namespace OpenGLConsoleVariables
 #define RESTRICT_SUBDATA_SIZE 0
 #endif
 
-void IncrementBufferMemory(GLenum Type, bool bStructuredBuffer, uint32 NumBytes);
-void DecrementBufferMemory(GLenum Type, bool bStructuredBuffer, uint32 NumBytes);
+void IncrementBufferMemory(GLenum Type, uint32 NumBytes);
+void DecrementBufferMemory(GLenum Type, uint32 NumBytes);
 
 // Extra stats for finer-grained timing
 // They shouldn't always be on, as they may impact overall performance
@@ -518,7 +518,7 @@ public:
 						LoadData( 0, FMath::Min<uint32>(InSize,RealSize), InData);
 					}
 #endif
-					IncrementBufferMemory(Type, BaseType::IsStructuredBuffer(), RealSize);
+					IncrementBufferMemory(Type, RealSize);
 				}
 				else
 				{
@@ -556,7 +556,7 @@ public:
 
 			RunOnGLRenderContextThread(MoveTemp(DeleteGLResources));
 			LockBuffer = nullptr;
-			DecrementBufferMemory(Type, BaseType::IsStructuredBuffer(), RealSize);
+			DecrementBufferMemory(Type, RealSize);
 
 			ReleaseCachedBuffer();
 		}
@@ -888,8 +888,6 @@ public:
 		// @todo-mobile
 	}
 
-	static bool IsStructuredBuffer() { return false; }
-
 private:
 	uint32 Size;
 	uint32 Usage;
@@ -955,8 +953,6 @@ public:
 		// @todo-mobile
 	}
 
-	static bool IsStructuredBuffer() { return false; }
-
 private:
 	void*	ZeroStrideBuffer;
 };
@@ -968,12 +964,12 @@ struct FOpenGLEUniformBufferData : public FRefCountedObject
 		uint32 SizeInUint32s = (SizeInBytes + 3) / 4;
 		Data.Empty(SizeInUint32s);
 		Data.AddUninitialized(SizeInUint32s);
-		IncrementBufferMemory(GL_UNIFORM_BUFFER,false,Data.GetAllocatedSize());
+		IncrementBufferMemory(GL_UNIFORM_BUFFER,Data.GetAllocatedSize());
 	}
 
 	~FOpenGLEUniformBufferData()
 	{
-		DecrementBufferMemory(GL_UNIFORM_BUFFER,false,Data.GetAllocatedSize());
+		DecrementBufferMemory(GL_UNIFORM_BUFFER,Data.GetAllocatedSize());
 	}
 
 	TArray<uint32> Data;
@@ -2005,34 +2001,33 @@ public:
 	}
 };
 
-
-class FOpenGLVertexBufferUnorderedAccessView : public FOpenGLUnorderedAccessView
+class FOpenGLTexBufferUnorderedAccessView : public FOpenGLUnorderedAccessView
 {
 public:
 
-	FOpenGLVertexBufferUnorderedAccessView();
-
-	FOpenGLVertexBufferUnorderedAccessView(	FOpenGLDynamicRHI* InOpenGLRHI, FRHIBuffer* InVertexBuffer, uint8 Format);
-
-	virtual ~FOpenGLVertexBufferUnorderedAccessView();
+	FOpenGLTexBufferUnorderedAccessView();
+	
+	FOpenGLTexBufferUnorderedAccessView(FOpenGLDynamicRHI* InOpenGLRHI, FRHIBuffer* InBuffer, uint8 Format);
+	
+	virtual ~FOpenGLTexBufferUnorderedAccessView();
 
-	FBufferRHIRef VertexBufferRHI; // to keep the vertex buffer alive
+	FBufferRHIRef BufferRHI; // to keep source buffer alive
 
 	FOpenGLDynamicRHI* OpenGLRHI;
 
 	virtual uint32 GetBufferSize() override;
 };
 
-class FOpenGLStructuredBufferUnorderedAccessView : public FOpenGLUnorderedAccessView
+class FOpenGLBufferUnorderedAccessView : public FOpenGLUnorderedAccessView
 {
 public:
-	FOpenGLStructuredBufferUnorderedAccessView();
-
-	FOpenGLStructuredBufferUnorderedAccessView(	FOpenGLDynamicRHI* InOpenGLRHI, FRHIBuffer* InBuffer, uint8 Format);
+	FOpenGLBufferUnorderedAccessView();
+	
+	FOpenGLBufferUnorderedAccessView(FOpenGLDynamicRHI* InOpenGLRHI, FRHIBuffer* InBuffer);
 
-	virtual ~FOpenGLStructuredBufferUnorderedAccessView();
+	virtual ~FOpenGLBufferUnorderedAccessView();
 
-	FBufferRHIRef StructuredBufferRHI; // to keep the stuctured buffer alive
+	FBufferRHIRef BufferRHI; // to keep source buffer alive
 
 	FOpenGLDynamicRHI* OpenGLRHI;
 
@@ -2041,9 +2036,6 @@ public:
 
 class FOpenGLShaderResourceView : public FRefCountedObject
 {
-	// In OpenGL 3.2, the only view that actually works is a Buffer<type> kind of view from D3D10,
-	// and it's mapped to OpenGL's buffer texture.
-
 public:
 
 	/** OpenGL texture the buffer is bound with */
@@ -2056,8 +2048,7 @@ public:
 	int32 LimitMip;
 
 	/** Needed on OS X to force a rebind of the texture buffer to the texture name to workaround radr://18379338 */
-	FBufferRHIRef VertexBuffer;
-	FBufferRHIRef IndexBuffer;
+	FBufferRHIRef BufferRHI;
 	uint64 ModificationVersion;
 	uint8 Format;
 
@@ -2068,40 +2059,41 @@ public:
 	,	ModificationVersion(0)
 	,	Format(0)
 	,	OpenGLRHI(InOpenGLRHI)
-	,	OwnsResource(true)
+	,	OwnsResource(InTarget != GL_SHADER_STORAGE_BUFFER)
 	{}
 
-	FOpenGLShaderResourceView(FOpenGLDynamicRHI* InOpenGLRHI, GLuint InResource, GLenum InTarget, FRHIBuffer* InIndexBuffer)
+	FOpenGLShaderResourceView(FOpenGLDynamicRHI* InOpenGLRHI, GLuint InResource, GLenum InTarget, FRHIBuffer* InBuffer)
 		: Resource(InResource)
 		, Target(InTarget)
 		, LimitMip(-1)
-		, IndexBuffer(InIndexBuffer)
+		, BufferRHI(InBuffer)
 		, ModificationVersion(0)
 		, Format(0)
 		, OpenGLRHI(InOpenGLRHI)
-		, OwnsResource(true)
+		, OwnsResource(InTarget != GL_SHADER_STORAGE_BUFFER)
 	{
-		if (IndexBuffer)
+		if (BufferRHI)
 		{
-			FOpenGLBuffer* IB = (FOpenGLBuffer*)IndexBuffer.GetReference();
-			ModificationVersion = IB->ModificationCount;
+			FOpenGLBuffer* GLBuffer = (FOpenGLBuffer*)BufferRHI.GetReference();
+			ModificationVersion = GLBuffer->ModificationCount;
 		}
 	}
 
-	FOpenGLShaderResourceView( FOpenGLDynamicRHI* InOpenGLRHI, GLuint InResource, GLenum InTarget, FRHIBuffer* InVertexBuffer, uint8 InFormat )
+	FOpenGLShaderResourceView( FOpenGLDynamicRHI* InOpenGLRHI, GLuint InResource, GLenum InTarget, FRHIBuffer* InBuffer, uint8 InFormat )
 	:	Resource(InResource)
 	,	Target(InTarget)
 	,	LimitMip(-1)
-	,	VertexBuffer(InVertexBuffer)
+	,	BufferRHI(InBuffer)
 	,	ModificationVersion(0)
 	,	Format(InFormat)
 	,	OpenGLRHI(InOpenGLRHI)
 	,	OwnsResource(true)
 	{
-		if (VertexBuffer)
+		check(InTarget != GL_SHADER_STORAGE_BUFFER);
+		if (BufferRHI)
 		{
-			FOpenGLBuffer* VB = (FOpenGLBuffer*)VertexBuffer.GetReference();
-			ModificationVersion = VB->ModificationCount;
+			FOpenGLBuffer* GLBuffer = (FOpenGLBuffer*)BufferRHI.GetReference();
+			ModificationVersion = GLBuffer->ModificationCount;
 		}
 	}
 
@@ -2113,7 +2105,9 @@ public:
 	,	Format(0)
 	,	OpenGLRHI(InOpenGLRHI)
 	,	OwnsResource(InOwnsResource)
-	{}
+	{
+		check(InTarget != GL_SHADER_STORAGE_BUFFER);
+	}
 
 	virtual ~FOpenGLShaderResourceView( void );
 
diff --git a/Engine/Source/Runtime/OpenGLDrv/Public/OpenGLState.h b/Engine/Source/Runtime/OpenGLDrv/Public/OpenGLState.h
--- a/Engine/Source/Runtime/OpenGLDrv/Public/OpenGLState.h
+++ b/Engine/Source/Runtime/OpenGLDrv/Public/OpenGLState.h
@@ -281,12 +281,11 @@ struct FOpenGLCommonState
 {
 	FTextureStage*			Textures;
 	FOpenGLSamplerState**	SamplerStates;
-	FUAVStage*				UAVs;
+	TArray<FUAVStage>		UAVs;
 
 	FOpenGLCommonState()
 	: Textures(NULL)
 	, SamplerStates(NULL)
-	, UAVs(NULL)
 	{}
 
 	virtual ~FOpenGLCommonState()
@@ -300,20 +299,21 @@ struct FOpenGLCommonState
 	{
 		check(NumCombinedTextures >= FOpenGL::GetMaxCombinedTextureImageUnits());
 		check(NumCombinedUAVUnits >= FOpenGL::GetMaxCombinedUAVUnits());
-		check(!Textures && !SamplerStates && !UAVs);
+		check(!Textures && !SamplerStates && UAVs.Num() == 0);
 		Textures = new FTextureStage[NumCombinedTextures];
 		SamplerStates = new FOpenGLSamplerState*[NumCombinedTextures];
 		FMemory::Memset( SamplerStates, 0, NumCombinedTextures * sizeof(*SamplerStates) );
-		UAVs = new FUAVStage[NumCombinedUAVUnits];
+		
+		UAVs.Reserve(NumCombinedUAVUnits);
+		UAVs.AddDefaulted(NumCombinedUAVUnits);
 	}
 
 	virtual void CleanupResources()
 	{
-		delete [] UAVs;
 		delete [] SamplerStates;
 		delete [] Textures;
 
-		UAVs = NULL;
+		UAVs.Empty();
 		SamplerStates = NULL;
 		Textures = NULL;
 	}
@@ -342,6 +342,7 @@ struct FOpenGLContextState : public FOpenGLCommonState
 	float							DepthMaxZ;
 	GLuint							ArrayBufferBound;
 	GLuint							ElementArrayBufferBound;
+	GLuint							StorageBufferBound;
 	GLuint							PixelUnpackBufferBound;
 	GLuint							UniformBufferBound;
 	FLinearColor					ClearColor;
@@ -386,6 +387,7 @@ struct FOpenGLContextState : public FOpenGLCommonState
 	,	DepthMaxZ(1.0f)
 	,	ArrayBufferBound(0)
 	,	ElementArrayBufferBound(0)
+	,	StorageBufferBound(0)
 	,	PixelUnpackBufferBound(0)
 	,	UniformBufferBound(0)
 	,	ClearColor(-1, -1, -1, -1)
diff --git a/Engine/Source/Runtime/RenderCore/Private/RenderUtils.cpp b/Engine/Source/Runtime/RenderCore/Private/RenderUtils.cpp
--- a/Engine/Source/Runtime/RenderCore/Private/RenderUtils.cpp
+++ b/Engine/Source/Runtime/RenderCore/Private/RenderUtils.cpp
@@ -1034,23 +1034,6 @@ RENDERCORE_API bool SupportsTextureCubeArray(ERHIFeatureLevel::Type FeatureLevel
 		|| IsMobileDeferredShadingEnabled(GMaxRHIShaderPlatform);
 }
 
-RENDERCORE_API bool GPUSceneUseTexture2D(const FStaticShaderPlatform Platform)
-{
-	if (IsMobilePlatform(Platform))
-	{
-		static TConsoleVariableData<int32>* CVar = IConsoleManager::Get().FindTConsoleVariableDataInt(TEXT("r.Mobile.UseGPUSceneTexture"));
-		if (Platform == SP_OPENGL_ES3_1_ANDROID)
-		{
-			return true;
-		}
-		else
-		{
-			return (CVar && CVar->GetValueOnAnyThread() != 0) ? true : false;
-		}
-	}
-	return false;
-}
-
 RENDERCORE_API bool MaskedInEarlyPass(const FStaticShaderPlatform Platform)
 {
 	static IConsoleVariable* CVarMobileEarlyZPassOnlyMaterialMasking = IConsoleManager::Get().FindConsoleVariable(TEXT("r.Mobile.EarlyZPassOnlyMaterialMasking"));
diff --git a/Engine/Source/Runtime/RenderCore/Private/Shader.cpp b/Engine/Source/Runtime/RenderCore/Private/Shader.cpp
--- a/Engine/Source/Runtime/RenderCore/Private/Shader.cpp
+++ b/Engine/Source/Runtime/RenderCore/Private/Shader.cpp
@@ -1517,18 +1517,6 @@ void ShaderMapAppendKeyString(EShaderPlatform Platform, FString& KeyString)
 			static IConsoleVariable* CVar = IConsoleManager::Get().FindConsoleVariable(TEXT("r.Mobile.SupportGPUScene"));
 			bool bMobileGpuScene = (CVar && CVar->GetInt() != 0);
 			KeyString += bMobileGpuScene ? TEXT("_MobGPUSc") : TEXT("");
-			if (bMobileGpuScene)
-			{
-				// Mobile specific verify if we are using texturebuffer or texture2D
-				if (!GPUSceneUseTexture2D(Platform))
-				{
-					KeyString += TEXT("_TexBuf");
-				}
-				else
-				{
-					KeyString += TEXT("_Tex2D");
-				}
-			}
 		}
 
 		{
diff --git a/Engine/Source/Runtime/RenderCore/Public/RenderUtils.h b/Engine/Source/Runtime/RenderCore/Public/RenderUtils.h
--- a/Engine/Source/Runtime/RenderCore/Public/RenderUtils.h
+++ b/Engine/Source/Runtime/RenderCore/Public/RenderUtils.h
@@ -498,8 +498,6 @@ RENDERCORE_API bool MobileRequiresSceneDepthAux(const FStaticShaderPlatform Plat
 
 RENDERCORE_API bool SupportsTextureCubeArray(ERHIFeatureLevel::Type FeatureLevel);
 
-RENDERCORE_API bool GPUSceneUseTexture2D(const FStaticShaderPlatform Platform);
-
 RENDERCORE_API bool MaskedInEarlyPass(const FStaticShaderPlatform Platform);
 
 RENDERCORE_API bool AllowPixelDepthOffset(const FStaticShaderPlatform Platform);
diff --git a/Engine/Source/Runtime/Renderer/Private/GPUScene.cpp b/Engine/Source/Runtime/Renderer/Private/GPUScene.cpp
--- a/Engine/Source/Runtime/Renderer/Private/GPUScene.cpp
+++ b/Engine/Source/Runtime/Renderer/Private/GPUScene.cpp
@@ -58,68 +58,28 @@ FAutoConsoleVariableRef CVarGPUSceneInstanceBVH(
 	ECVF_RenderThreadSafe | ECVF_ReadOnly
 );
 
-
-template<typename ResourceType>
-ResourceType* GetMirrorGPU(FGPUScene& GPUScene);
-
-template<>
-FRWBufferStructured* GetMirrorGPU<FRWBufferStructured>(FGPUScene& GPUScene)
+FRWBufferStructured* GetMirrorGPU(FGPUScene& GPUScene)
 {
 	return &GPUScene.PrimitiveBuffer;
 }
 
-
-template<>
-FTextureRWBuffer2D* GetMirrorGPU<FTextureRWBuffer2D>(FGPUScene& GPUScene)
-{
-	return &GPUScene.PrimitiveTexture;
-}
-
 /**
- * Returns a pointer to the underlying storage (either a texture or buffer), 
- * Only implemented for Buffer at present, since the texture data layout is opaque and complex to iterate.
+ * Returns a pointer to the underlying storage, 
  * Only used for debugging at present.
  */
-template<typename ResourceType>
-void* LockResource(ResourceType& Resource);
-template<>
-void* LockResource<FRWBufferStructured>(FRWBufferStructured& Resource)
+void* LockResource(FRWBufferStructured& Resource)
 {
 	return RHILockBuffer(Resource.Buffer, 0, Resource.NumBytes, RLM_ReadOnly);
 }
-template<>
-void* LockResource<FTextureRWBuffer2D>(FTextureRWBuffer2D& Resource)
-{
-	// RHILockTexture2D(Resource.Buffer, 0, RLM_ReadOnly, RowStrideInBytesOut, false);
-	return nullptr;
-}
 
-template<typename ResourceType>
-void UnlockResourceGPUScene(ResourceType& Resource);
-template<>
-void UnlockResourceGPUScene<FRWBufferStructured>(FRWBufferStructured& Resource)
+void UnlockResourceGPUScene(FRWBufferStructured& Resource)
 {
 	RHIUnlockBuffer(Resource.Buffer);
 }
-template<>
-void UnlockResourceGPUScene<FTextureRWBuffer2D>(FTextureRWBuffer2D& Resource)
-{
-	// RHIUnlockTexture2D(Resource.Buffer, 0, false);
-}
-
-template<typename ResourceType>
-void UpdateUniformResource(FViewInfo& View, FGPUScene& GPUScene);
 
-
-template<>
-void UpdateUniformResource<FTextureRWBuffer2D>(FViewInfo& View, FGPUScene& GPUScene)
-{
-	View.CachedViewUniformShaderParameters->PrimitiveSceneDataTexture = GetMirrorGPU<FTextureRWBuffer2D>(GPUScene)->Buffer;
-}
-template<>
-void UpdateUniformResource<FRWBufferStructured>(FViewInfo& View, FGPUScene& GPUScene)
+void UpdateUniformResource(FViewInfo& View, FGPUScene& GPUScene)
 {
-	View.CachedViewUniformShaderParameters->PrimitiveSceneData = GetMirrorGPU<FRWBufferStructured>(GPUScene)->SRV;
+	View.CachedViewUniformShaderParameters->PrimitiveSceneData = GetMirrorGPU(GPUScene)->SRV;
 }
 
 static int32 GetMaxPrimitivesUpdate(uint32 NumUploads, uint32 InStrideInFloat4s)
@@ -422,7 +382,6 @@ void FGPUScene::EndRender()
 	CurrentDynamicContext = nullptr;
 }
 
-template<typename ResourceType>
 void FGPUScene::UpdateInternal(FRHICommandListImmediate& RHICmdList, FScene& Scene)
 {
 	ensure(bInBeginEndBlock);
@@ -461,16 +420,16 @@ void FGPUScene::UpdateInternal(FRHICommandListImmediate& RHICmdList, FScene& Sce
 		}
 
 		FUploadDataSourceAdapterScenePrimitives Adapter(*this, Scene);
-		UploadGeneral<ResourceType, FUploadDataSourceAdapterScenePrimitives>(RHICmdList, &Scene, Adapter);
+		UploadGeneral<FUploadDataSourceAdapterScenePrimitives>(RHICmdList, &Scene, Adapter);
 
 		PrimitivesMarkedToUpdate.Init(false, PrimitivesMarkedToUpdate.Num());
 
 
 #if DO_CHECK
 		// Validate the scene primitives are identical to the uploaded data (not the dynamic ones).
-		if (GGPUSceneValidatePrimitiveBuffer && (PrimitiveBuffer.NumBytes > 0 || PrimitiveTexture.NumBytes > 0))
+		if (GGPUSceneValidatePrimitiveBuffer && PrimitiveBuffer.NumBytes > 0)
 		{
-			ResourceType* MirrorResourceGPU = GetMirrorGPU<ResourceType>(*this);
+			FRWBufferStructured* MirrorResourceGPU = GetMirrorGPU(*this);
 			//UE_LOG(LogRenderer, Warning, TEXT("r.GPUSceneValidatePrimitiveBuffer enabled, doing slow readback from GPU"));
 			const FPrimitiveSceneShaderData* PrimitiveBufferPtr = reinterpret_cast<const FPrimitiveSceneShaderData*>(LockResource(*MirrorResourceGPU));
 			ensureMsgf(PrimitiveBufferPtr != nullptr, TEXT("Validation not implemented for Texture2D buffer type as layout is opaque."));
@@ -510,7 +469,7 @@ void FGPUScene::UpdateInternal(FRHICommandListImmediate& RHICmdList, FScene& Sce
 	checkSlow(PrimitivesToUpdate.Num() == 0);
 }
 
-template<typename ResourceType, typename FUploadDataSourceAdapter>
+template<typename FUploadDataSourceAdapter>
 void FGPUScene::UploadGeneral(FRHICommandListImmediate& RHICmdList, FScene *Scene, FUploadDataSourceAdapter &UploadDataSourceAdapter)
 {
 	ensure(bInBeginEndBlock);
@@ -527,7 +486,7 @@ void FGPUScene::UploadGeneral(FRHICommandListImmediate& RHICmdList, FScene *Scen
 		const bool bExecuteInParallel = GGPUSceneParallelUpdate != 0 && FApp::ShouldUseThreadingForPerformance();
 		const bool bNaniteEnabled = DoesPlatformSupportNanite(GMaxRHIShaderPlatform);
 
-		ResourceType* MirrorResourceGPU = GetMirrorGPU<ResourceType>(*this);
+		FRWBufferStructured* MirrorResourceGPU = GetMirrorGPU(*this);
 
 		const uint32 SizeReserve = FMath::RoundUpToPowerOfTwo(FMath::Max(DynamicPrimitivesOffset, 256));
 		const bool bResizedPrimitiveData = ResizeResourceIfNeeded(RHICmdList, *MirrorResourceGPU, SizeReserve * sizeof(FPrimitiveSceneShaderData::Data), TEXT("PrimitiveData"));
@@ -1081,7 +1040,6 @@ struct FUploadDataSourceAdapterDynamicPrimitives
 	int32 InstanceIDStartOffset;
 };
 
-template<typename ResourceType>
 void FGPUScene::UploadDynamicPrimitiveShaderDataForViewInternal(FRHICommandListImmediate& RHICmdList, FScene *Scene, FViewInfo& View)
 {
 	ensure(bInBeginEndBlock);
@@ -1112,10 +1070,10 @@ void FGPUScene::UploadDynamicPrimitiveShaderDataForViewInternal(FRHICommandListI
 #endif // defined(GPUCULL_TODO)
 
 		FUploadDataSourceAdapterDynamicPrimitives UploadAdapter(DynamicPrimitiveShaderData, UploadIdStart, Collector.UploadData->InstanceDataOffset);
-		UploadGeneral<ResourceType, FUploadDataSourceAdapterDynamicPrimitives>(RHICmdList, Scene, UploadAdapter);
+		UploadGeneral<FUploadDataSourceAdapterDynamicPrimitives>(RHICmdList, Scene, UploadAdapter);
 	}
 
-	UpdateUniformResource<ResourceType>(View, *this);
+	UpdateUniformResource(View, *this);
 
 	// Update view uniform buffer
 	View.CachedViewUniformShaderParameters->InstanceSceneData = InstanceDataBuffer.SRV;
@@ -1159,14 +1117,8 @@ void FGPUScene::Update(FRDGBuilder& GraphBuilder, FScene& Scene)
 		{
 			Scene.VirtualShadowMapArrayCacheManager->ProcessPrimitivesToUpdate(GraphBuilder, Scene);
 		}
-		if (GPUSceneUseTexture2D(GetShaderPlatform()))
-		{
-			UpdateInternal<FTextureRWBuffer2D>(GraphBuilder.RHICmdList, Scene);
-		}
-		else
-		{
-			UpdateInternal<FRWBufferStructured>(GraphBuilder.RHICmdList, Scene);
-		}
+		
+		UpdateInternal(GraphBuilder.RHICmdList, Scene);
 	}
 }
 
@@ -1174,14 +1126,7 @@ void FGPUScene::UploadDynamicPrimitiveShaderDataForView(FRHICommandListImmediate
 {
 	if (bIsEnabled)
 	{
-		if (GPUSceneUseTexture2D(GetShaderPlatform()))
-		{
-			UploadDynamicPrimitiveShaderDataForViewInternal<FTextureRWBuffer2D>(RHICmdList, Scene, View);
-		}
-		else
-		{
-			UploadDynamicPrimitiveShaderDataForViewInternal<FRWBufferStructured>(RHICmdList, Scene, View);
-		}
+		UploadDynamicPrimitiveShaderDataForViewInternal(RHICmdList, Scene, View);
 	}
 }
 
diff --git a/Engine/Source/Runtime/Renderer/Private/GPUScene.h b/Engine/Source/Runtime/Renderer/Private/GPUScene.h
--- a/Engine/Source/Runtime/Renderer/Private/GPUScene.h
+++ b/Engine/Source/Runtime/Renderer/Private/GPUScene.h
@@ -198,9 +198,7 @@ public:
 	TBitArray<>				PrimitivesMarkedToUpdate;
 
 	/** GPU mirror of Primitives */
-	/** Only one of the resources(TextureBuffer or Texture2D) will be used depending on the Mobile.UseGPUSceneTexture cvar */
 	FRWBufferStructured PrimitiveBuffer;
-	FTextureRWBuffer2D PrimitiveTexture;
 	FScatterUploadBuffer PrimitiveUploadBuffer;
 	FScatterUploadBuffer PrimitiveUploadViewBuffer;
 
@@ -247,13 +245,11 @@ private:
 	 * Generalized upload that uses an adapter to abstract the data souce. Enables uploading scene primitives & dynamic primitives using a single path.
 	 * @parameter Scene may be null, as it is only needed for the Nanite material table update (which is coupled to the Scene at the moment).
 	 */
-	template<typename ResourceType, typename FUploadDataSourceAdapter>
+	template<typename FUploadDataSourceAdapter>
 	void UploadGeneral(FRHICommandListImmediate& RHICmdList, FScene* Scene, FUploadDataSourceAdapter& UploadDataSourceAdapter);
 
-	template<typename ResourceType>
 	void UploadDynamicPrimitiveShaderDataForViewInternal(FRHICommandListImmediate& RHICmdList, FScene* Scene, FViewInfo& View);
 
-	template<typename ResourceType>
 	void UpdateInternal(FRHICommandListImmediate& RHICmdList, FScene& Scene);
 };
 
diff --git a/Engine/Source/Runtime/Renderer/Private/MeshDrawCommands.cpp b/Engine/Source/Runtime/Renderer/Private/MeshDrawCommands.cpp
--- a/Engine/Source/Runtime/Renderer/Private/MeshDrawCommands.cpp
+++ b/Engine/Source/Runtime/Renderer/Private/MeshDrawCommands.cpp
@@ -468,16 +468,7 @@ void BuildMeshDrawCommandPrimitiveIdBuffer(
 			{
 				//@todo - refactor into memcpy
 				checkSlow(PrimitiveIdIndex < MaxPrimitiveId);
-				if (!GPUSceneUseTexture2D(ShaderPlatform))
-				{
-					PrimitiveIds[PrimitiveIdIndex] = VisibleMeshDrawCommand.DrawPrimitiveId;
-				}
-				else
-				{
-					//Packing for mobile texture2D GPUScene. Must be in sync with SceneData.ush
-					uint16 PrimitivesPerTextureLine = FPrimitiveSceneShaderData::GetPrimitivesPerTextureLine();
-					PrimitiveIds[PrimitiveIdIndex] = ((VisibleMeshDrawCommand.DrawPrimitiveId / PrimitivesPerTextureLine) << 16) | (VisibleMeshDrawCommand.DrawPrimitiveId % PrimitivesPerTextureLine);
-				}
+				PrimitiveIds[PrimitiveIdIndex] = VisibleMeshDrawCommand.DrawPrimitiveId;
 			}
 		}
 
diff --git a/Engine/Source/Runtime/Renderer/Private/Renderer.cpp b/Engine/Source/Runtime/Renderer/Private/Renderer.cpp
--- a/Engine/Source/Runtime/Renderer/Private/Renderer.cpp
+++ b/Engine/Source/Runtime/Renderer/Private/Renderer.cpp
@@ -177,14 +177,7 @@ void FRendererModule::DrawTileMesh(FCanvasRenderContext& RenderContext, FMeshPas
 
 				SinglePrimitiveStructured.UploadToGPU();
 
-				if (!GPUSceneUseTexture2D(View.GetShaderPlatform()))
-				{
-					View.PrimitiveSceneDataOverrideSRV = SinglePrimitiveStructured.PrimitiveSceneDataBufferSRV;
-				}
-				else
-				{
-					View.PrimitiveSceneDataTextureOverrideRHI = SinglePrimitiveStructured.PrimitiveSceneDataTextureRHI;
-				}
+				View.PrimitiveSceneDataOverrideSRV = SinglePrimitiveStructured.PrimitiveSceneDataBufferSRV;
 				View.InstanceSceneDataOverrideSRV  = SinglePrimitiveStructured.InstanceSceneDataBufferSRV;
 				View.LightmapSceneDataOverrideSRV = SinglePrimitiveStructured.LightmapSceneDataBufferSRV;
 			}
diff --git a/Engine/Source/Runtime/Renderer/Private/SceneRendering.cpp b/Engine/Source/Runtime/Renderer/Private/SceneRendering.cpp
--- a/Engine/Source/Runtime/Renderer/Private/SceneRendering.cpp
+++ b/Engine/Source/Runtime/Renderer/Private/SceneRendering.cpp
@@ -1773,15 +1773,6 @@ void FViewInfo::SetupUniformBufferParameters(
 			ViewUniformShaderParameters.PrimitiveSceneData = Scene->GPUScene.PrimitiveBuffer.SRV;
 		}
 
-		if (PrimitiveSceneDataTextureOverrideRHI)
-		{
-			ViewUniformShaderParameters.PrimitiveSceneDataTexture = PrimitiveSceneDataTextureOverrideRHI;
-		}
-		else if (Scene)
-		{
-			ViewUniformShaderParameters.PrimitiveSceneDataTexture = OrBlack2DIfNull(Scene->GPUScene.PrimitiveTexture.Buffer);
-		}
-		
 		if (InstanceSceneDataOverrideSRV)
 		{
 			ViewUniformShaderParameters.InstanceSceneData = InstanceSceneDataOverrideSRV;
