Enable DEBUG_GL_SHADERS in editor configuration

diff --git a/Engine/Source/Runtime/OpenGLDrv/Public/OpenGLUtil.h b/Engine/Source/Runtime/OpenGLDrv/Public/OpenGLUtil.h
--- a/Engine/Source/Runtime/OpenGLDrv/Public/OpenGLUtil.h
+++ b/Engine/Source/Runtime/OpenGLDrv/Public/OpenGLUtil.h
@@ -24,7 +24,7 @@
 #define ENABLE_UNIFORM_BUFFER_LAYOUT_DUMP 0
 
 /** Set to 1 to enable shader debugging which e.g. keeps the GLSL source as members of TOpenGLShader*/
-#define DEBUG_GL_SHADERS (UE_BUILD_DEBUG)
+#define DEBUG_GL_SHADERS (UE_BUILD_DEBUG || UE_EDITOR)
 
 /** Set to 1 to enable calls to place event markers into the OpenGL stream
     this is purposefully not considered for OPENGL_PERFORMANCE_DATA_INVALID, 
