Added missing declaration of glGetProgramResourceIndex for Lunux and Lumin OpenGL RHI

diff --git a/Engine/Source/Runtime/OpenGLDrv/Private/Linux/OpenGLLinux.h b/Engine/Source/Runtime/OpenGLDrv/Private/Linux/OpenGLLinux.h
--- a/Engine/Source/Runtime/OpenGLDrv/Private/Linux/OpenGLLinux.h
+++ b/Engine/Source/Runtime/OpenGLDrv/Private/Linux/OpenGLLinux.h
@@ -297,7 +297,8 @@ THIRD_PARTY_INCLUDES_END
 	EnumMacro(PFNGLFENCESYNCPROC, glFenceSync)\
 	EnumMacro(PFNGLGETSYNCIVPROC, glGetSynciv)\
 	EnumMacro(PFNGLCLIENTWAITSYNCPROC, glClientWaitSync)\
-	EnumMacro(PFNGLBINDBUFFERRANGEPROC, glBindBufferRange)
+	EnumMacro(PFNGLBINDBUFFERRANGEPROC, glBindBufferRange) \
+	EnumMacro(PFNGLGETPROGRAMRESOURCEINDEXPROC, glGetProgramResourceIndex)
 
 #define ENUM_GL_ENTRYPOINTS_OPTIONAL(EnumMacro) \
 	EnumMacro(PFNGLCLIPCONTROLPROC,glClipControl) \
diff --git a/Engine/Source/Runtime/OpenGLDrv/Private/Lumin/LuminGL4.h b/Engine/Source/Runtime/OpenGLDrv/Private/Lumin/LuminGL4.h
--- a/Engine/Source/Runtime/OpenGLDrv/Private/Lumin/LuminGL4.h
+++ b/Engine/Source/Runtime/OpenGLDrv/Private/Lumin/LuminGL4.h
@@ -304,7 +304,8 @@ typedef void (APIENTRYP PFNGLGETVERTEXATTRIBLUI64VNVPROC) (GLuint index, GLenum
 	EnumMacro(PFNGLGETUNIFORMBLOCKINDEXPROC,glGetUniformBlockIndex)	 \
 	EnumMacro(PFNGLTEXBUFFERPROC,glTexBuffer) \
 	EnumMacro(PFNGLDRAWARRAYSINSTANCEDPROC,glDrawArraysInstanced) \
-	EnumMacro(PFNGLCOPYBUFFERSUBDATAPROC, glCopyBufferSubData)
+	EnumMacro(PFNGLCOPYBUFFERSUBDATAPROC, glCopyBufferSubData) \
+	EnumMacro(PFNGLGETPROGRAMRESOURCEINDEXPROC, glGetProgramResourceIndex)
 
 #define ENUM_GL_ENTRYPOINTS_3_2(EnumMacro)                       \
 	EnumMacro(PFNGLISSYNCPROC, glIsSync)						 \
