Remove wrong check on GL SRV renaming

diff --git a/Engine/Source/Runtime/OpenGLDrv/Private/OpenGLUAV.cpp b/Engine/Source/Runtime/OpenGLDrv/Private/OpenGLUAV.cpp
--- a/Engine/Source/Runtime/OpenGLDrv/Private/OpenGLUAV.cpp
+++ b/Engine/Source/Runtime/OpenGLDrv/Private/OpenGLUAV.cpp
@@ -179,9 +179,9 @@ void FOpenGLDynamicRHI::RHIUpdateShaderResourceView(FRHIShaderResourceView* SRV,
 	}
 	else
 	{
-		check(SRVGL->Format == Format && SRVGL->Target == GL_TEXTURE_BUFFER);
 		FOpenGL::TexBuffer(GL_TEXTURE_BUFFER, GLFormat.InternalFormat[0], BufferGL->Resource);
 		SRVGL->BufferRHI = InBuffer;
+		SRVGL->Format = Format;
 		SRVGL->ModificationVersion = BufferGL->ModificationCount;
 	}
 }
