Fixed: QA-Waterworks fails to Launch On Samsung S20+
Changed water to use typed buffers instead of StructuredBuffers, as MALI GPUs do not support access to StructuredBuffers in a vertex shaders

diff --git a/Engine/Plugins/Experimental/Water/Source/Runtime/Private/GerstnerWaterWaveViewExtension.cpp b/Engine/Plugins/Experimental/Water/Source/Runtime/Private/GerstnerWaterWaveViewExtension.cpp
--- a/Engine/Plugins/Experimental/Water/Source/Runtime/Private/GerstnerWaterWaveViewExtension.cpp
+++ b/Engine/Plugins/Experimental/Water/Source/Runtime/Private/GerstnerWaterWaveViewExtension.cpp
@@ -100,12 +100,12 @@ void FGerstnerWaterWaveViewExtension::SetupViewFamily(FSceneViewFamily& InViewFa
 			[this, WaterDataBuffer, WaterIndirectionBuffer](FRHICommandListImmediate& RHICmdList) mutable
 			{
 				FRHIResourceCreateInfo CreateInfoData(TEXT("WaterDataBuffer"), &WaterDataBuffer);
-				DataBuffer = RHICreateStructuredBuffer(sizeof(FVector4), WaterDataBuffer.GetResourceDataSize(), BUF_StructuredBuffer | BUF_ShaderResource | BUF_Static, ERHIAccess::SRVMask, CreateInfoData);
-				DataSRV = RHICreateShaderResourceView(DataBuffer);
+				DataBuffer = RHICreateBuffer(WaterDataBuffer.GetResourceDataSize(), BUF_VertexBuffer | BUF_ShaderResource | BUF_Static, sizeof(FVector4), ERHIAccess::SRVMask, CreateInfoData);
+				DataSRV = RHICreateShaderResourceView(DataBuffer, sizeof(FVector4), PF_A32B32G32R32F);
 
 				FRHIResourceCreateInfo CreateInfoIndirection(TEXT("WaterIndirectionBuffer"), &WaterIndirectionBuffer);
-				IndirectionBuffer = RHICreateStructuredBuffer(sizeof(FVector4), WaterIndirectionBuffer.GetResourceDataSize(), BUF_StructuredBuffer | BUF_ShaderResource | BUF_Static, ERHIAccess::SRVMask, CreateInfoIndirection);
-				IndirectionSRV = RHICreateShaderResourceView(IndirectionBuffer);
+				IndirectionBuffer = RHICreateBuffer(WaterIndirectionBuffer.GetResourceDataSize(), BUF_VertexBuffer | BUF_ShaderResource | BUF_Static, sizeof(FVector4), ERHIAccess::SRVMask, CreateInfoIndirection);
+				IndirectionSRV = RHICreateShaderResourceView(IndirectionBuffer, sizeof(FVector4), PF_A32B32G32R32F);
 			}
 		);
 
diff --git a/Engine/Source/Runtime/Engine/Private/SceneManagement.cpp b/Engine/Source/Runtime/Engine/Private/SceneManagement.cpp
--- a/Engine/Source/Runtime/Engine/Private/SceneManagement.cpp
+++ b/Engine/Source/Runtime/Engine/Private/SceneManagement.cpp
@@ -799,9 +799,8 @@ FViewUniformShaderParameters::FViewUniformShaderParameters()
 
 	PhysicsFieldClipmapBuffer = GWhiteVertexBufferWithSRV->ShaderResourceViewRHI;
 
-	// [todo] Default to some other buffer
-	WaterIndirection = GIdentityPrimitiveBuffer.PrimitiveSceneDataBufferSRV;
-	WaterData = GIdentityPrimitiveBuffer.PrimitiveSceneDataBufferSRV;
+	WaterIndirection = GWhiteVertexBufferWithSRV->ShaderResourceViewRHI;
+	WaterData = GWhiteVertexBufferWithSRV->ShaderResourceViewRHI;
 
 	HairScatteringLUTTexture = BlackVolume;
 	HairScatteringLUTSampler = TStaticSamplerState<SF_Bilinear>::GetRHI();
diff --git a/Engine/Source/Runtime/Engine/Public/SceneView.h b/Engine/Source/Runtime/Engine/Public/SceneView.h
--- a/Engine/Source/Runtime/Engine/Public/SceneView.h
+++ b/Engine/Source/Runtime/Engine/Public/SceneView.h
@@ -870,8 +870,8 @@ BEGIN_GLOBAL_SHADER_PARAMETER_STRUCT_WITH_CONSTRUCTOR(FViewUniformShaderParamete
 	SHADER_PARAMETER_TEXTURE(Texture3D, HairScatteringLUTTexture)
 	SHADER_PARAMETER_SAMPLER(SamplerState, HairScatteringLUTSampler)
 
-	SHADER_PARAMETER_SRV(StructuredBuffer<float4>, WaterIndirection)
-	SHADER_PARAMETER_SRV(StructuredBuffer<float4>, WaterData)
+	SHADER_PARAMETER_SRV(Buffer<float4>, WaterIndirection)
+	SHADER_PARAMETER_SRV(Buffer<float4>, WaterData)
 
 	SHADER_PARAMETER_UAV(RWBuffer<uint>, VTFeedbackBuffer)
 	SHADER_PARAMETER_SRV(Buffer<uint>, EditorVisualizeLevelInstanceIds)
diff --git a/Engine/Source/Runtime/OpenGLDrv/Private/OpenGLUAV.cpp b/Engine/Source/Runtime/OpenGLDrv/Private/OpenGLUAV.cpp
--- a/Engine/Source/Runtime/OpenGLDrv/Private/OpenGLUAV.cpp
+++ b/Engine/Source/Runtime/OpenGLDrv/Private/OpenGLUAV.cpp
@@ -39,81 +39,51 @@ FShaderResourceViewRHIRef FOpenGLDynamicRHI::RHICreateShaderResourceView(const F
 {
 	FShaderResourceViewInitializer::FBufferShaderResourceViewInitializer Desc = Initializer.AsBufferSRV();
 	FRHIBuffer* BufferRHI = Desc.Buffer;
-
-	switch (Initializer.GetType())
+	if (BufferRHI == nullptr)
 	{
-		case FShaderResourceViewInitializer::EType::VertexBufferSRV:
+		return new FOpenGLShaderResourceViewProxy([=](FRHIShaderResourceView* OwnerRHI)
 		{
-			const uint8 Format = Desc.Format;
-
-			FShaderResourceViewRHIRef Result = new FOpenGLShaderResourceViewProxy([=, OGLRHI = this](FRHIShaderResourceView* OwnerRHI)
-			{
-				VERIFY_GL_SCOPE();
-				GLuint TextureID = 0;
-				if (FOpenGL::SupportsResourceView())
-				{
-					FOpenGL::GenTextures(1, &TextureID);
-					UE_CLOG(!GPixelFormats[Format].Supported, LogRHI, Error, TEXT("Unsupported EPixelFormat %d"), Format);
-					if (BufferRHI)
-					{
-						FOpenGLBuffer* VertexBuffer = FOpenGLDynamicRHI::ResourceCast(BufferRHI);
-
-						const uint32 FormatBPP = GPixelFormats[Format].BlockBytes;
-
-						const FOpenGLTextureFormat& GLFormat = GOpenGLTextureFormats[Format];
-
-						// Use a texture stage that's not likely to be used for draws, to avoid waiting
-						OGLRHI->CachedSetupTextureStage(OGLRHI->GetContextStateForCurrentContext(), FOpenGL::GetMaxCombinedTextureImageUnits() - 1, GL_TEXTURE_BUFFER, TextureID, -1, 1);
-						BindGLTexBufferRange(GL_TEXTURE_BUFFER, GLFormat.InternalFormat[0], VertexBuffer->Resource, Desc.StartOffsetBytes, Desc.NumElements, FormatBPP);
-					}
-				}
-				// No need to restore texture stage; leave it like this,
-				// and the next draw will take care of cleaning it up; or
-				// next operation that needs the stage will switch something else in on it.
-
-				return new FOpenGLShaderResourceView(OGLRHI, TextureID, GL_TEXTURE_BUFFER, BufferRHI, Format);
-			});
-
-			return Result;
-		}
+			return new FOpenGLShaderResourceView(this, 0, GL_TEXTURE_BUFFER);
+		});
+	}
 
-		case FShaderResourceViewInitializer::EType::StructuredBufferSRV:
-		{
-			return new FOpenGLShaderResourceViewProxy([=](FRHIShaderResourceView* OwnerRHI)
-			{
-				VERIFY_GL_SCOPE();
-				FOpenGLBuffer* BufferGL = ResourceCast(BufferRHI);
-				return new FOpenGLShaderResourceView(this, BufferGL->Resource, GL_SHADER_STORAGE_BUFFER, BufferRHI);
-			});
-		}
+	EPixelFormat Format = Desc.Format;
+	
+	if (Initializer.GetType() == FShaderResourceViewInitializer::EType::IndexBufferSRV && Format == PF_Unknown)
+	{		
+		uint32 Stride = BufferRHI->GetStride();
+		Format = (Stride == 2) ? PF_R16_UINT : PF_R32_UINT;
+	}
 
-		case FShaderResourceViewInitializer::EType::IndexBufferSRV:
+	if (Format != PF_Unknown)
+	{
+		UE_CLOG(!GPixelFormats[Format].Supported, LogRHI, Error, TEXT("Unsupported EPixelFormat %d"), Format);
+		
+		return new FOpenGLShaderResourceViewProxy([=, OGLRHI = this](FRHIShaderResourceView* OwnerRHI)
 		{
-			return new FOpenGLShaderResourceViewProxy([=](FRHIShaderResourceView* OwnerRHI)
-			{
-				VERIFY_GL_SCOPE();
-				GLuint TextureID = 0;
-				if (FOpenGL::SupportsResourceView())
-				{
-					FOpenGL::GenTextures(1, &TextureID);
-					if (BufferRHI)
-					{
-						FOpenGLBuffer* IndexBuffer = ResourceCast(BufferRHI);
-						CachedSetupTextureStage(GetContextStateForCurrentContext(), FOpenGL::GetMaxCombinedTextureImageUnits() - 1, GL_TEXTURE_BUFFER, TextureID, -1, 1);
-						uint32 Stride = BufferRHI->GetStride();
-						GLenum Format = (Stride == 2) ? GL_R16UI : GL_R32UI;
-						BindGLTexBufferRange(GL_TEXTURE_BUFFER, Format, IndexBuffer->Resource, Desc.StartOffsetBytes, Desc.NumElements, Stride);
-					}
-				}
-				return new FOpenGLShaderResourceView(this, TextureID, GL_TEXTURE_BUFFER, BufferRHI);
-			});
-		}
+			VERIFY_GL_SCOPE();
+			GLuint TextureID = 0;
+			FOpenGL::GenTextures(1, &TextureID);
+			FOpenGLBuffer* GLBuffer = FOpenGLDynamicRHI::ResourceCast(BufferRHI);
+			const uint32 FormatBPP = GPixelFormats[Format].BlockBytes;
+			const FOpenGLTextureFormat& GLFormat = GOpenGLTextureFormats[Format];
+			// Use a texture stage that's not likely to be used for draws, to avoid waiting
+			OGLRHI->CachedSetupTextureStage(OGLRHI->GetContextStateForCurrentContext(), FOpenGL::GetMaxCombinedTextureImageUnits() - 1, GL_TEXTURE_BUFFER, TextureID, -1, 1);
+			BindGLTexBufferRange(GL_TEXTURE_BUFFER, GLFormat.InternalFormat[0], GLBuffer->Resource, Desc.StartOffsetBytes, Desc.NumElements, FormatBPP);
+			return new FOpenGLShaderResourceView(OGLRHI, TextureID, GL_TEXTURE_BUFFER, BufferRHI, Format);
+		});
+	}
+	else
+	{
+		//TODO: add range views for SSBO
+		ensure(Desc.IsWholeResource());
 
-		default:
+		return new FOpenGLShaderResourceViewProxy([=](FRHIShaderResourceView* OwnerRHI)
 		{
-			checkNoEntry();
-			return nullptr;
-		}
+			VERIFY_GL_SCOPE();
+			FOpenGLBuffer* BufferGL = ResourceCast(BufferRHI);
+			return new FOpenGLShaderResourceView(this, BufferGL->Resource, GL_SHADER_STORAGE_BUFFER, BufferRHI);
+		});
 	}
 }
 
diff --git a/Engine/Source/Runtime/Renderer/Private/SceneRendering.cpp b/Engine/Source/Runtime/Renderer/Private/SceneRendering.cpp
--- a/Engine/Source/Runtime/Renderer/Private/SceneRendering.cpp
+++ b/Engine/Source/Runtime/Renderer/Private/SceneRendering.cpp
@@ -1830,8 +1830,8 @@ void FViewInfo::SetupUniformBufferParameters(
 	}
 	else
 	{
-		ViewUniformShaderParameters.WaterIndirection = GIdentityPrimitiveBuffer.PrimitiveSceneDataBufferSRV;
-		ViewUniformShaderParameters.WaterData = GIdentityPrimitiveBuffer.PrimitiveSceneDataBufferSRV;
+		ViewUniformShaderParameters.WaterIndirection = GWhiteVertexBufferWithSRV->ShaderResourceViewRHI;
+		ViewUniformShaderParameters.WaterData = GWhiteVertexBufferWithSRV->ShaderResourceViewRHI;
 	}
 
 	ViewUniformShaderParameters.VTFeedbackBuffer = GVirtualTextureFeedbackBuffer.GetUAV();
