Rename UE4 -> UnrealEditor

diff --git a/Engine/Source/Runtime/OpenGLDrv/Private/Linux/OpenGLLinux.cpp b/Engine/Source/Runtime/OpenGLDrv/Private/Linux/OpenGLLinux.cpp
--- a/Engine/Source/Runtime/OpenGLDrv/Private/Linux/OpenGLLinux.cpp
+++ b/Engine/Source/Runtime/OpenGLDrv/Private/Linux/OpenGLLinux.cpp
@@ -134,7 +134,7 @@ void Linux_PlatformCreateDummyGLWindow( FPlatformOpenGLContext *OutContext )
 	}
 	else
 	{
-		SDL_SetWindowTitle(DummyWindow, "UE4 Dummy GL window");
+		SDL_SetWindowTitle(DummyWindow, "UnrealEditor Dummy GL window");
 	}
 
 	OutContext->hWnd					= DummyWindow;
diff --git a/Engine/Source/Runtime/OpenGLDrv/Private/OpenGLCommands.cpp b/Engine/Source/Runtime/OpenGLDrv/Private/OpenGLCommands.cpp
--- a/Engine/Source/Runtime/OpenGLDrv/Private/OpenGLCommands.cpp
+++ b/Engine/Source/Runtime/OpenGLDrv/Private/OpenGLCommands.cpp
@@ -766,7 +766,7 @@ void FOpenGLDynamicRHI::SetupTexturesForDraw( FOpenGLContextState& ContextState,
 	{
 		if (!NeededBits[TextureStageIndex])
 		{
-			// Current program doesn't make use of this texture stage. No matter what UE4 wants to have on in,
+			// Current program doesn't make use of this texture stage. No matter what UnrealEditor wants to have on in,
 			// it won't be useful for this draw, so telling OpenGL we don't really need it to give the driver
 			// more leeway in memory management, and avoid false alarms about same texture being set on
 			// texture stage and in framebuffer.
diff --git a/Engine/Source/Runtime/OpenGLDrv/Private/OpenGLES.cpp b/Engine/Source/Runtime/OpenGLDrv/Private/OpenGLES.cpp
--- a/Engine/Source/Runtime/OpenGLDrv/Private/OpenGLES.cpp
+++ b/Engine/Source/Runtime/OpenGLDrv/Private/OpenGLES.cpp
@@ -154,7 +154,7 @@ void FOpenGLES::ProcessQueryGLInt()
 	if (MaxVertexAttribs < 16)
 	{
 		UE_LOG(LogRHI, Error,
-			TEXT("Device reports support for %d vertex attributes, UE4 requires 16. Rendering artifacts may occur."),
+			TEXT("Device reports support for %d vertex attributes, UnrealEditor requires 16. Rendering artifacts may occur."),
 			MaxVertexAttribs
 		);
 	}
@@ -180,7 +180,7 @@ void FOpenGLES::ProcessQueryGLInt()
 	if (MaxVertexUniformComponents < RequiredMaxVertexUniformComponents)
 	{
 		UE_LOG(LogRHI, Warning,
-			TEXT("Device reports support for %d vertex uniform vectors, UE4 requires %d. Rendering artifacts may occur, especially with skeletal meshes. Some drivers, e.g. iOS, report a smaller number than is actually supported."),
+			TEXT("Device reports support for %d vertex uniform vectors, UnrealEditor requires %d. Rendering artifacts may occur, especially with skeletal meshes. Some drivers, e.g. iOS, report a smaller number than is actually supported."),
 			MaxVertexUniformComponents,
 			RequiredMaxVertexUniformComponents
 		);
diff --git a/Engine/Source/Runtime/OpenGLDrv/Public/OpenGLResources.h b/Engine/Source/Runtime/OpenGLDrv/Public/OpenGLResources.h
--- a/Engine/Source/Runtime/OpenGLDrv/Public/OpenGLResources.h
+++ b/Engine/Source/Runtime/OpenGLDrv/Public/OpenGLResources.h
@@ -616,7 +616,7 @@ public:
 				// we can stick with the STATIC_DRAW hint workaround for glBufferData().
 				//
 				// We haven't seen the buffer mapping/unmapping issue show up elsewhere in the pipeline in our test scenes. 
-				// However, depending on the UE4 features that are used, this issue might pop up elsewhere that we're yet to see.
+				// However, depending on the UnrealEditor features that are used, this issue might pop up elsewhere that we're yet to see.
 				// As there are concerns for maximum number of bone matrices, going for the GL_STATIC_DRAW hint should be safer, 
 				// given the fact that it won't constrain the actual usage of the data store as per the OpenGL4 spec.
 #if PLATFORM_LUMINGL4
