Tessellation Removal: CreateBoundShaderState changes to remove hull and domain shaders. Making CreateHullShader and CreateDomainShader optional.

diff --git a/Engine/Source/Runtime/Apple/MetalRHI/Private/MetalDynamicRHI.h b/Engine/Source/Runtime/Apple/MetalRHI/Private/MetalDynamicRHI.h
--- a/Engine/Source/Runtime/Apple/MetalRHI/Private/MetalDynamicRHI.h
+++ b/Engine/Source/Runtime/Apple/MetalRHI/Private/MetalDynamicRHI.h
@@ -51,7 +51,7 @@ public:
 	virtual FDomainShaderRHIRef RHICreateDomainShader(TArrayView<const uint8> Code, const FSHAHash& Hash) final override;
 	virtual FGeometryShaderRHIRef RHICreateGeometryShader(TArrayView<const uint8> Code, const FSHAHash& Hash) final override;
 	virtual FComputeShaderRHIRef RHICreateComputeShader(TArrayView<const uint8> Code, const FSHAHash& Hash) final override;
-	virtual FBoundShaderStateRHIRef RHICreateBoundShaderState(FRHIVertexDeclaration* VertexDeclaration, FRHIVertexShader* VertexShader, FRHIHullShader* HullShader, FRHIDomainShader* DomainShader, FRHIPixelShader* PixelShader, FRHIGeometryShader* GeometryShader) final override;
+	virtual FBoundShaderStateRHIRef RHICreateBoundShaderState(FRHIVertexDeclaration* VertexDeclaration, FRHIVertexShader* VertexShader, FRHIPixelShader* PixelShader, FRHIGeometryShader* GeometryShader) final override;
 	virtual FGraphicsPipelineStateRHIRef RHICreateGraphicsPipelineState(const FGraphicsPipelineStateInitializer& Initializer) final override;
 	virtual TRefCountPtr<FRHIComputePipelineState> RHICreateComputePipelineState(FRHIComputeShader* ComputeShader) final override;
 	virtual FUniformBufferRHIRef RHICreateUniformBuffer(const void* Contents, const FRHIUniformBufferLayout& Layout, EUniformBufferUsage Usage, EUniformBufferValidation Validation) final override;
diff --git a/Engine/Source/Runtime/Apple/MetalRHI/Private/MetalDynamicRHI_Shaders.cpp b/Engine/Source/Runtime/Apple/MetalRHI/Private/MetalDynamicRHI_Shaders.cpp
--- a/Engine/Source/Runtime/Apple/MetalRHI/Private/MetalDynamicRHI_Shaders.cpp
+++ b/Engine/Source/Runtime/Apple/MetalRHI/Private/MetalDynamicRHI_Shaders.cpp
@@ -152,8 +152,6 @@ FRHIShaderLibraryRef FMetalDynamicRHI::RHICreateShaderLibrary(EShaderPlatform Pl
 FBoundShaderStateRHIRef FMetalDynamicRHI::RHICreateBoundShaderState(
 	FRHIVertexDeclaration* VertexDeclarationRHI,
 	FRHIVertexShader* VertexShaderRHI,
-	FRHIHullShader* HullShaderRHI,
-	FRHIDomainShader* DomainShaderRHI,
 	FRHIPixelShader* PixelShaderRHI,
 	FRHIGeometryShader* GeometryShaderRHI)
 {
diff --git a/Engine/Source/Runtime/D3D12RHI/Private/D3D12RHIPrivate.h b/Engine/Source/Runtime/D3D12RHI/Private/D3D12RHIPrivate.h
--- a/Engine/Source/Runtime/D3D12RHI/Private/D3D12RHIPrivate.h
+++ b/Engine/Source/Runtime/D3D12RHI/Private/D3D12RHIPrivate.h
@@ -276,7 +276,7 @@ public:
 	virtual FStagingBufferRHIRef RHICreateStagingBuffer() final override;
 	virtual void* RHILockStagingBuffer(FRHIStagingBuffer* StagingBuffer, FRHIGPUFence* Fence, uint32 Offset, uint32 SizeRHI) final override;
     virtual void RHIUnlockStagingBuffer(FRHIStagingBuffer* StagingBuffer) final override;
-	virtual FBoundShaderStateRHIRef RHICreateBoundShaderState(FRHIVertexDeclaration* VertexDeclaration, FRHIVertexShader* VertexShader, FRHIHullShader* HullShader, FRHIDomainShader* DomainShader, FRHIPixelShader* PixelShader, FRHIGeometryShader* GeometryShader) final override;
+	virtual FBoundShaderStateRHIRef RHICreateBoundShaderState(FRHIVertexDeclaration* VertexDeclaration, FRHIVertexShader* VertexShader, FRHIPixelShader* PixelShader, FRHIGeometryShader* GeometryShader) final override;
 	FBoundShaderStateRHIRef DX12CreateBoundShaderState(const FBoundShaderStateInput& BoundShaderStateInput);
 	virtual FGraphicsPipelineStateRHIRef RHICreateGraphicsPipelineState(const FGraphicsPipelineStateInitializer& Initializer) final override;
 	virtual TRefCountPtr<FRHIComputePipelineState> RHICreateComputePipelineState(FRHIComputeShader* ComputeShader) final override;
diff --git a/Engine/Source/Runtime/D3D12RHI/Private/D3D12Shader.h b/Engine/Source/Runtime/D3D12RHI/Private/D3D12Shader.h
--- a/Engine/Source/Runtime/D3D12RHI/Private/D3D12Shader.h
+++ b/Engine/Source/Runtime/D3D12RHI/Private/D3D12Shader.h
@@ -198,8 +198,6 @@ public:
 		FRHIVertexDeclaration* InVertexDeclarationRHI,
 		FRHIVertexShader* InVertexShaderRHI,
 		FRHIPixelShader* InPixelShaderRHI,
-		FRHIHullShader* InHullShaderRHI,
-		FRHIDomainShader* InDomainShaderRHI,
 		FRHIGeometryShader* InGeometryShaderRHI,
 		FD3D12Adapter* InAdapter
 	);
diff --git a/Engine/Source/Runtime/D3D12RHI/Private/D3D12Shaders.cpp b/Engine/Source/Runtime/D3D12RHI/Private/D3D12Shaders.cpp
--- a/Engine/Source/Runtime/D3D12RHI/Private/D3D12Shaders.cpp
+++ b/Engine/Source/Runtime/D3D12RHI/Private/D3D12Shaders.cpp
@@ -438,12 +438,10 @@ FD3D12BoundShaderState::FD3D12BoundShaderState(
 	FRHIVertexDeclaration* InVertexDeclarationRHI,
 	FRHIVertexShader* InVertexShaderRHI,
 	FRHIPixelShader* InPixelShaderRHI,
-	FRHIHullShader* InHullShaderRHI,
-	FRHIDomainShader* InDomainShaderRHI,
 	FRHIGeometryShader* InGeometryShaderRHI,
 	FD3D12Adapter* InAdapter
 	) :
-	CacheLink(InVertexDeclarationRHI, InVertexShaderRHI, InPixelShaderRHI, InHullShaderRHI, InDomainShaderRHI, InGeometryShaderRHI, this)
+	CacheLink(InVertexDeclarationRHI, InVertexShaderRHI, InPixelShaderRHI, InGeometryShaderRHI, this)
 {
 	INC_DWORD_STAT(STAT_D3D12NumBoundShaderState);
 
@@ -507,8 +505,6 @@ FBoundShaderStateRHIRef FD3D12DynamicRHI::DX12CreateBoundShaderState(const FBoun
 		BoundShaderStateInput.VertexDeclarationRHI,
 		BoundShaderStateInput.VertexShaderRHI,
 		BoundShaderStateInput.PixelShaderRHI,
-		BoundShaderStateInput.HullShaderRHI,
-		BoundShaderStateInput.DomainShaderRHI,
 		BoundShaderStateInput.GeometryShaderRHI,
 		BoundShaderStateInput.MeshShaderRHI,
 		BoundShaderStateInput.AmplificationShaderRHI
@@ -525,8 +521,6 @@ FBoundShaderStateRHIRef FD3D12DynamicRHI::DX12CreateBoundShaderState(const FBoun
 		BoundShaderStateInput.VertexDeclarationRHI,
 		BoundShaderStateInput.VertexShaderRHI,
 		BoundShaderStateInput.PixelShaderRHI,
-		BoundShaderStateInput.HullShaderRHI,
-		BoundShaderStateInput.DomainShaderRHI,
 		BoundShaderStateInput.GeometryShaderRHI,
 		BoundShaderStateInput.MeshShaderRHI,
 		BoundShaderStateInput.AmplificationShaderRHI
@@ -557,8 +551,6 @@ FBoundShaderStateRHIRef FD3D12DynamicRHI::DX12CreateBoundShaderState(const FBoun
 				BoundShaderStateInput.VertexDeclarationRHI,
 				BoundShaderStateInput.VertexShaderRHI,
 				BoundShaderStateInput.PixelShaderRHI,
-				BoundShaderStateInput.HullShaderRHI,
-				BoundShaderStateInput.DomainShaderRHI,
 				BoundShaderStateInput.GeometryShaderRHI,
 				&GetAdapter());
 		}
@@ -578,8 +570,6 @@ FBoundShaderStateRHIRef FD3D12DynamicRHI::DX12CreateBoundShaderState(const FBoun
 FBoundShaderStateRHIRef FD3D12DynamicRHI::RHICreateBoundShaderState(
 	FRHIVertexDeclaration* VertexDeclarationRHI,
 	FRHIVertexShader* VertexShaderRHI,
-	FRHIHullShader* HullShaderRHI,
-	FRHIDomainShader* DomainShaderRHI,
 	FRHIPixelShader* PixelShaderRHI,
 	FRHIGeometryShader* GeometryShaderRHI
 	)
@@ -589,7 +579,5 @@ FBoundShaderStateRHIRef FD3D12DynamicRHI::RHICreateBoundShaderState(
 		, GeometryShaderRHI
 #endif
 	);
-	Inputs.HullShaderRHI = HullShaderRHI;
-	Inputs.DomainShaderRHI = DomainShaderRHI;
 	return DX12CreateBoundShaderState(Inputs);
 }
diff --git a/Engine/Source/Runtime/EmptyRHI/Private/EmptyShaders.cpp b/Engine/Source/Runtime/EmptyRHI/Private/EmptyShaders.cpp
--- a/Engine/Source/Runtime/EmptyRHI/Private/EmptyShaders.cpp
+++ b/Engine/Source/Runtime/EmptyRHI/Private/EmptyShaders.cpp
@@ -64,24 +64,18 @@ FEmptyBoundShaderState::FEmptyBoundShaderState(
 			FRHIVertexDeclaration* InVertexDeclarationRHI,
 			FRHIVertexShader* InVertexShaderRHI,
 			FRHIPixelShader* InPixelShaderRHI,
-			FRHIHullShader* InHullShaderRHI,
-	FRHIDomainShader* InDomainShaderRHI,
 	FRHIGeometryShader* InGeometryShaderRHI)
-	:	CacheLink(InVertexDeclarationRHI,InVertexShaderRHI,InPixelShaderRHI,InHullShaderRHI,InDomainShaderRHI,InGeometryShaderRHI,this)
+	:	CacheLink(InVertexDeclarationRHI,InVertexShaderRHI,InPixelShaderRHI,InGeometryShaderRHI,this)
 {
 	FEmptyVertexDeclaration* InVertexDeclaration = FEmptyDynamicRHI::ResourceCast(InVertexDeclarationRHI);
 	FEmptyVertexShader* InVertexShader = FEmptyDynamicRHI::ResourceCast(InVertexShaderRHI);
 	FEmptyPixelShader* InPixelShader = FEmptyDynamicRHI::ResourceCast(InPixelShaderRHI);
-	FEmptyHullShader* InHullShader = FEmptyDynamicRHI::ResourceCast(InHullShaderRHI);
-	FEmptyDomainShader* InDomainShader = FEmptyDynamicRHI::ResourceCast(InDomainShaderRHI);
 	FEmptyGeometryShader* InGeometryShader = FEmptyDynamicRHI::ResourceCast(InGeometryShaderRHI);
 
 	// cache everything
 	VertexDeclaration = InVertexDeclaration;
 	VertexShader = InVertexShader;
 	PixelShader = InPixelShader;
-	HullShader = InHullShader;
-	DomainShader = InDomainShader;
 	GeometryShader = InGeometryShader;
 }
 
@@ -93,8 +87,6 @@ FEmptyBoundShaderState::~FEmptyBoundShaderState()
 FBoundShaderStateRHIRef FEmptyDynamicRHI::RHICreateBoundShaderState(
 	FRHIVertexDeclaration* VertexDeclarationRHI,
 	FRHIVertexShader* VertexShaderRHI,
-	FRHIHullShader* HullShaderRHI,
-	FRHIDomainShader* DomainShaderRHI,
 	FRHIPixelShader* PixelShaderRHI,
 	FRHIGeometryShader* GeometryShaderRHI
 	)
@@ -105,8 +97,6 @@ FBoundShaderStateRHIRef FEmptyDynamicRHI::RHICreateBoundShaderState(
 		VertexDeclarationRHI,
 		VertexShaderRHI,
 		PixelShaderRHI,
-		HullShaderRHI,
-		DomainShaderRHI,
 		GeometryShaderRHI
 		);
 
@@ -117,6 +107,6 @@ FBoundShaderStateRHIRef FEmptyDynamicRHI::RHICreateBoundShaderState(
 	}
 	else
 	{
-		return new FEmptyBoundShaderState(VertexDeclarationRHI,VertexShaderRHI,PixelShaderRHI,HullShaderRHI,DomainShaderRHI,GeometryShaderRHI);
+		return new FEmptyBoundShaderState(VertexDeclarationRHI,VertexShaderRHI,PixelShaderRHI,GeometryShaderRHI);
 	}
 }
diff --git a/Engine/Source/Runtime/EmptyRHI/Public/EmptyRHI.h b/Engine/Source/Runtime/EmptyRHI/Public/EmptyRHI.h
--- a/Engine/Source/Runtime/EmptyRHI/Public/EmptyRHI.h
+++ b/Engine/Source/Runtime/EmptyRHI/Public/EmptyRHI.h
@@ -47,7 +47,7 @@ public:
 	virtual FDomainShaderRHIRef RHICreateDomainShader(TArrayView<const uint8> Code, const FSHAHash& Hash) final override;
 	virtual FGeometryShaderRHIRef RHICreateGeometryShader(TArrayView<const uint8> Code, const FSHAHash& Hash) final override;
 	virtual FComputeShaderRHIRef RHICreateComputeShader(TArrayView<const uint8> Code, const FSHAHash& Hash) final override;
-	virtual FBoundShaderStateRHIRef RHICreateBoundShaderState(FRHIVertexDeclaration* VertexDeclaration, FRHIVertexShader* VertexShader, FRHIHullShader* HullShader, FRHIDomainShaderFRHIDomainShader*, FRHIPixelShader* PixelShader, FRHIGeometryShader* GeometryShader) final override;
+	virtual FBoundShaderStateRHIRef RHICreateBoundShaderState(FRHIVertexDeclaration* VertexDeclaration, FRHIVertexShader* VertexShader, FRHIPixelShader* PixelShader, FRHIGeometryShader* GeometryShader) final override;
 	virtual FUniformBufferRHIRef RHICreateUniformBuffer(const void* Contents, const FRHIUniformBufferLayout& Layout, EUniformBufferUsage Usage, EUniformBufferValidation Validation) final override;
 	virtual FBufferRHIRef RHICreateIndexBuffer(uint32 Stride, uint32 Size, uint32 InUsage, FRHIResourceCreateInfo& CreateInfo) final override;
 	virtual void* LockIndexBuffer_BottomOfPipe(FRHICommandListImmediate& RHICmdList, FRHIBuffer* IndexBuffer, uint32 Offset, uint32 Size, EResourceLockMode LockMode) final override;
diff --git a/Engine/Source/Runtime/EmptyRHI/Public/EmptyResources.h b/Engine/Source/Runtime/EmptyRHI/Public/EmptyResources.h
--- a/Engine/Source/Runtime/EmptyRHI/Public/EmptyResources.h
+++ b/Engine/Source/Runtime/EmptyRHI/Public/EmptyResources.h
@@ -46,8 +46,6 @@ public:
 		FRHIVertexDeclaration* InVertexDeclarationRHI,
 		FRHIVertexShader* InVertexShaderRHI,
 		FRHIPixelShader* InPixelShaderRHI,
-		FRHIHullShader* InHullShaderRHI,
-		FRHIDomainShader* InDomainShaderRHI,
 		FRHIGeometryShader* InGeometryShaderRHI);
 
 	/**
diff --git a/Engine/Source/Runtime/NullDrv/Public/NullRHI.h b/Engine/Source/Runtime/NullDrv/Public/NullRHI.h
--- a/Engine/Source/Runtime/NullDrv/Public/NullRHI.h
+++ b/Engine/Source/Runtime/NullDrv/Public/NullRHI.h
@@ -73,7 +73,7 @@ public:
 	}
 
 
-	virtual FBoundShaderStateRHIRef RHICreateBoundShaderState(FRHIVertexDeclaration* VertexDeclaration, FRHIVertexShader* VertexShader, FRHIHullShader* HullShader, FRHIDomainShader* DomainShader, FRHIPixelShader* PixelShader, FRHIGeometryShader* GeometryShader) final override
+	virtual FBoundShaderStateRHIRef RHICreateBoundShaderState(FRHIVertexDeclaration* VertexDeclaration, FRHIVertexShader* VertexShader, FRHIPixelShader* PixelShader, FRHIGeometryShader* GeometryShader) final override
 	{ 
 		return new FRHIBoundShaderState(); 
 	}
diff --git a/Engine/Source/Runtime/OpenGLDrv/Private/OpenGLShaders.cpp b/Engine/Source/Runtime/OpenGLDrv/Private/OpenGLShaders.cpp
--- a/Engine/Source/Runtime/OpenGLDrv/Private/OpenGLShaders.cpp
+++ b/Engine/Source/Runtime/OpenGLDrv/Private/OpenGLShaders.cpp
@@ -3130,8 +3130,6 @@ static FCriticalSection GProgramBinaryCacheCS;
 FBoundShaderStateRHIRef FOpenGLDynamicRHI::RHICreateBoundShaderState_OnThisThread(
 	FRHIVertexDeclaration* VertexDeclarationRHI,
 	FRHIVertexShader* VertexShaderRHI,
-	FRHIHullShader* HullShaderRHI,
-	FRHIDomainShader* DomainShaderRHI,
 	FRHIPixelShader* PixelShaderRHI,
 	FRHIGeometryShader* GeometryShaderRHI,
 	bool bFromPSOFileCache
@@ -3151,12 +3149,10 @@ FBoundShaderStateRHIRef FOpenGLDynamicRHI::RHICreateBoundShaderState_OnThisThrea
 		PixelShaderRHI = TShaderMapRef<FNULLPS>(GetGlobalShaderMap(GMaxRHIFeatureLevel)).GetPixelShader();
 	}
 
-	auto CreateConfig = [VertexShaderRHI, HullShaderRHI, DomainShaderRHI, PixelShaderRHI, GeometryShaderRHI]()
+	auto CreateConfig = [VertexShaderRHI, PixelShaderRHI, GeometryShaderRHI]()
 	{
 		FOpenGLVertexShader* VertexShader = ResourceCast(VertexShaderRHI);
 		FOpenGLPixelShader* PixelShader = ResourceCast(PixelShaderRHI);
-		FOpenGLHullShader* HullShader = ResourceCast(HullShaderRHI);
-		FOpenGLDomainShader* DomainShader = ResourceCast(DomainShaderRHI);
 		FOpenGLGeometryShader* GeometryShader = ResourceCast(GeometryShaderRHI);
 
 		FOpenGLLinkedProgramConfiguration Config;
@@ -3169,39 +3165,14 @@ FBoundShaderStateRHIRef FOpenGLDynamicRHI::RHICreateBoundShaderState_OnThisThrea
 		Config.Shaders[CrossCompiler::SHADER_STAGE_VERTEX].Resource = VertexShader->Resource;
 		Config.ProgramKey.ShaderHashes[CrossCompiler::SHADER_STAGE_VERTEX] = VertexShaderRHI->GetHash();
 
-		if (FOpenGL::SupportsTessellation())
-		{
-			if (HullShader)
-			{
-				check(VertexShader);
-				BindShaderStage(Config, CrossCompiler::SHADER_STAGE_HULL, HullShader, HullShaderRHI->GetHash(), CrossCompiler::SHADER_STAGE_VERTEX, VertexShader);
-			}
-			if (DomainShader)
-			{
-				check(HullShader);
-				BindShaderStage(Config, CrossCompiler::SHADER_STAGE_DOMAIN, DomainShader, DomainShaderRHI->GetHash(), CrossCompiler::SHADER_STAGE_HULL, HullShader);
-			}
-		}
-
 		if (GeometryShader)
 		{
-			check(DomainShader || VertexShader);
-			if (DomainShader)
-			{
-				BindShaderStage(Config, CrossCompiler::SHADER_STAGE_GEOMETRY, GeometryShader, GeometryShaderRHI->GetHash(), CrossCompiler::SHADER_STAGE_DOMAIN, DomainShader);
-			}
-			else
-			{
-				BindShaderStage(Config, CrossCompiler::SHADER_STAGE_GEOMETRY, GeometryShader, GeometryShaderRHI->GetHash(), CrossCompiler::SHADER_STAGE_VERTEX, VertexShader);
-			}
+			check(VertexShader);
+			BindShaderStage(Config, CrossCompiler::SHADER_STAGE_GEOMETRY, GeometryShader, GeometryShaderRHI->GetHash(), CrossCompiler::SHADER_STAGE_VERTEX, VertexShader);
 		}
 
-		check(DomainShader || GeometryShader || VertexShader);
-		if (DomainShader)
-		{
-			BindShaderStage(Config, CrossCompiler::SHADER_STAGE_PIXEL, PixelShader, PixelShaderRHI->GetHash(), CrossCompiler::SHADER_STAGE_DOMAIN, DomainShader);
-		}
-		else if (GeometryShader)
+		check(GeometryShader || VertexShader);
+		if (GeometryShader)
 		{
 			BindShaderStage(Config, CrossCompiler::SHADER_STAGE_PIXEL, PixelShader, PixelShaderRHI->GetHash(), CrossCompiler::SHADER_STAGE_GEOMETRY, GeometryShader);
 		}
@@ -3217,8 +3188,6 @@ FBoundShaderStateRHIRef FOpenGLDynamicRHI::RHICreateBoundShaderState_OnThisThrea
 		VertexDeclarationRHI,
 		VertexShaderRHI,
 		PixelShaderRHI,
-		HullShaderRHI,
-		DomainShaderRHI,
 		GeometryShaderRHI
 		);
 
@@ -3293,8 +3262,6 @@ FBoundShaderStateRHIRef FOpenGLDynamicRHI::RHICreateBoundShaderState_OnThisThrea
 			{
 				FOpenGLVertexShader* VertexShader = ResourceCast(VertexShaderRHI);
 				FOpenGLPixelShader* PixelShader = ResourceCast(PixelShaderRHI);
-				FOpenGLHullShader* HullShader = ResourceCast(HullShaderRHI);
-				FOpenGLDomainShader* DomainShader = ResourceCast(DomainShaderRHI);
 				FOpenGLGeometryShader* GeometryShader = ResourceCast(GeometryShaderRHI);
 		
 				// Make sure we have OpenGL context set up, and invalidate the parameters cache and current program (as we'll link a new one soon)
@@ -3320,17 +3287,6 @@ FBoundShaderStateRHIRef FOpenGLDynamicRHI::RHICreateBoundShaderState_OnThisThrea
 					{
 						UE_LOG(LogRHI, Error, TEXT("Geometry Shader:\n%s"), ANSI_TO_TCHAR(GeometryShader->GlslCode.GetData()));
 					}
-					if (FOpenGL::SupportsTessellation())
-					{
-						if (HullShader)
-						{
-							UE_LOG(LogRHI, Error, TEXT("Hull Shader:\n%s"), ANSI_TO_TCHAR(HullShader->GlslCode.GetData()));
-						}
-						if (DomainShader)
-						{
-							UE_LOG(LogRHI, Error, TEXT("Domain Shader:\n%s"), ANSI_TO_TCHAR(DomainShader->GlslCode.GetData()));
-						}
-					}
 #endif //DEBUG_GL_SHADERS
 					FName LinkFailurePanic = bFromPSOFileCache ? FName("FailedProgramLinkDuringPrecompile") : FName("FailedProgramLink");
 					RHIGetPanicDelegate().ExecuteIfBound(LinkFailurePanic);
@@ -3355,9 +3311,7 @@ FBoundShaderStateRHIRef FOpenGLDynamicRHI::RHICreateBoundShaderState_OnThisThrea
 			VertexDeclarationRHI,
 			VertexShaderRHI,
 			PixelShaderRHI,
-			GeometryShaderRHI,
-			HullShaderRHI,
-			DomainShaderRHI
+			GeometryShaderRHI
 			);
 
 		return BoundShaderState;
@@ -3555,20 +3509,16 @@ FOpenGLBoundShaderState::FOpenGLBoundShaderState(
 	FRHIVertexDeclaration* InVertexDeclarationRHI,
 	FRHIVertexShader* InVertexShaderRHI,
 	FRHIPixelShader* InPixelShaderRHI,
-	FRHIGeometryShader* InGeometryShaderRHI,
-	FRHIHullShader* InHullShaderRHI,
-	FRHIDomainShader* InDomainShaderRHI
+	FRHIGeometryShader* InGeometryShaderRHI
 	)
 	:	CacheLink(InVertexDeclarationRHI, InVertexShaderRHI, InPixelShaderRHI,
-		InHullShaderRHI, InDomainShaderRHI,	InGeometryShaderRHI, this)
+		InGeometryShaderRHI, this)
 {
 	FOpenGLVertexDeclaration* InVertexDeclaration = FOpenGLDynamicRHI::ResourceCast(InVertexDeclarationRHI);
 	VertexDeclaration = InVertexDeclaration;
 	VertexShaderProxy = static_cast<FOpenGLVertexShaderProxy*>(InVertexShaderRHI);
 	PixelShaderProxy = static_cast<FOpenGLPixelShaderProxy*>(InPixelShaderRHI);
 	GeometryShaderProxy = static_cast<FOpenGLGeometryShaderProxy*>(InGeometryShaderRHI);
-	HullShaderProxy = static_cast<FOpenGLHullShaderProxy*>(InHullShaderRHI);
-	DomainShaderProxy = static_cast<FOpenGLDomainShaderProxy*>(InDomainShaderRHI);
 
 	LinkedProgram = InLinkedProgram;
 
diff --git a/Engine/Source/Runtime/OpenGLDrv/Public/OpenGLDrv.h b/Engine/Source/Runtime/OpenGLDrv/Public/OpenGLDrv.h
--- a/Engine/Source/Runtime/OpenGLDrv/Public/OpenGLDrv.h
+++ b/Engine/Source/Runtime/OpenGLDrv/Public/OpenGLDrv.h
@@ -954,8 +954,6 @@ public:
 			RHICreateBoundShaderState_internal(
 				PsoInit.BoundShaderState.VertexDeclarationRHI,
 				PsoInit.BoundShaderState.VertexShaderRHI,
-				TESSELLATION_SHADER(PsoInit.BoundShaderState.HullShaderRHI),
-				TESSELLATION_SHADER(PsoInit.BoundShaderState.DomainShaderRHI),
 				PsoInit.BoundShaderState.PixelShaderRHI,
 				GEOMETRY_SHADER(PsoInit.BoundShaderState.GeometryShaderRHI),
 				PsoInit.bFromPSOFileCache
@@ -973,8 +971,6 @@ public:
 		if (bApplyAdditionalState)
 		{
 			ApplyStaticUniformBuffers(PsoInit.BoundShaderState.VertexShaderRHI, ResourceCast(PsoInit.BoundShaderState.VertexShaderRHI));
-			ApplyStaticUniformBuffers(PsoInit.BoundShaderState.HullShaderRHI, ResourceCast(PsoInit.BoundShaderState.HullShaderRHI));
-			ApplyStaticUniformBuffers(PsoInit.BoundShaderState.DomainShaderRHI, ResourceCast(PsoInit.BoundShaderState.DomainShaderRHI));
 			ApplyStaticUniformBuffers(PsoInit.BoundShaderState.GeometryShaderRHI, ResourceCast(PsoInit.BoundShaderState.GeometryShaderRHI));
 			ApplyStaticUniformBuffers(PsoInit.BoundShaderState.PixelShaderRHI, ResourceCast(PsoInit.BoundShaderState.PixelShaderRHI));
 		}
@@ -986,8 +982,6 @@ public:
 	FBoundShaderStateRHIRef RHICreateBoundShaderState_internal(
 		FRHIVertexDeclaration* VertexDeclarationRHI,
 		FRHIVertexShader* VertexShaderRHI,
-		FRHIHullShader* HullShaderRHI,
-		FRHIDomainShader* DomainShaderRHI,
 		FRHIPixelShader* PixelShaderRHI,
 		FRHIGeometryShader* GeometryShaderRHI,
 		bool FromPSOFileCache
@@ -998,8 +992,6 @@ public:
 		RHITHREAD_GLCOMMAND_PROLOGUE()
 			return RHICreateBoundShaderState_OnThisThread(VertexDeclarationRHI,
 				VertexShaderRHI,
-				HullShaderRHI,
-				DomainShaderRHI,
 				PixelShaderRHI,
 				GeometryShaderRHI,
 				FromPSOFileCache);
@@ -1009,8 +1001,6 @@ public:
 	virtual FBoundShaderStateRHIRef RHICreateBoundShaderState(
 		FRHIVertexDeclaration* VertexDeclarationRHI,
 		FRHIVertexShader* VertexShaderRHI,
-		FRHIHullShader* HullShaderRHI,
-		FRHIDomainShader* DomainShaderRHI,
 		FRHIPixelShader* PixelShaderRHI,
 		FRHIGeometryShader* GeometryShaderRHI
 	) final override
@@ -1018,8 +1008,6 @@ public:
 		return RHICreateBoundShaderState_internal(
 			VertexDeclarationRHI,
 			VertexShaderRHI,
-			HullShaderRHI,
-			DomainShaderRHI,
 			PixelShaderRHI,
 			GeometryShaderRHI,
 			false);
@@ -1027,7 +1015,7 @@ public:
 
 	class FOpenGLLinkedProgram* GetLinkedComputeProgram(FRHIComputeShader* ComputeShaderRHI);
 
-	FBoundShaderStateRHIRef RHICreateBoundShaderState_OnThisThread(FRHIVertexDeclaration* VertexDeclaration, FRHIVertexShader* VertexShader, FRHIHullShader* HullShader, FRHIDomainShader* DomainShader, FRHIPixelShader* PixelShader, FRHIGeometryShader* GeometryShader, bool FromPSOFileCache);
+	FBoundShaderStateRHIRef RHICreateBoundShaderState_OnThisThread(FRHIVertexDeclaration* VertexDeclaration, FRHIVertexShader* VertexShader, FRHIPixelShader* PixelShader, FRHIGeometryShader* GeometryShader, bool FromPSOFileCache);
 	void RHIPerFrameRHIFlushComplete();
 
 	FOpenGLGPUProfiler& GetGPUProfilingData() {
diff --git a/Engine/Source/Runtime/OpenGLDrv/Public/OpenGLResources.h b/Engine/Source/Runtime/OpenGLDrv/Public/OpenGLResources.h
--- a/Engine/Source/Runtime/OpenGLDrv/Public/OpenGLResources.h
+++ b/Engine/Source/Runtime/OpenGLDrv/Public/OpenGLResources.h
@@ -1170,9 +1170,7 @@ public:
 		FRHIVertexDeclaration* InVertexDeclarationRHI,
 		FRHIVertexShader* InVertexShaderRHI,
 		FRHIPixelShader* InPixelShaderRHI,
-		FRHIGeometryShader* InGeometryShaderRHI,
-		FRHIHullShader* InHullShaderRHI,
-		FRHIDomainShader* InDomainShaderRHI
+		FRHIGeometryShader* InGeometryShaderRHI
 		);
 
 	const TBitArray<>& GetTextureNeeds(int32& OutMaxTextureStageUsed);
diff --git a/Engine/Source/Runtime/RHI/Private/BoundShaderStateCache.cpp b/Engine/Source/Runtime/RHI/Private/BoundShaderStateCache.cpp
--- a/Engine/Source/Runtime/RHI/Private/BoundShaderStateCache.cpp
+++ b/Engine/Source/Runtime/RHI/Private/BoundShaderStateCache.cpp
@@ -49,13 +49,11 @@ FCachedBoundShaderStateLink::FCachedBoundShaderStateLink(
 	FRHIVertexDeclaration* VertexDeclaration,
 	FRHIVertexShader* VertexShader,
 	FRHIPixelShader* PixelShader,
-	FRHIHullShader* HullShader,
-	FRHIDomainShader* DomainShader,
 	FRHIGeometryShader* GeometryShader,
 	FRHIBoundShaderState* InBoundShaderState,
 	bool bAddToSingleThreadedCache)
 	: BoundShaderState(InBoundShaderState)
-	, Key(VertexDeclaration, VertexShader, PixelShader, HullShader, DomainShader, GeometryShader)
+	, Key(VertexDeclaration, VertexShader, PixelShader, GeometryShader)
 	, bAddedToSingleThreadedCache(bAddToSingleThreadedCache)
 {
 	if (bAddToSingleThreadedCache)
@@ -93,8 +91,6 @@ FCachedBoundShaderStateLink* GetCachedBoundShaderState(
 	FRHIVertexDeclaration* VertexDeclaration,
 	FRHIVertexShader* VertexShader,
 	FRHIPixelShader* PixelShader,
-	FRHIHullShader* HullShader,
-	FRHIDomainShader* DomainShader,
 	FRHIGeometryShader* GeometryShader,
 	FRHIMeshShader* MeshShader,
 	FRHIAmplificationShader* AmplificationShader
@@ -111,7 +107,7 @@ FCachedBoundShaderStateLink* GetCachedBoundShaderState(
 	{
 		// Find the existing bound shader state in the cache.
 		return GetBoundShaderStateCache().FindRef(
-			FBoundShaderStateLookupKey(VertexDeclaration, VertexShader, PixelShader, HullShader, DomainShader, GeometryShader)
+			FBoundShaderStateLookupKey(VertexDeclaration, VertexShader, PixelShader, GeometryShader)
 		);
 	}
 }
@@ -133,8 +129,6 @@ FBoundShaderStateRHIRef GetCachedBoundShaderState_Threadsafe(
 	FRHIVertexDeclaration* VertexDeclaration,
 	FRHIVertexShader* VertexShader,
 	FRHIPixelShader* PixelShader,
-	FRHIHullShader* HullShader,
-	FRHIDomainShader* DomainShader,
 	FRHIGeometryShader* GeometryShader,
 	FRHIMeshShader* MeshShader,
 	FRHIAmplificationShader* AmplificationShader
@@ -153,7 +147,7 @@ FBoundShaderStateRHIRef GetCachedBoundShaderState_Threadsafe(
 	else
 	{
 		CachedBoundShaderStateLink = GetBoundShaderStateCache_Threadsafe().FindRef(
-			FBoundShaderStateLookupKey(VertexDeclaration, VertexShader, PixelShader, HullShader, DomainShader, GeometryShader)
+			FBoundShaderStateLookupKey(VertexDeclaration, VertexShader, PixelShader, GeometryShader)
 		);
 	}
 	
diff --git a/Engine/Source/Runtime/RHI/Public/BoundShaderStateCache.h b/Engine/Source/Runtime/RHI/Public/BoundShaderStateCache.h
--- a/Engine/Source/Runtime/RHI/Public/BoundShaderStateCache.h
+++ b/Engine/Source/Runtime/RHI/Public/BoundShaderStateCache.h
@@ -21,15 +21,11 @@ public:
 		FRHIVertexDeclaration* InVertexDeclaration,
 		FRHIVertexShader* InVertexShader,
 		FRHIPixelShader* InPixelShader,
-		FRHIHullShader* InHullShader = nullptr,
-		FRHIDomainShader* InDomainShader = nullptr,
 		FRHIGeometryShader* InGeometryShader = nullptr
 		)
 		: VertexDeclaration(InVertexDeclaration)
 		, VertexShader(InVertexShader)
 		, PixelShader(InPixelShader)
-		, HullShader(InHullShader)
-		, DomainShader(InDomainShader)
 		, GeometryShader(InGeometryShader)
 	{}
 
@@ -51,8 +47,8 @@ public:
 	FORCEINLINE FRHIMeshShader*     GetMeshShader() const   { return MeshShader; }
 	FORCEINLINE FRHIAmplificationShader*   GetAmplificationShader() const   { return AmplificationShader; }
 	FORCEINLINE FRHIPixelShader*    GetPixelShader() const    { return PixelShader; }
-	FORCEINLINE FRHIHullShader*     GetHullShader() const     { return HullShader; }
-	FORCEINLINE FRHIDomainShader*   GetDomainShader() const   { return DomainShader; }
+	FORCEINLINE FRHIHullShader*     GetHullShader() const     { return nullptr; }
+	FORCEINLINE FRHIDomainShader*   GetDomainShader() const   { return nullptr; }
 	FORCEINLINE FRHIGeometryShader* GetGeometryShader() const { return GeometryShader; }
 
 	/**
@@ -80,10 +76,6 @@ private:
 	FAmplificationShaderRHIRef AmplificationShader;
 	/** ps for this combination */
 	FPixelShaderRHIRef PixelShader;
-	/** hs for this combination */
-	FHullShaderRHIRef HullShader;
-	/** ds for this combination */
-	FDomainShaderRHIRef DomainShader;
 	/** gs for this combination */
 	FGeometryShaderRHIRef GeometryShader;
 
@@ -105,8 +97,6 @@ public:
 		, MeshShader(Key.MeshShader.GetReference())
 		, AmplificationShader(Key.AmplificationShader.GetReference())
 		, PixelShader(Key.PixelShader.GetReference())
-		, HullShader(Key.HullShader.GetReference())
-		, DomainShader(Key.DomainShader.GetReference())
 		, GeometryShader(Key.GeometryShader.GetReference())
 	{}
 
@@ -114,15 +104,11 @@ public:
 		FRHIVertexDeclaration* InVertexDeclaration,
 		FRHIVertexShader* InVertexShader,
 		FRHIPixelShader* InPixelShader,
-		FRHIHullShader* InHullShader = nullptr,
-		FRHIDomainShader* InDomainShader = nullptr,
 		FRHIGeometryShader* InGeometryShader = nullptr
 	)
 		: VertexDeclaration(InVertexDeclaration)
 		, VertexShader(InVertexShader)
 		, PixelShader(InPixelShader)
-		, HullShader(InHullShader)
-		, DomainShader(InDomainShader)
 		, GeometryShader(InGeometryShader)
 	{}
 
@@ -148,8 +134,6 @@ public:
 			A.MeshShader == B.MeshShader &&
 			A.AmplificationShader == B.AmplificationShader &&
 			A.PixelShader == B.PixelShader &&
-			A.HullShader == B.HullShader &&
-			A.DomainShader == B.DomainShader &&
 			A.GeometryShader == B.GeometryShader;
 	}
 
@@ -165,8 +149,6 @@ public:
 			GetTypeHash(Key.MeshShader) ^
 			GetTypeHash(Key.AmplificationShader) ^
 			GetTypeHash(Key.PixelShader) ^
-			GetTypeHash(Key.HullShader) ^
-			GetTypeHash(Key.DomainShader) ^
 			GetTypeHash(Key.GeometryShader);
 	}
 
@@ -176,8 +158,6 @@ private:
 	const FRHIMeshShader* MeshShader = nullptr;
 	const FRHIAmplificationShader* AmplificationShader = nullptr;
 	const FRHIPixelShader* PixelShader = nullptr;
-	const FRHIHullShader* HullShader = nullptr;
-	const FRHIDomainShader* DomainShader = nullptr;
 	const FRHIGeometryShader* GeometryShader = nullptr;
 };
 
@@ -210,8 +190,6 @@ public:
 		FRHIVertexDeclaration* VertexDeclaration,
 		FRHIVertexShader* VertexShader,
 		FRHIPixelShader* PixelShader,
-		FRHIHullShader* HullShader,
-		FRHIDomainShader* DomainShader,
 		FRHIGeometryShader* GeometryShader,
 		FRHIBoundShaderState* InBoundShaderState,
 		bool bAddToSingleThreadedCache = true
@@ -259,8 +237,6 @@ extern RHI_API FCachedBoundShaderStateLink* GetCachedBoundShaderState(
 	FRHIVertexDeclaration* VertexDeclaration,
 	FRHIVertexShader* VertexShader,
 	FRHIPixelShader* PixelShader,
-	FRHIHullShader* HullShader = nullptr,
-	FRHIDomainShader* DomainShader = nullptr,
 	FRHIGeometryShader* GeometryShader = nullptr,
 	FRHIMeshShader* MeshShader = nullptr,
 	FRHIAmplificationShader* AmplificationShader = nullptr
@@ -287,12 +263,10 @@ public:
 		FRHIVertexDeclaration* VertexDeclaration,
 		FRHIVertexShader* VertexShader,
 		FRHIPixelShader* PixelShader,
-		FRHIHullShader* HullShader,
-		FRHIDomainShader* DomainShader,
 		FRHIGeometryShader* GeometryShader,
 		FRHIBoundShaderState* InBoundShaderState
 		)
-		: FCachedBoundShaderStateLink(VertexDeclaration, VertexShader, PixelShader, HullShader, DomainShader, GeometryShader, InBoundShaderState, false)
+		: FCachedBoundShaderStateLink(VertexDeclaration, VertexShader, PixelShader, GeometryShader, InBoundShaderState, false)
 	{
 	}
 
@@ -319,8 +293,6 @@ extern RHI_API FBoundShaderStateRHIRef GetCachedBoundShaderState_Threadsafe(
 	FRHIVertexDeclaration* VertexDeclaration,
 	FRHIVertexShader* VertexShader,
 	FRHIPixelShader* PixelShader,
-	FRHIHullShader* HullShader = nullptr,
-	FRHIDomainShader* DomainShader = nullptr,
 	FRHIGeometryShader* GeometryShader = nullptr,
 	FRHIMeshShader* MeshShader = nullptr,
 	FRHIAmplificationShader* AmplificationShader = nullptr
diff --git a/Engine/Source/Runtime/RHI/Public/DynamicRHI.h b/Engine/Source/Runtime/RHI/Public/DynamicRHI.h
--- a/Engine/Source/Runtime/RHI/Public/DynamicRHI.h
+++ b/Engine/Source/Runtime/RHI/Public/DynamicRHI.h
@@ -398,10 +398,16 @@ public:
 	virtual FVertexShaderRHIRef RHICreateVertexShader(TArrayView<const uint8> Code, const FSHAHash& Hash) = 0;
 
 	// FlushType: Wait RHI Thread
-	virtual FHullShaderRHIRef RHICreateHullShader(TArrayView<const uint8> Code, const FSHAHash& Hash) = 0;
+	virtual FHullShaderRHIRef RHICreateHullShader(TArrayView<const uint8> Code, const FSHAHash& Hash)
+	{
+		return FHullShaderRHIRef();
+	}
 
 	// FlushType: Wait RHI Thread
-	virtual FDomainShaderRHIRef RHICreateDomainShader(TArrayView<const uint8> Code, const FSHAHash& Hash) = 0;
+	virtual FDomainShaderRHIRef RHICreateDomainShader(TArrayView<const uint8> Code, const FSHAHash& Hash)
+	{
+		return FDomainShaderRHIRef();
+	}
 
 	// FlushType: Wait RHI Thread
 	virtual FGeometryShaderRHIRef RHICreateGeometryShader(TArrayView<const uint8> Code, const FSHAHash& Hash) = 0;
@@ -535,7 +541,7 @@ public:
 	* @param PixelShader - existing pixel shader
 	*/
 	// FlushType: Thread safe, but varies depending on the RHI
-	virtual FBoundShaderStateRHIRef RHICreateBoundShaderState(FRHIVertexDeclaration* VertexDeclaration, FRHIVertexShader* VertexShader, FRHIHullShader* HullShader, FRHIDomainShader* DomainShader, FRHIPixelShader* PixelShader, FRHIGeometryShader* GeometryShader) = 0;
+	virtual FBoundShaderStateRHIRef RHICreateBoundShaderState(FRHIVertexDeclaration* VertexDeclaration, FRHIVertexShader* VertexShader, FRHIPixelShader* PixelShader, FRHIGeometryShader* GeometryShader) = 0;
 
 	/**
 	* Creates a graphics pipeline state object (PSO) that represents a complete gpu pipeline for rendering.
@@ -1462,9 +1468,9 @@ FORCEINLINE FBlendStateRHIRef RHICreateBlendState(const FBlendStateInitializerRH
 	return GDynamicRHI->RHICreateBlendState(Initializer);
 }
 
-FORCEINLINE FBoundShaderStateRHIRef RHICreateBoundShaderState(FRHIVertexDeclaration* VertexDeclaration, FRHIVertexShader* VertexShader, FRHIHullShader* HullShader, FRHIDomainShader* DomainShader, FRHIPixelShader* PixelShader, FRHIGeometryShader* GeometryShader)
+FORCEINLINE FBoundShaderStateRHIRef RHICreateBoundShaderState(FRHIVertexDeclaration* VertexDeclaration, FRHIVertexShader* VertexShader, FRHIPixelShader* PixelShader, FRHIGeometryShader* GeometryShader)
 {
-	return GDynamicRHI->RHICreateBoundShaderState(VertexDeclaration, VertexShader, HullShader, DomainShader, PixelShader, GeometryShader);
+	return GDynamicRHI->RHICreateBoundShaderState(VertexDeclaration, VertexShader, PixelShader, GeometryShader);
 }
 
 /** Before using this directly go through PipelineStateCache::GetAndOrCreateGraphicsPipelineState() */
diff --git a/Engine/Source/Runtime/RHI/Public/RHICommandList.h b/Engine/Source/Runtime/RHI/Public/RHICommandList.h
--- a/Engine/Source/Runtime/RHI/Public/RHICommandList.h
+++ b/Engine/Source/Runtime/RHI/Public/RHICommandList.h
@@ -4052,10 +4052,10 @@ public:
 		return GDynamicRHI->RHICreateStagingBuffer();
 	}
 
-	FORCEINLINE FBoundShaderStateRHIRef CreateBoundShaderState(FRHIVertexDeclaration* VertexDeclaration, FRHIVertexShader* VertexShader, FRHIHullShader* HullShader, FRHIDomainShader* DomainShader, FRHIPixelShader* PixelShader, FRHIGeometryShader* GeometryShader)
+	FORCEINLINE FBoundShaderStateRHIRef CreateBoundShaderState(FRHIVertexDeclaration* VertexDeclaration, FRHIVertexShader* VertexShader, FRHIPixelShader* PixelShader, FRHIGeometryShader* GeometryShader)
 	{
 		LLM_SCOPE(ELLMTag::Shaders);
-		return RHICreateBoundShaderState(VertexDeclaration, VertexShader, HullShader, DomainShader, PixelShader, GeometryShader);
+		return RHICreateBoundShaderState(VertexDeclaration, VertexShader, PixelShader, GeometryShader);
 	}
 
 	FORCEINLINE FGraphicsPipelineStateRHIRef CreateGraphicsPipelineState(const FGraphicsPipelineStateInitializer& Initializer)
diff --git a/Engine/Source/Runtime/RHI/Public/RHIContext.h b/Engine/Source/Runtime/RHI/Public/RHIContext.h
--- a/Engine/Source/Runtime/RHI/Public/RHIContext.h
+++ b/Engine/Source/Runtime/RHI/Public/RHIContext.h
@@ -776,8 +776,6 @@ public:
 FORCEINLINE FBoundShaderStateRHIRef RHICreateBoundShaderState(
 	FRHIVertexDeclaration* VertexDeclaration,
 	FRHIVertexShader* VertexShader,
-	FRHIHullShader* HullShader,
-	FRHIDomainShader* DomainShader,
 	FRHIPixelShader* PixelShader,
 	FRHIGeometryShader* GeometryShader
 );
@@ -814,8 +812,6 @@ public:
 			RHICreateBoundShaderState(
 				PsoInit.BoundShaderState.VertexDeclarationRHI,
 				PsoInit.BoundShaderState.VertexShaderRHI,
-				PsoInit.BoundShaderState.HullShaderRHI,
-				PsoInit.BoundShaderState.DomainShaderRHI,
 				PsoInit.BoundShaderState.PixelShaderRHI,
 				PsoInit.BoundShaderState.GeometryShaderRHI
 			).GetReference()
diff --git a/Engine/Source/Runtime/RHI/Public/RHIValidation.h b/Engine/Source/Runtime/RHI/Public/RHIValidation.h
--- a/Engine/Source/Runtime/RHI/Public/RHIValidation.h
+++ b/Engine/Source/Runtime/RHI/Public/RHIValidation.h
@@ -292,9 +292,9 @@ public:
 	* @param PixelShader - existing pixel shader
 	*/
 	// FlushType: Thread safe, but varies depending on the RHI
-	virtual FBoundShaderStateRHIRef RHICreateBoundShaderState(FRHIVertexDeclaration* VertexDeclaration, FRHIVertexShader* VertexShader, FRHIHullShader* HullShader, FRHIDomainShader* DomainShader, FRHIPixelShader* PixelShader, FRHIGeometryShader* GeometryShader) override final
+	virtual FBoundShaderStateRHIRef RHICreateBoundShaderState(FRHIVertexDeclaration* VertexDeclaration, FRHIVertexShader* VertexShader, FRHIPixelShader* PixelShader, FRHIGeometryShader* GeometryShader) override final
 	{
-		return RHI->RHICreateBoundShaderState(VertexDeclaration, VertexShader, HullShader, DomainShader, PixelShader, GeometryShader);
+		return RHI->RHICreateBoundShaderState(VertexDeclaration, VertexShader, PixelShader, GeometryShader);
 	}
 
 	/**
diff --git a/Engine/Source/Runtime/VulkanRHI/Private/VulkanResources.h b/Engine/Source/Runtime/VulkanRHI/Private/VulkanResources.h
--- a/Engine/Source/Runtime/VulkanRHI/Private/VulkanResources.h
+++ b/Engine/Source/Runtime/VulkanRHI/Private/VulkanResources.h
@@ -226,8 +226,6 @@ public:
 		FRHIVertexDeclaration* InVertexDeclarationRHI,
 		FRHIVertexShader* InVertexShaderRHI,
 		FRHIPixelShader* InPixelShaderRHI,
-		FRHIHullShader* InHullShaderRHI,
-		FRHIDomainShader* InDomainShaderRHI,
 		FRHIGeometryShader* InGeometryShaderRHI
 	);
 
diff --git a/Engine/Source/Runtime/VulkanRHI/Private/VulkanShaders.cpp b/Engine/Source/Runtime/VulkanRHI/Private/VulkanShaders.cpp
--- a/Engine/Source/Runtime/VulkanRHI/Private/VulkanShaders.cpp
+++ b/Engine/Source/Runtime/VulkanRHI/Private/VulkanShaders.cpp
@@ -661,9 +661,8 @@ void FVulkanGfxPipelineDescriptorInfo::Initialize(const FDescriptorSetRemappingI
 
 
 FVulkanBoundShaderState::FVulkanBoundShaderState(FRHIVertexDeclaration* InVertexDeclarationRHI, FRHIVertexShader* InVertexShaderRHI,
-	FRHIPixelShader* InPixelShaderRHI, FRHIHullShader* InHullShaderRHI,
-	FRHIDomainShader* InDomainShaderRHI, FRHIGeometryShader* InGeometryShaderRHI)
-	: CacheLink(InVertexDeclarationRHI, InVertexShaderRHI, InPixelShaderRHI, InHullShaderRHI, InDomainShaderRHI, InGeometryShaderRHI, this)
+	FRHIPixelShader* InPixelShaderRHI, FRHIGeometryShader* InGeometryShaderRHI)
+	: CacheLink(InVertexDeclarationRHI, InVertexShaderRHI, InPixelShaderRHI, InGeometryShaderRHI, this)
 {
 	CacheLink.AddToCache();
 }
@@ -676,8 +675,6 @@ FVulkanBoundShaderState::~FVulkanBoundShaderState()
 FBoundShaderStateRHIRef FVulkanDynamicRHI::RHICreateBoundShaderState(
 	FRHIVertexDeclaration* VertexDeclarationRHI,
 	FRHIVertexShader* VertexShaderRHI,
-	FRHIHullShader* HullShaderRHI,
-	FRHIDomainShader* DomainShaderRHI,
 	FRHIPixelShader* PixelShaderRHI,
 	FRHIGeometryShader* GeometryShaderRHI
 	)
@@ -687,8 +684,6 @@ FBoundShaderStateRHIRef FVulkanDynamicRHI::RHICreateBoundShaderState(
 		VertexDeclarationRHI,
 		VertexShaderRHI,
 		PixelShaderRHI,
-		HullShaderRHI,
-		DomainShaderRHI,
 		GeometryShaderRHI
 	);
 	if (CachedBoundShaderState.GetReference())
@@ -697,7 +692,7 @@ FBoundShaderStateRHIRef FVulkanDynamicRHI::RHICreateBoundShaderState(
 		return CachedBoundShaderState;
 	}
 
-	return new FVulkanBoundShaderState(VertexDeclarationRHI, VertexShaderRHI, PixelShaderRHI, HullShaderRHI, DomainShaderRHI, GeometryShaderRHI);
+	return new FVulkanBoundShaderState(VertexDeclarationRHI, VertexShaderRHI, PixelShaderRHI, GeometryShaderRHI);
 }
 
 
diff --git a/Engine/Source/Runtime/VulkanRHI/Public/VulkanDynamicRHI.h b/Engine/Source/Runtime/VulkanRHI/Public/VulkanDynamicRHI.h
--- a/Engine/Source/Runtime/VulkanRHI/Public/VulkanDynamicRHI.h
+++ b/Engine/Source/Runtime/VulkanRHI/Public/VulkanDynamicRHI.h
@@ -76,7 +76,7 @@ public:
 	virtual void RHICreateTransition(FRHITransition* Transition, const FRHITransitionCreateInfo& CreateInfo) final override;
 	virtual void RHIReleaseTransition(FRHITransition* Transition) final override;
 	virtual FStagingBufferRHIRef RHICreateStagingBuffer() final override;
-	virtual FBoundShaderStateRHIRef RHICreateBoundShaderState(FRHIVertexDeclaration* VertexDeclaration, FRHIVertexShader* VertexShader, FRHIHullShader* HullShader, FRHIDomainShader* DomainShader, FRHIPixelShader* PixelShader, FRHIGeometryShader* GeometryShader) final override;
+	virtual FBoundShaderStateRHIRef RHICreateBoundShaderState(FRHIVertexDeclaration* VertexDeclaration, FRHIVertexShader* VertexShader, FRHIPixelShader* PixelShader, FRHIGeometryShader* GeometryShader) final override;
 	virtual FGraphicsPipelineStateRHIRef RHICreateGraphicsPipelineState(const FGraphicsPipelineStateInitializer& Initializer) final override;
 	virtual FUniformBufferRHIRef RHICreateUniformBuffer(const void* Contents, const FRHIUniformBufferLayout& Layout, EUniformBufferUsage Usage, EUniformBufferValidation Validation) final override;
 	virtual void RHIUpdateUniformBuffer(FRHIUniformBuffer* UniformBufferRHI, const void* Contents) final override;
diff --git a/Engine/Source/Runtime/Windows/D3D11RHI/Private/D3D11RHIPrivate.h b/Engine/Source/Runtime/Windows/D3D11RHI/Private/D3D11RHIPrivate.h
--- a/Engine/Source/Runtime/Windows/D3D11RHI/Private/D3D11RHIPrivate.h
+++ b/Engine/Source/Runtime/Windows/D3D11RHI/Private/D3D11RHIPrivate.h
@@ -434,7 +434,7 @@ public:
 	virtual void RHIWriteGPUFence(FRHIGPUFence* Fence) final override;
 	virtual void* RHILockStagingBuffer(FRHIStagingBuffer* StagingBuffer, FRHIGPUFence* Fence, uint32 Offset, uint32 SizeRHI) final override;
     virtual void RHIUnlockStagingBuffer(FRHIStagingBuffer* StagingBuffer) final override;
-	virtual FBoundShaderStateRHIRef RHICreateBoundShaderState(FRHIVertexDeclaration* VertexDeclaration, FRHIVertexShader* VertexShader, FRHIHullShader* HullShader, FRHIDomainShader* DomainShader, FRHIPixelShader* PixelShader, FRHIGeometryShader* GeometryShader) final override;
+	virtual FBoundShaderStateRHIRef RHICreateBoundShaderState(FRHIVertexDeclaration* VertexDeclaration, FRHIVertexShader* VertexShader, FRHIPixelShader* PixelShader, FRHIGeometryShader* GeometryShader) final override;
 	virtual FUniformBufferRHIRef RHICreateUniformBuffer(const void* Contents, const FRHIUniformBufferLayout& Layout, EUniformBufferUsage Usage, EUniformBufferValidation Validation) final override;
 	virtual void RHIUpdateUniformBuffer(FRHIUniformBuffer* UniformBufferRHI, const void* Contents) final override;
 	virtual FBufferRHIRef RHICreateBuffer(uint32 Size, EBufferUsageFlags Usage, uint32 Stride, ERHIAccess ResourceState, FRHIResourceCreateInfo& CreateInfo) final override;
diff --git a/Engine/Source/Runtime/Windows/D3D11RHI/Private/D3D11Shaders.cpp b/Engine/Source/Runtime/Windows/D3D11RHI/Private/D3D11Shaders.cpp
--- a/Engine/Source/Runtime/Windows/D3D11RHI/Private/D3D11Shaders.cpp
+++ b/Engine/Source/Runtime/Windows/D3D11RHI/Private/D3D11Shaders.cpp
@@ -341,20 +341,16 @@ FD3D11BoundShaderState::FD3D11BoundShaderState(
 	FRHIVertexDeclaration* InVertexDeclarationRHI,
 	FRHIVertexShader* InVertexShaderRHI,
 	FRHIPixelShader* InPixelShaderRHI,
-	FRHIHullShader* InHullShaderRHI,
-	FRHIDomainShader* InDomainShaderRHI,
 	FRHIGeometryShader* InGeometryShaderRHI,
 	ID3D11Device* Direct3DDevice
 	):
-	CacheLink(InVertexDeclarationRHI,InVertexShaderRHI,InPixelShaderRHI,InHullShaderRHI,InDomainShaderRHI,InGeometryShaderRHI,this)
+	CacheLink(InVertexDeclarationRHI,InVertexShaderRHI,InPixelShaderRHI,InGeometryShaderRHI,this)
 {
 	INC_DWORD_STAT(STAT_D3D11NumBoundShaderState);
 
 	FD3D11VertexDeclaration* InVertexDeclaration = FD3D11DynamicRHI::ResourceCast(InVertexDeclarationRHI);
 	FD3D11VertexShader* InVertexShader = FD3D11DynamicRHI::ResourceCast(InVertexShaderRHI);
 	FD3D11PixelShader* InPixelShader = FD3D11DynamicRHI::ResourceCast(InPixelShaderRHI);
-	FD3D11HullShader* InHullShader = FD3D11DynamicRHI::ResourceCast(InHullShaderRHI);
-	FD3D11DomainShader* InDomainShader = FD3D11DynamicRHI::ResourceCast(InDomainShaderRHI);
 	FD3D11GeometryShader* InGeometryShader = FD3D11DynamicRHI::ResourceCast(InGeometryShaderRHI);
 
 	// Create an input layout for this combination of vertex declaration and vertex shader.
@@ -385,15 +381,11 @@ FD3D11BoundShaderState::FD3D11BoundShaderState(
 
 	VertexShader = InVertexShader->Resource;
 	PixelShader = InPixelShader ? InPixelShader->Resource : nullptr;
-	HullShader = InHullShader ? InHullShader->Resource : nullptr;
-	DomainShader = InDomainShader ? InDomainShader->Resource : nullptr;
 	GeometryShader = InGeometryShader ? InGeometryShader->Resource : nullptr;
 
 	FMemory::Memzero(&bShaderNeedsGlobalConstantBuffer,sizeof(bShaderNeedsGlobalConstantBuffer));
 
 	bShaderNeedsGlobalConstantBuffer[SF_Vertex] = InVertexShader->bShaderNeedsGlobalConstantBuffer;
-	bShaderNeedsGlobalConstantBuffer[SF_Hull] = InHullShader ? InHullShader->bShaderNeedsGlobalConstantBuffer : false;
-	bShaderNeedsGlobalConstantBuffer[SF_Domain] = InDomainShader ? InDomainShader->bShaderNeedsGlobalConstantBuffer : false;
 	bShaderNeedsGlobalConstantBuffer[SF_Pixel] = InPixelShader ? InPixelShader->bShaderNeedsGlobalConstantBuffer : false;
 	bShaderNeedsGlobalConstantBuffer[SF_Geometry] = InGeometryShader ? InGeometryShader->bShaderNeedsGlobalConstantBuffer : false;
 
@@ -418,8 +410,6 @@ FD3D11BoundShaderState::~FD3D11BoundShaderState()
 FBoundShaderStateRHIRef FD3D11DynamicRHI::RHICreateBoundShaderState(
 	FRHIVertexDeclaration* VertexDeclarationRHI,
 	FRHIVertexShader* VertexShaderRHI,
-	FRHIHullShader* HullShaderRHI,
-	FRHIDomainShader* DomainShaderRHI,
 	FRHIPixelShader* PixelShaderRHI,
 	FRHIGeometryShader* GeometryShaderRHI
 	)
@@ -435,8 +425,6 @@ FBoundShaderStateRHIRef FD3D11DynamicRHI::RHICreateBoundShaderState(
 		VertexDeclarationRHI,
 		VertexShaderRHI,
 		PixelShaderRHI,
-		HullShaderRHI,
-		DomainShaderRHI,
 		GeometryShaderRHI
 		);
 	if(CachedBoundShaderStateLink)
@@ -447,6 +435,6 @@ FBoundShaderStateRHIRef FD3D11DynamicRHI::RHICreateBoundShaderState(
 	else
 	{
 		SCOPE_CYCLE_COUNTER(STAT_D3D11NewBoundShaderStateTime);
-		return new FD3D11BoundShaderState(VertexDeclarationRHI,VertexShaderRHI,PixelShaderRHI,HullShaderRHI,DomainShaderRHI,GeometryShaderRHI,Direct3DDevice);
+		return new FD3D11BoundShaderState(VertexDeclarationRHI,VertexShaderRHI,PixelShaderRHI,GeometryShaderRHI,Direct3DDevice);
 	}
 }
diff --git a/Engine/Source/Runtime/Windows/D3D11RHI/Public/D3D11Resources.h b/Engine/Source/Runtime/Windows/D3D11RHI/Public/D3D11Resources.h
--- a/Engine/Source/Runtime/Windows/D3D11RHI/Public/D3D11Resources.h
+++ b/Engine/Source/Runtime/Windows/D3D11RHI/Public/D3D11Resources.h
@@ -135,8 +135,6 @@ public:
 		FRHIVertexDeclaration* InVertexDeclarationRHI,
 		FRHIVertexShader* InVertexShaderRHI,
 		FRHIPixelShader* InPixelShaderRHI,
-		FRHIHullShader* InHullShaderRHI,
-		FRHIDomainShader* InDomainShaderRHI,
 		FRHIGeometryShader* InGeometryShaderRHI,
 		ID3D11Device* Direct3DDevice
 		);
