Tessellation Removal: Removing Tessellation from OpenGLDrv.

diff --git a/Engine/Source/Runtime/OpenGLDrv/Private/OpenGL3.cpp b/Engine/Source/Runtime/OpenGLDrv/Private/OpenGL3.cpp
--- a/Engine/Source/Runtime/OpenGLDrv/Private/OpenGL3.cpp
+++ b/Engine/Source/Runtime/OpenGLDrv/Private/OpenGL3.cpp
@@ -11,7 +11,6 @@
 
 #if OPENGL_GL3
 
-bool FOpenGL3::bSupportsTessellation = false;
 bool FOpenGL3::bSupportsSeparateShaderObjects = false;
 bool FOpenGL3::bAndroidGLESCompatibilityMode = false;
 
@@ -34,21 +33,6 @@ void FOpenGL3::ProcessQueryGLInt()
 	GET_GL_INT(GL_MAX_GEOMETRY_TEXTURE_IMAGE_UNITS, 0, MaxGeometryTextureImageUnits);
 	GET_GL_INT(GL_MAX_VARYING_VECTORS, 0, MaxVaryingVectors);
 
-	if (bSupportsTessellation)
-	{
-		GET_GL_INT(GL_MAX_TESS_CONTROL_UNIFORM_COMPONENTS, 0, MaxHullUniformComponents);
-		GET_GL_INT(GL_MAX_TESS_EVALUATION_UNIFORM_COMPONENTS, 0, MaxDomainUniformComponents);
-		GET_GL_INT(GL_MAX_TESS_CONTROL_TEXTURE_IMAGE_UNITS, 0, MaxHullTextureImageUnits);
-		GET_GL_INT(GL_MAX_TESS_EVALUATION_TEXTURE_IMAGE_UNITS, 0, MaxDomainTextureImageUnits);
-	}
-	else
-	{
-		MaxHullUniformComponents = 0;
-		MaxDomainUniformComponents = 0;
-		MaxHullTextureImageUnits = 0;
-		MaxDomainTextureImageUnits = 0;
-	}
-
 #if !defined(__GNUC__) && !defined(__clang__)
 #define LOG_AND_GET_GL_QUERY_INT(IntEnum, Default, Dest) do { if (IntEnum) {glGetQueryiv(IntEnum, GL_QUERY_COUNTER_BITS, &Dest);} else {Dest = Default;} /*FPlatformMisc::LowLevelOutputDebugStringf(TEXT(" GL_QUERY_COUNTER_BITS: ") ## TEXT(#IntEnum) ## TEXT(": %d"), Dest);*/ } while(0)
 #else
@@ -73,10 +57,6 @@ void FOpenGL3::ProcessExtensions( const FString& ExtensionsString )
 	}
 	check(MajorVersion!=0);
 
-	// 3.3+ may expose this as an extension, 4.0+ should have this. 
-	// https://www.opengl.org/registry/specs/ARB/tessellation_shader.txt
-	bSupportsTessellation = ExtensionsString.Contains(TEXT("GL_ARB_tessellation_shader")) || ((MajorVersion >= 4));
-
 	ProcessQueryGLInt();
 	FOpenGLBase::ProcessExtensions(ExtensionsString);
 
diff --git a/Engine/Source/Runtime/OpenGLDrv/Private/OpenGLCommands.cpp b/Engine/Source/Runtime/OpenGLDrv/Private/OpenGLCommands.cpp
--- a/Engine/Source/Runtime/OpenGLDrv/Private/OpenGLCommands.cpp
+++ b/Engine/Source/Runtime/OpenGLDrv/Private/OpenGLCommands.cpp
@@ -295,14 +295,6 @@ static FORCEINLINE EShaderFrequency GetShaderFrequency(FRHIGraphicsShader* Shade
 	case SF_Vertex:
 		VALIDATE_BOUND_SHADER(ShaderRHI, Vertex);
 		return SF_Vertex;
-#if PLATFORM_SUPPORTS_TESSELLATION_SHADERS
-	case SF_Hull:
-		VALIDATE_BOUND_SHADER(ShaderRHI, Hull);
-		return SF_Hull;
-	case SF_Domain:
-		VALIDATE_BOUND_SHADER(ShaderRHI, Domain);
-		return SF_Domain;
-#endif
 #if PLATFORM_SUPPORTS_GEOMETRY_SHADERS
 	case SF_Geometry:
 		VALIDATE_BOUND_SHADER(ShaderRHI, Geometry);
@@ -325,14 +317,6 @@ static FORCEINLINE CrossCompiler::EShaderStage GetShaderCrossCompilerStage(FRHIG
 	case SF_Vertex:
 		VALIDATE_BOUND_SHADER(ShaderRHI, Vertex);
 		return CrossCompiler::SHADER_STAGE_VERTEX;
-#if PLATFORM_SUPPORTS_TESSELLATION_SHADERS
-	case SF_Hull:
-		VALIDATE_BOUND_SHADER(ShaderRHI, Hull);
-		return CrossCompiler::SHADER_STAGE_HULL;
-	case SF_Domain:
-		VALIDATE_BOUND_SHADER(ShaderRHI, Domain);
-		return CrossCompiler::SHADER_STAGE_DOMAIN;
-#endif
 #if PLATFORM_SUPPORTS_GEOMETRY_SHADERS
 	case SF_Geometry:
 		VALIDATE_BOUND_SHADER(ShaderRHI, Geometry);
@@ -357,18 +341,6 @@ static FORCEINLINE void GetShaderStageIndexAndMaxUnits(FRHIGraphicsShader* Shade
 		OutIndex = FOpenGL::GetFirstVertexTextureUnit();
 		OutMaxUnits = FOpenGL::GetMaxVertexTextureImageUnits();
 		break;
-#if PLATFORM_SUPPORTS_TESSELLATION_SHADERS
-	case SF_Hull:
-		VALIDATE_BOUND_SHADER(ShaderRHI, Hull);
-		OutIndex = FOpenGL::GetFirstHullTextureUnit();
-		OutMaxUnits = FOpenGL::GetMaxHullTextureImageUnits();
-		break;
-	case SF_Domain:
-		VALIDATE_BOUND_SHADER(ShaderRHI, Domain);
-		OutIndex = FOpenGL::GetFirstDomainTextureUnit();
-		OutMaxUnits = FOpenGL::GetMaxDomainTextureImageUnits();
-		break;
-#endif
 #if PLATFORM_SUPPORTS_GEOMETRY_SHADERS
 	case SF_Geometry:
 		VALIDATE_BOUND_SHADER(ShaderRHI, Geometry);
@@ -2497,8 +2469,6 @@ template <EShaderFrequency Frequency>
 uint32 GetFirstTextureUnit();
 
 template <> FORCEINLINE uint32 GetFirstTextureUnit<SF_Vertex>() { return FOpenGL::GetFirstVertexTextureUnit(); }
-template <> FORCEINLINE uint32 GetFirstTextureUnit<SF_Hull>() { return FOpenGL::GetFirstHullTextureUnit(); }
-template <> FORCEINLINE uint32 GetFirstTextureUnit<SF_Domain>() { return FOpenGL::GetFirstDomainTextureUnit(); }
 template <> FORCEINLINE uint32 GetFirstTextureUnit<SF_Pixel>() { return FOpenGL::GetFirstPixelTextureUnit(); }
 template <> FORCEINLINE uint32 GetFirstTextureUnit<SF_Geometry>() { return FOpenGL::GetFirstGeometryTextureUnit(); }
 template <> FORCEINLINE uint32 GetFirstTextureUnit<SF_Compute>() { return FOpenGL::GetFirstComputeTextureUnit(); }
@@ -2507,8 +2477,6 @@ template <EShaderFrequency Frequency>
 uint32 GetNumTextureUnits();
 
 template <> FORCEINLINE uint32 GetNumTextureUnits<SF_Vertex>() { return FOpenGL::GetMaxVertexTextureImageUnits(); }
-template <> FORCEINLINE uint32 GetNumTextureUnits<SF_Hull>() { return FOpenGL::GetMaxHullTextureImageUnits(); }
-template <> FORCEINLINE uint32 GetNumTextureUnits<SF_Domain>() { return FOpenGL::GetMaxDomainTextureImageUnits(); }
 template <> FORCEINLINE uint32 GetNumTextureUnits<SF_Pixel>() { return FOpenGL::GetMaxTextureImageUnits(); }
 template <> FORCEINLINE uint32 GetNumTextureUnits<SF_Geometry>() { return FOpenGL::GetMaxGeometryTextureImageUnits(); }
 template <> FORCEINLINE uint32 GetNumTextureUnits<SF_Compute>() { return FOpenGL::GetMaxComputeTextureImageUnits(); }
@@ -2664,20 +2632,6 @@ void FOpenGLDynamicRHI::CommitGraphicsResourceTablesInner()
 			SetResourcesFromTables(Shader);
 		}
 	}
-	if (PendingState.DirtyUniformBuffers[SF_Hull])
-	{
-		if (auto* Shader = PendingState.BoundShaderState->GetHullShader())
-		{
-			SetResourcesFromTables(Shader);
-		}
-	}
-	if (PendingState.DirtyUniformBuffers[SF_Domain])
-	{
-		if (auto* Shader = PendingState.BoundShaderState->GetDomainShader())
-		{
-			SetResourcesFromTables(Shader);
-		}
-	}
 	if (PendingState.DirtyUniformBuffers[SF_Geometry])
 	{
 		if (auto* Shader = PendingState.BoundShaderState->GetGeometryShader())
@@ -2767,13 +2721,8 @@ void FOpenGLDynamicRHI::RHIDrawPrimitive(uint32 BaseVertexIndex,uint32 NumPrimit
 	GLenum DrawMode = GL_TRIANGLES;
 	GLsizei NumElements = 0;
 	GLint PatchSize = 0;
-	FindPrimitiveType(PrimitiveType, ContextState.bUsingTessellation, NumPrimitives, DrawMode, NumElements, PatchSize);
+	FindPrimitiveType(PrimitiveType, NumPrimitives, DrawMode, NumElements, PatchSize);
 
-	if (FOpenGL::SupportsTessellation() && DrawMode == GL_PATCHES )
-	{
-		FOpenGL::PatchParameteri(GL_PATCH_VERTICES, PatchSize);
-	}
-	
 #if DEBUG_GL_SHADERS
 	VerifyProgramPipeline();
 #endif
@@ -2822,12 +2771,7 @@ void FOpenGLDynamicRHI::RHIDrawPrimitiveIndirect(FRHIBuffer* ArgumentBufferRHI,
 		GLenum DrawMode = GL_TRIANGLES;
 		GLsizei NumElements = 0;
 		GLint PatchSize = 0;
-		FindPrimitiveType(PrimitiveType, ContextState.bUsingTessellation, 0, DrawMode, NumElements, PatchSize);
-
-		if (FOpenGL::SupportsTessellation() && DrawMode == GL_PATCHES )
-		{
-			FOpenGL::PatchParameteri(GL_PATCH_VERTICES, PatchSize);
-		} 
+		FindPrimitiveType(PrimitiveType, 0, DrawMode, NumElements, PatchSize);
 
 		FOpenGLBuffer* ArgumentBuffer = ResourceCast(ArgumentBufferRHI);
 
@@ -2880,12 +2824,7 @@ void FOpenGLDynamicRHI::RHIDrawIndexedIndirect(FRHIBuffer* IndexBufferRHI, FRHIB
 		GLenum DrawMode = GL_TRIANGLES;
 		GLsizei NumElements = 0;
 		GLint PatchSize = 0;
-		FindPrimitiveType(PrimitiveType, ContextState.bUsingTessellation, 0, DrawMode, NumElements, PatchSize);
-
-		if (FOpenGL::SupportsTessellation() && DrawMode == GL_PATCHES )
-		{
-			FOpenGL::PatchParameteri(GL_PATCH_VERTICES, PatchSize);
-		} 
+		FindPrimitiveType(PrimitiveType, 0, DrawMode, NumElements, PatchSize);
 
 		GLenum IndexType = IndexBuffer->GetStride() == sizeof(uint32) ? GL_UNSIGNED_INT : GL_UNSIGNED_SHORT;
 
@@ -2973,12 +2912,7 @@ void FOpenGLDynamicRHI::RHIDrawIndexedPrimitive(FRHIBuffer* IndexBufferRHI, int3
 	GLenum DrawMode = GL_TRIANGLES;
 	GLsizei NumElements = 0;
 	GLint PatchSize = 0;
-	FindPrimitiveType(PrimitiveType, ContextState.bUsingTessellation, NumPrimitives, DrawMode, NumElements, PatchSize);
-
-	if (FOpenGL::SupportsTessellation() && DrawMode == GL_PATCHES )
-	{
-		FOpenGL::PatchParameteri(GL_PATCH_VERTICES, PatchSize);
-	}
+	FindPrimitiveType(PrimitiveType, NumPrimitives, DrawMode, NumElements, PatchSize);
 
 	GLenum IndexType = IndexBuffer->GetStride() == sizeof(uint32) ? GL_UNSIGNED_INT : GL_UNSIGNED_SHORT;
 	StartIndex *= IndexBuffer->GetStride() == sizeof(uint32) ? sizeof(uint32) : sizeof(uint16);
@@ -3041,12 +2975,7 @@ void FOpenGLDynamicRHI::RHIDrawIndexedPrimitiveIndirect(FRHIBuffer* IndexBufferR
 		GLenum DrawMode = GL_TRIANGLES;
 		GLsizei NumElements = 0;
 		GLint PatchSize = 0;
-		FindPrimitiveType(PrimitiveType, ContextState.bUsingTessellation, 0, DrawMode, NumElements, PatchSize);
-
-		if (FOpenGL::SupportsTessellation() && DrawMode == GL_PATCHES )
-		{
-			FOpenGL::PatchParameteri(GL_PATCH_VERTICES, PatchSize);
-		} 
+		FindPrimitiveType(PrimitiveType, 0, DrawMode, NumElements, PatchSize);
 
 		GLenum IndexType = IndexBuffer->GetStride() == sizeof(uint32) ? GL_UNSIGNED_INT : GL_UNSIGNED_SHORT;
 
diff --git a/Engine/Source/Runtime/OpenGLDrv/Private/OpenGLDrv.cpp b/Engine/Source/Runtime/OpenGLDrv/Private/OpenGLDrv.cpp
--- a/Engine/Source/Runtime/OpenGLDrv/Private/OpenGLDrv.cpp
+++ b/Engine/Source/Runtime/OpenGLDrv/Private/OpenGLDrv.cpp
@@ -401,15 +401,11 @@ GLint FOpenGLBase::MaxTextureImageUnits = -1;
 GLint FOpenGLBase::MaxCombinedTextureImageUnits = -1;
 GLint FOpenGLBase::MaxVertexTextureImageUnits = -1;
 GLint FOpenGLBase::MaxGeometryTextureImageUnits = -1;
-GLint FOpenGLBase::MaxHullTextureImageUnits = -1;
-GLint FOpenGLBase::MaxDomainTextureImageUnits = -1;
 GLint FOpenGLBase::MaxVaryingVectors = -1;
 GLint FOpenGLBase::TextureBufferAlignment = -1;
 GLint FOpenGLBase::MaxVertexUniformComponents = -1;
 GLint FOpenGLBase::MaxPixelUniformComponents = -1;
 GLint FOpenGLBase::MaxGeometryUniformComponents = -1;
-GLint FOpenGLBase::MaxHullUniformComponents = -1;
-GLint FOpenGLBase::MaxDomainUniformComponents = -1;
 bool  FOpenGLBase::bSupportsClipControl = false;
 bool  FOpenGLBase::bSupportsASTC = false;
 bool  FOpenGLBase::bSupportsCopyImage = false;
@@ -439,8 +435,6 @@ void FOpenGLBase::ProcessExtensions( const FString& ExtensionsString )
 		MaxTextureImageUnits = MaxTextureImageUnits > 16 ? 16 : MaxTextureImageUnits;
 		MaxVertexTextureImageUnits = MaxVertexTextureImageUnits > 8 ? 8 : MaxVertexTextureImageUnits;
 		MaxGeometryTextureImageUnits = MaxGeometryTextureImageUnits > 8 ? 8 : MaxGeometryTextureImageUnits;
-		MaxHullTextureImageUnits = 0;
-		MaxDomainTextureImageUnits = 0;
 		MaxCombinedTextureImageUnits = MaxCombinedTextureImageUnits > 32 ? 32 : MaxCombinedTextureImageUnits;
 	}
 	else
@@ -451,8 +445,6 @@ void FOpenGLBase::ProcessExtensions( const FString& ExtensionsString )
 			MaxTextureImageUnits = MaxTextureImageUnits > 16 ? 16 : MaxTextureImageUnits;
 			MaxVertexTextureImageUnits = MaxVertexTextureImageUnits > 8 ? 8 : MaxVertexTextureImageUnits;
 			MaxGeometryTextureImageUnits = MaxGeometryTextureImageUnits > 8 ? 8 : MaxGeometryTextureImageUnits;
-			MaxHullTextureImageUnits = MaxHullTextureImageUnits > 8 ? 8 : MaxHullTextureImageUnits;
-			MaxDomainTextureImageUnits = MaxDomainTextureImageUnits > 8 ? 8 : MaxDomainTextureImageUnits;
 			MaxCombinedTextureImageUnits = MaxCombinedTextureImageUnits > 48 ? 48 : MaxCombinedTextureImageUnits;
 		}
 	}
diff --git a/Engine/Source/Runtime/OpenGLDrv/Private/OpenGLDrvPrivate.h b/Engine/Source/Runtime/OpenGLDrv/Private/OpenGLDrvPrivate.h
--- a/Engine/Source/Runtime/OpenGLDrv/Private/OpenGLDrvPrivate.h
+++ b/Engine/Source/Runtime/OpenGLDrv/Private/OpenGLDrvPrivate.h
@@ -365,7 +365,7 @@ inline uint32 FindMaxMipmapLevel(uint32 Width, uint32 Height, uint32 Depth)
 	return FindMaxMipmapLevel((Width > Height) ? Width : Height, Depth);
 }
 
-inline void FindPrimitiveType(uint32 InPrimitiveType, bool bUsingTessellation, uint32 InNumPrimitives, GLenum &DrawMode, GLsizei &NumElements, GLint &PatchSize)
+inline void FindPrimitiveType(uint32 InPrimitiveType, uint32 InNumPrimitives, GLenum &DrawMode, GLsizei &NumElements, GLint &PatchSize)
 {
 	DrawMode = GL_TRIANGLES;
 	NumElements = InNumPrimitives;
@@ -374,30 +374,18 @@ inline void FindPrimitiveType(uint32 InPrimitiveType, bool bUsingTessellation, u
 	switch (InPrimitiveType)
 	{
 	case PT_TriangleList:
-		if (bUsingTessellation) // see GetD3D11PrimitiveType
-		{ 
-			DrawMode = GL_PATCHES;
-			PatchSize = 3;
-			NumElements = InNumPrimitives * PatchSize;
-		}
-		else
-		{
-			DrawMode = GL_TRIANGLES;
-			NumElements = InNumPrimitives * 3;
-		}
+		DrawMode = GL_TRIANGLES;
+		NumElements = InNumPrimitives * 3;
 		break;
 	case PT_TriangleStrip:
-		check(!bUsingTessellation);
 		DrawMode = GL_TRIANGLE_STRIP;
 		NumElements = InNumPrimitives + 2;
 		break;
 	case PT_LineList:
-		check(!bUsingTessellation);
 		DrawMode = GL_LINES;
 		NumElements = InNumPrimitives * 2;
 		break;
 	case PT_PointList:
-		check(!bUsingTessellation);
 		DrawMode = GL_POINTS;
 		NumElements = InNumPrimitives;
 		break;
diff --git a/Engine/Source/Runtime/OpenGLDrv/Private/OpenGLES.cpp b/Engine/Source/Runtime/OpenGLDrv/Private/OpenGLES.cpp
--- a/Engine/Source/Runtime/OpenGLDrv/Private/OpenGLES.cpp
+++ b/Engine/Source/Runtime/OpenGLDrv/Private/OpenGLES.cpp
@@ -188,8 +188,6 @@ void FOpenGLES::ProcessQueryGLInt()
 	MaxVertexUniformComponents = FMath::Max<GLint>(MaxVertexUniformComponents, RequiredMaxVertexUniformComponents);
 	MaxGeometryUniformComponents = 0;
 	MaxGeometryTextureImageUnits = 0;
-	MaxHullTextureImageUnits = 0;
-	MaxDomainTextureImageUnits = 0;
 
 	// Set lowest possible limits for texture units, to avoid extra work in GL RHI
 	MaxTextureImageUnits = FMath::Min(MaxTextureImageUnits, 16);
diff --git a/Engine/Source/Runtime/OpenGLDrv/Private/OpenGLShaders.cpp b/Engine/Source/Runtime/OpenGLDrv/Private/OpenGLShaders.cpp
--- a/Engine/Source/Runtime/OpenGLDrv/Private/OpenGLShaders.cpp
+++ b/Engine/Source/Runtime/OpenGLDrv/Private/OpenGLShaders.cpp
@@ -359,8 +359,6 @@ static const TCHAR* ShaderNameFromShaderType(GLenum ShaderType)
 		case GL_VERTEX_SHADER: return TEXT("vertex");
 		case GL_FRAGMENT_SHADER: return TEXT("fragment");
 		case GL_GEOMETRY_SHADER: return TEXT("geometry");
-		case GL_TESS_CONTROL_SHADER: return TEXT("hull");
-		case GL_TESS_EVALUATION_SHADER: return TEXT("domain");
 		case GL_COMPUTE_SHADER: return TEXT("compute");
 		default: return NULL;
 	}
@@ -651,8 +649,6 @@ ShaderType* CompileOpenGLShader(TArrayView<const uint8> InShaderCode, const FSHA
 		|| (TypeEnum == GL_FRAGMENT_SHADER && Header.FrequencyMarker != 0x5053)
 		|| (TypeEnum == GL_GEOMETRY_SHADER && Header.FrequencyMarker != 0x4753)
 		|| (TypeEnum == GL_COMPUTE_SHADER && Header.FrequencyMarker != 0x4353)
-		|| (TypeEnum == GL_TESS_CONTROL_SHADER && Header.FrequencyMarker != 0x4853 && FOpenGL::SupportsTessellation()) /* hull shader*/
-		|| (TypeEnum == GL_TESS_EVALUATION_SHADER && Header.FrequencyMarker != 0x4453 && FOpenGL::SupportsTessellation()) /* domain shader*/
 		)
 	{
 		UE_LOG(LogRHI,Fatal,
@@ -1073,18 +1069,6 @@ FGeometryShaderRHIRef FOpenGLDynamicRHI::RHICreateGeometryShader(TArrayView<cons
 	return CreateProxyShader<FRHIGeometryShader, FOpenGLGeometryShaderProxy>(Code, Hash);
 }
 
-FHullShaderRHIRef FOpenGLDynamicRHI::RHICreateHullShader(TArrayView<const uint8> Code, const FSHAHash& Hash)
-{
-	check(GMaxRHIFeatureLevel >= ERHIFeatureLevel::SM5);
-	return CreateProxyShader<FRHIHullShader, FOpenGLHullShaderProxy>(Code, Hash);
-}
-
-FDomainShaderRHIRef FOpenGLDynamicRHI::RHICreateDomainShader(TArrayView<const uint8> Code, const FSHAHash& Hash)
-{
-	check(GMaxRHIFeatureLevel >= ERHIFeatureLevel::SM5);
-	return CreateProxyShader<FRHIDomainShader, FOpenGLDomainShaderProxy>(Code, Hash);
-}
-
 static void MarkShaderParameterCachesDirty(FOpenGLShaderParameterCache* ShaderParameters, bool UpdateCompute)
 {
 	VERIFY_GL_SCOPE();
@@ -1251,7 +1235,6 @@ public:
 	FStagePackedUniformInfo	StagePackedUniformInfo[CrossCompiler::NUM_SHADER_STAGES];
 
 	GLuint		Program;
-	bool		bUsingTessellation;
 	bool		bDrawn;
 	bool		bConfigIsInitalized;
 
@@ -1278,7 +1261,7 @@ public:
 
 private:
 	FOpenGLLinkedProgram()
-	: Program(0), bUsingTessellation(false), bDrawn(false), bConfigIsInitalized(false), MaxTextureStage(-1), MaxUAVUnitUsed(-1)
+	: Program(0), bDrawn(false), bConfigIsInitalized(false), MaxTextureStage(-1), MaxUAVUnitUsed(-1)
 	{
 		TextureStageNeeds.Init( false, FOpenGL::GetMaxCombinedTextureImageUnits() );
 		UAVStageNeeds.Init( false, FOpenGL::GetMaxCombinedUAVUnits() );
@@ -1302,12 +1285,11 @@ public:
 		Config.ProgramKey = ProgramKeyIn;
 	}
 
-	FOpenGLLinkedProgram(const FOpenGLLinkedProgramConfiguration& ConfigIn, GLuint ProgramIn, bool bUsingTessellationIn)
+	FOpenGLLinkedProgram(const FOpenGLLinkedProgramConfiguration& ConfigIn, GLuint ProgramIn)
 		: FOpenGLLinkedProgram()
 	{
 		SetConfig(ConfigIn);
 		Program = ProgramIn;
-		bUsingTessellation = bUsingTessellationIn;
 	}
 
 	~FOpenGLLinkedProgram()
@@ -1457,29 +1439,6 @@ static void ConfigureStageStates(FOpenGLLinkedProgram* LinkedProgram)
 		check(LinkedProgram->StagePackedUniformInfo[CrossCompiler::SHADER_STAGE_GEOMETRY].PackedUniformInfos.Num() <= Config.Shaders[CrossCompiler::SHADER_STAGE_GEOMETRY].Bindings.PackedGlobalArrays.Num());
 	}
 
-	if (Config.Shaders[CrossCompiler::SHADER_STAGE_HULL].Resource)
-	{
-		LinkedProgram->ConfigureShaderStage(
-			CrossCompiler::SHADER_STAGE_HULL,
-			OGL_FIRST_UNIFORM_BUFFER +
-			Config.Shaders[CrossCompiler::SHADER_STAGE_VERTEX].Bindings.NumUniformBuffers +
-			Config.Shaders[CrossCompiler::SHADER_STAGE_PIXEL].Bindings.NumUniformBuffers +
-			Config.Shaders[CrossCompiler::SHADER_STAGE_GEOMETRY].Bindings.NumUniformBuffers
-		);
-	}
-
-	if (Config.Shaders[CrossCompiler::SHADER_STAGE_DOMAIN].Resource)
-	{
-		LinkedProgram->ConfigureShaderStage(
-			CrossCompiler::SHADER_STAGE_DOMAIN,
-			OGL_FIRST_UNIFORM_BUFFER +
-			Config.Shaders[CrossCompiler::SHADER_STAGE_VERTEX].Bindings.NumUniformBuffers +
-			Config.Shaders[CrossCompiler::SHADER_STAGE_PIXEL].Bindings.NumUniformBuffers +
-			Config.Shaders[CrossCompiler::SHADER_STAGE_GEOMETRY].Bindings.NumUniformBuffers +
-			Config.Shaders[CrossCompiler::SHADER_STAGE_HULL].Bindings.NumUniformBuffers
-		);
-	}
-
 	if (Config.Shaders[CrossCompiler::SHADER_STAGE_COMPUTE].Resource)
 	{
 		LinkedProgram->ConfigureShaderStage(
@@ -2220,8 +2179,8 @@ void FOpenGLLinkedProgram::ConfigureShaderStage( int Stage, uint32 FirstUniformB
 		FOpenGL::GetFirstVertexTextureUnit(),
 		FOpenGL::GetFirstPixelTextureUnit(),
 		FOpenGL::GetFirstGeometryTextureUnit(),
-		FOpenGL::GetFirstHullTextureUnit(),
-		FOpenGL::GetFirstDomainTextureUnit(),
+		0,
+		0,
 		FOpenGL::GetFirstComputeTextureUnit()
 	};
 	static const GLint FirstUAVUnit[CrossCompiler::NUM_SHADER_STAGES] =
@@ -2727,14 +2686,6 @@ static FOpenGLLinkedProgram* LinkProgram( const FOpenGLLinkedProgramConfiguratio
 		{
 			FOpenGL::UseProgramStages(Program, GL_GEOMETRY_SHADER_BIT, Config.Shaders[CrossCompiler::SHADER_STAGE_GEOMETRY].Resource);
 		}
-		if (Config.Shaders[CrossCompiler::SHADER_STAGE_HULL].Resource)
-		{
-			FOpenGL::UseProgramStages(Program, GL_TESS_CONTROL_SHADER_BIT, Config.Shaders[CrossCompiler::SHADER_STAGE_HULL].Resource);
-		}
-		if (Config.Shaders[CrossCompiler::SHADER_STAGE_DOMAIN].Resource)
-		{
-			FOpenGL::UseProgramStages(Program, GL_TESS_EVALUATION_SHADER_BIT, Config.Shaders[CrossCompiler::SHADER_STAGE_DOMAIN].Resource);
-		}
 		if (Config.Shaders[CrossCompiler::SHADER_STAGE_COMPUTE].Resource)
 		{
 			FOpenGL::UseProgramStages(Program, GL_COMPUTE_SHADER_BIT, Config.Shaders[CrossCompiler::SHADER_STAGE_COMPUTE].Resource);
@@ -2772,8 +2723,7 @@ static FOpenGLLinkedProgram* LinkProgram( const FOpenGLLinkedProgramConfiguratio
 	
 	FOpenGL::BindProgramPipeline(Program);
 
-	bool bUsingTessellation = Config.Shaders[CrossCompiler::SHADER_STAGE_HULL].Resource && Config.Shaders[CrossCompiler::SHADER_STAGE_DOMAIN].Resource;
-	FOpenGLLinkedProgram* LinkedProgram = new FOpenGLLinkedProgram(Config, Program, bUsingTessellation);
+	FOpenGLLinkedProgram* LinkedProgram = new FOpenGLLinkedProgram(Config, Program);
 
 	if (GetOpenGLProgramsCache().IsUsingLRU() && CVarLRUKeepProgramBinaryResident.GetValueOnAnyThread() && CachedProgramBinary.Num())
 	{
@@ -3031,12 +2981,6 @@ static void BindShaderStage(FOpenGLLinkedProgramConfiguration& Config, CrossComp
 						case SF_Geometry:
 							Header.FrequencyMarker = 0x4753;
 							break;
-						case SF_Hull:
-							Header.FrequencyMarker = 0x4853;
-							break;
-						case SF_Domain:
-							Header.FrequencyMarker = 0x4453;
-							break;
 						case SF_Compute:
 							Header.FrequencyMarker = 0x4353;
 							break;
@@ -3426,7 +3370,6 @@ void FOpenGLDynamicRHI::BindPendingShaderState( FOpenGLContextState& ContextStat
 	{
 		FOpenGL::BindProgramPipeline(PendingProgram);
 		ContextState.Program = PendingProgram;
-		ContextState.bUsingTessellation = PendingState.BoundShaderState->LinkedProgram->bUsingTessellation;
 		MarkShaderParameterCachesDirty(PendingState.ShaderParameters, false);
 		PendingState.LinkedProgramAndDirtyFlag = nullptr;
 
@@ -3476,27 +3419,6 @@ void FOpenGLDynamicRHI::BindPendingShaderState( FOpenGLContextState& ContextStat
 			NextUniformBufferIndex += NumUniformBuffers[SF_Geometry];
 		}
 
-		if (NumUniformBuffers[SF_Hull] >= 0)
-		{
-			PendingState.BoundShaderState->LinkedProgram->VerifyUniformBlockBindings(CrossCompiler::SHADER_STAGE_HULL, NextUniformBufferIndex);
-			BindUniformBufferBase(ContextState,
-				NumUniformBuffers[SF_Hull],
-				PendingState.BoundUniformBuffers[SF_Hull],
-				NextUniformBufferIndex,
-				ForceUniformBindingUpdate);
-			NextUniformBufferIndex += NumUniformBuffers[SF_Hull];
-		}
-
-		if (NumUniformBuffers[SF_Domain] >= 0)
-		{
-			PendingState.BoundShaderState->LinkedProgram->VerifyUniformBlockBindings(CrossCompiler::SHADER_STAGE_DOMAIN, NextUniformBufferIndex);
-			BindUniformBufferBase(ContextState,
-				NumUniformBuffers[SF_Domain],
-				PendingState.BoundUniformBuffers[SF_Domain],
-				NextUniformBufferIndex,
-				ForceUniformBindingUpdate);
-			NextUniformBufferIndex += NumUniformBuffers[SF_Domain];
-		}
 		if (FOpenGL::SupportsBindlessTexture())
 		{
 			SetupBindlessTextures(ContextState, PendingState.BoundShaderState->LinkedProgram->Samplers);
@@ -3599,16 +3521,12 @@ void FOpenGLBoundShaderState::GetNumUniformBuffers(int32 NumUniformBuffers[SF_Nu
 		NumUniformBuffers[SF_Vertex] = VertexShaderProxy->GetGLResourceObject_OnRHIThread()->Bindings.NumUniformBuffers;
 		NumUniformBuffers[SF_Pixel] = PixelShaderProxy->GetGLResourceObject_OnRHIThread()->Bindings.NumUniformBuffers;
 		NumUniformBuffers[SF_Geometry] = GeometryShaderProxy ? GeometryShaderProxy->GetGLResourceObject_OnRHIThread()->Bindings.NumUniformBuffers : -1;
-		NumUniformBuffers[SF_Hull] = HullShaderProxy ? HullShaderProxy->GetGLResourceObject_OnRHIThread()->Bindings.NumUniformBuffers : -1;
-		NumUniformBuffers[SF_Domain] = DomainShaderProxy ? DomainShaderProxy->GetGLResourceObject_OnRHIThread()->Bindings.NumUniformBuffers : -1;
 	}
 	else
 	{
 		NumUniformBuffers[SF_Vertex] = VertexShaderProxy->GetGLResourceObject()->Bindings.NumUniformBuffers;
 		NumUniformBuffers[SF_Pixel] = PixelShaderProxy->GetGLResourceObject()->Bindings.NumUniformBuffers;
 		NumUniformBuffers[SF_Geometry] = GeometryShaderProxy ? GeometryShaderProxy->GetGLResourceObject()->Bindings.NumUniformBuffers : -1;
-		NumUniformBuffers[SF_Hull] = HullShaderProxy ? HullShaderProxy->GetGLResourceObject()->Bindings.NumUniformBuffers : -1;
-		NumUniformBuffers[SF_Domain] = DomainShaderProxy ? DomainShaderProxy->GetGLResourceObject()->Bindings.NumUniformBuffers : -1;
 	}
 }
 
diff --git a/Engine/Source/Runtime/OpenGLDrv/Private/OpenGLState.cpp b/Engine/Source/Runtime/OpenGLDrv/Private/OpenGLState.cpp
--- a/Engine/Source/Runtime/OpenGLDrv/Private/OpenGLState.cpp
+++ b/Engine/Source/Runtime/OpenGLDrv/Private/OpenGLState.cpp
@@ -469,12 +469,6 @@ void FOpenGLRHIState::InitializeResources(int32 NumCombinedTextures, int32 NumCo
 	ShaderParameters[CrossCompiler::SHADER_STAGE_VERTEX].InitializeResources(FOpenGL::GetMaxVertexUniformComponents() * 4 * sizeof(float));
 	ShaderParameters[CrossCompiler::SHADER_STAGE_PIXEL].InitializeResources(FOpenGL::GetMaxPixelUniformComponents() * 4 * sizeof(float));
 	ShaderParameters[CrossCompiler::SHADER_STAGE_GEOMETRY].InitializeResources(FOpenGL::GetMaxGeometryUniformComponents() * 4 * sizeof(float));
-		
-	if ( FOpenGL::SupportsTessellation() )
-	{
-		ShaderParameters[CrossCompiler::SHADER_STAGE_HULL].InitializeResources(FOpenGL::GetMaxHullUniformComponents() * 4 * sizeof(float));
-		ShaderParameters[CrossCompiler::SHADER_STAGE_DOMAIN].InitializeResources(FOpenGL::GetMaxDomainUniformComponents() * 4 * sizeof(float));
-	}
 
 	LinkedProgramAndDirtyFlag = nullptr;
 	ShaderParameters[CrossCompiler::SHADER_STAGE_COMPUTE].InitializeResources(FOpenGL::GetMaxComputeUniformComponents() * 4 * sizeof(float));
diff --git a/Engine/Source/Runtime/OpenGLDrv/Public/OpenGL.h b/Engine/Source/Runtime/OpenGLDrv/Public/OpenGL.h
--- a/Engine/Source/Runtime/OpenGLDrv/Public/OpenGL.h
+++ b/Engine/Source/Runtime/OpenGLDrv/Public/OpenGL.h
@@ -145,19 +145,15 @@ public:
 	static FORCEINLINE GLint GetMaxTextureImageUnits()			{ check(MaxTextureImageUnits != -1); return MaxTextureImageUnits; }
 	static FORCEINLINE GLint GetMaxVertexTextureImageUnits()	{ check(MaxVertexTextureImageUnits != -1); return MaxVertexTextureImageUnits; }
 	static FORCEINLINE GLint GetMaxGeometryTextureImageUnits()	{ check(MaxGeometryTextureImageUnits != -1); return MaxGeometryTextureImageUnits; }
-	static FORCEINLINE GLint GetMaxHullTextureImageUnits()		{ check(MaxHullTextureImageUnits != -1); return MaxHullTextureImageUnits; }
-	static FORCEINLINE GLint GetMaxDomainTextureImageUnits()	{ check(MaxDomainTextureImageUnits != -1); return MaxDomainTextureImageUnits; }
 	static FORCEINLINE GLint GetMaxComputeTextureImageUnits()	{ return 0; }
 	static FORCEINLINE GLint GetMaxCombinedTextureImageUnits()	{ check(MaxCombinedTextureImageUnits != -1); return MaxCombinedTextureImageUnits; }
 	static FORCEINLINE GLint GetTextureBufferAlignment()		{ return TextureBufferAlignment; }
 
 
-	// Indices per unit are set in this order [Pixel, Vertex, Geometry, Hull, Domain]
+	// Indices per unit are set in this order [Pixel, Vertex, Geometry]
 	static FORCEINLINE GLint GetFirstPixelTextureUnit()			{ return 0; }
 	static FORCEINLINE GLint GetFirstVertexTextureUnit()		{ return GetFirstPixelTextureUnit() + GetMaxTextureImageUnits(); }
 	static FORCEINLINE GLint GetFirstGeometryTextureUnit()		{ return GetFirstVertexTextureUnit() + GetMaxVertexTextureImageUnits(); }
-	static FORCEINLINE GLint GetFirstHullTextureUnit()			{ return GetFirstGeometryTextureUnit() + GetMaxGeometryTextureImageUnits(); }
-	static FORCEINLINE GLint GetFirstDomainTextureUnit()		{ return GetFirstHullTextureUnit() + GetMaxHullTextureImageUnits(); }
 
 	static FORCEINLINE GLint GetFirstComputeTextureUnit()		{ return 0; }
 	
@@ -173,8 +169,6 @@ public:
 	static FORCEINLINE GLint GetMaxPixelUniformComponents()		{ check(MaxPixelUniformComponents != -1); return MaxPixelUniformComponents; }
 	static FORCEINLINE GLint GetMaxVertexUniformComponents()	{ check(MaxVertexUniformComponents != -1); return MaxVertexUniformComponents; }
 	static FORCEINLINE GLint GetMaxGeometryUniformComponents()	{ check(MaxGeometryUniformComponents != -1); return MaxGeometryUniformComponents; }
-	static FORCEINLINE GLint GetMaxHullUniformComponents()		{ check(MaxHullUniformComponents != -1); return MaxHullUniformComponents; }
-	static FORCEINLINE GLint GetMaxDomainUniformComponents()	{ check(MaxDomainUniformComponents != -1); return MaxDomainUniformComponents; }
 	static FORCEINLINE GLint GetMaxComputeUniformComponents()	{ return 0; }
 
 	static FORCEINLINE uint64 GetVideoMemorySize()				{ return 0; }
@@ -379,13 +373,9 @@ protected:
 	static GLint MaxCombinedTextureImageUnits;
 	static GLint MaxVertexTextureImageUnits;
 	static GLint MaxGeometryTextureImageUnits;
-	static GLint MaxHullTextureImageUnits;
-	static GLint MaxDomainTextureImageUnits;
 	static GLint MaxVertexUniformComponents;
 	static GLint MaxPixelUniformComponents;
 	static GLint MaxGeometryUniformComponents;
-	static GLint MaxHullUniformComponents;
-	static GLint MaxDomainUniformComponents;
 	static GLint MaxVaryingVectors;
 	static GLint TextureBufferAlignment;
 
diff --git a/Engine/Source/Runtime/OpenGLDrv/Public/OpenGL3.h b/Engine/Source/Runtime/OpenGLDrv/Public/OpenGL3.h
--- a/Engine/Source/Runtime/OpenGLDrv/Public/OpenGL3.h
+++ b/Engine/Source/Runtime/OpenGLDrv/Public/OpenGL3.h
@@ -41,7 +41,6 @@ struct FOpenGL3 : public FOpenGLBase
 	static FORCEINLINE bool SupportsVolumeTextureRendering()		{ return bSupportsVolumeTextureRendering; }
 	static FORCEINLINE bool SupportsGenerateMipmap()				{ return true; }
 	static FORCEINLINE bool AmdWorkaround()							{ return bAmdWorkaround; }
-	static FORCEINLINE bool SupportsTessellation()					{ return bSupportsTessellation; }
 	static FORCEINLINE bool SupportsSeparateShaderObjects()			{ return bSupportsSeparateShaderObjects; }
 	static FORCEINLINE bool SupportsBufferStorage()					{ return true; }
 
@@ -783,7 +782,6 @@ protected:
 	static GLint TimestampQueryBits;
 	
 	static bool bDebugContext;
-	static bool bSupportsTessellation;
 	static bool bSupportsSeparateShaderObjects;
 	static bool bAndroidGLESCompatibilityMode;
 };
diff --git a/Engine/Source/Runtime/OpenGLDrv/Public/OpenGLDrv.h b/Engine/Source/Runtime/OpenGLDrv/Public/OpenGLDrv.h
--- a/Engine/Source/Runtime/OpenGLDrv/Public/OpenGLDrv.h
+++ b/Engine/Source/Runtime/OpenGLDrv/Public/OpenGLDrv.h
@@ -340,8 +340,6 @@ public:
 
 	virtual FPixelShaderRHIRef RHICreatePixelShader(TArrayView<const uint8> Code, const FSHAHash& Hash) final override;
 	virtual FVertexShaderRHIRef RHICreateVertexShader(TArrayView<const uint8> Code, const FSHAHash& Hash) final override;
-	virtual FHullShaderRHIRef RHICreateHullShader(TArrayView<const uint8> Code, const FSHAHash& Hash) final override;
-	virtual FDomainShaderRHIRef RHICreateDomainShader(TArrayView<const uint8> Code, const FSHAHash& Hash) final override;
 	virtual FGeometryShaderRHIRef RHICreateGeometryShader(TArrayView<const uint8> Code, const FSHAHash& Hash) final override;
 	virtual FComputeShaderRHIRef RHICreateComputeShader(TArrayView<const uint8> Code, const FSHAHash& Hash) final override;
 
@@ -925,11 +923,6 @@ public:
 		return this->RHICreateComputeShader(Code, Hash);
 	}
 
-	virtual FHullShaderRHIRef CreateHullShader_RenderThread(class FRHICommandListImmediate& RHICmdList, TArrayView<const uint8> Code, const FSHAHash& Hash) final override
-	{
-		return this->RHICreateHullShader(Code, Hash);
-	}
-
 	virtual FSamplerStateRHIRef RHICreateSamplerState(const FSamplerStateInitializerRHI& Initializer) final override;
 
 	virtual void UpdateTexture2D_RenderThread(class FRHICommandListImmediate& RHICmdList, FRHITexture2D* Texture, uint32 MipIndex, const struct FUpdateTextureRegion2D& UpdateRegion, uint32 SourcePitch, const uint8* SourceData) final override
diff --git a/Engine/Source/Runtime/OpenGLDrv/Public/OpenGLResources.h b/Engine/Source/Runtime/OpenGLDrv/Public/OpenGLResources.h
--- a/Engine/Source/Runtime/OpenGLDrv/Public/OpenGLResources.h
+++ b/Engine/Source/Runtime/OpenGLDrv/Public/OpenGLResources.h
@@ -358,8 +358,6 @@ private:
 typedef TOpenGLResourceProxy<FRHIVertexShader, FOpenGLVertexShader> FOpenGLVertexShaderProxy;
 typedef TOpenGLResourceProxy<FRHIPixelShader, FOpenGLPixelShader> FOpenGLPixelShaderProxy;
 typedef TOpenGLResourceProxy<FRHIGeometryShader, FOpenGLGeometryShader> FOpenGLGeometryShaderProxy;
-typedef TOpenGLResourceProxy<FRHIHullShader, FOpenGLHullShader> FOpenGLHullShaderProxy;
-typedef TOpenGLResourceProxy<FRHIDomainShader, FOpenGLDomainShader> FOpenGLDomainShaderProxy;
 typedef TOpenGLResourceProxy<FRHIComputeShader, FOpenGLComputeShader> FOpenGLComputeShaderProxy;
 
 
@@ -1161,8 +1159,6 @@ public:
 	TRefCountPtr<FOpenGLVertexShaderProxy> VertexShaderProxy;
 	TRefCountPtr<FOpenGLPixelShaderProxy> PixelShaderProxy;
 	TRefCountPtr<FOpenGLGeometryShaderProxy> GeometryShaderProxy;
-	TRefCountPtr<FOpenGLHullShaderProxy> HullShaderProxy;
-	TRefCountPtr<FOpenGLDomainShaderProxy> DomainShaderProxy;
 
 	/** Initialization constructor. */
 	FOpenGLBoundShaderState(
@@ -1194,8 +1190,6 @@ public:
 	}
 
 	FOpenGLGeometryShader* GetGeometryShader()	{ return GeometryShaderProxy ? GeometryShaderProxy->GetGLResourceObject() : nullptr;}
-	FOpenGLHullShader* GetHullShader()	{ return HullShaderProxy ? HullShaderProxy->GetGLResourceObject() : nullptr; }
-	FOpenGLDomainShader* GetDomainShader()	{ return DomainShaderProxy ? DomainShaderProxy->GetGLResourceObject() : nullptr;}
 
 	virtual ~FOpenGLBoundShaderState();
 };
@@ -2291,16 +2285,6 @@ struct TOpenGLResourceTraits<FRHIGeometryShader>
 	typedef FOpenGLGeometryShaderProxy TConcreteType;
 };
 template<>
-struct TOpenGLResourceTraits<FRHIHullShader>
-{
-	typedef FOpenGLHullShaderProxy TConcreteType;
-};
-template<>
-struct TOpenGLResourceTraits<FRHIDomainShader>
-{
-	typedef FOpenGLDomainShaderProxy TConcreteType;
-};
-template<>
 struct TOpenGLResourceTraits<FRHIPixelShader>
 {
 	typedef FOpenGLPixelShaderProxy TConcreteType;
diff --git a/Engine/Source/Runtime/OpenGLDrv/Public/OpenGLShaderResources.h b/Engine/Source/Runtime/OpenGLDrv/Public/OpenGLShaderResources.h
--- a/Engine/Source/Runtime/OpenGLDrv/Public/OpenGLShaderResources.h
+++ b/Engine/Source/Runtime/OpenGLDrv/Public/OpenGLShaderResources.h
@@ -298,8 +298,6 @@ public:
 typedef TOpenGLShader<FRefCountedObject, GL_VERTEX_SHADER, SF_Vertex> FOpenGLVertexShader;
 typedef TOpenGLShader<FRefCountedObject, GL_FRAGMENT_SHADER, SF_Pixel> FOpenGLPixelShader;
 typedef TOpenGLShader<FRefCountedObject, GL_GEOMETRY_SHADER, SF_Geometry> FOpenGLGeometryShader;
-typedef TOpenGLShader<FRefCountedObject, GL_TESS_CONTROL_SHADER, SF_Hull> FOpenGLHullShader;
-typedef TOpenGLShader<FRefCountedObject, GL_TESS_EVALUATION_SHADER, SF_Domain> FOpenGLDomainShader;
 
 
 class FOpenGLComputeShader : public TOpenGLShader<FRefCountedObject, GL_COMPUTE_SHADER, SF_Compute>
diff --git a/Engine/Source/Runtime/OpenGLDrv/Public/OpenGLState.h b/Engine/Source/Runtime/OpenGLDrv/Public/OpenGLState.h
--- a/Engine/Source/Runtime/OpenGLDrv/Public/OpenGLState.h
+++ b/Engine/Source/Runtime/OpenGLDrv/Public/OpenGLState.h
@@ -330,7 +330,6 @@ struct FOpenGLContextState : public FOpenGLCommonState
 	uint32							RenderTargetHeight;
 	GLuint							OcclusionQuery;
 	GLuint							Program;
-	bool							bUsingTessellation;
 	GLuint 							UniformBuffers[CrossCompiler::NUM_SHADER_STAGES*OGL_MAX_UNIFORM_BUFFER_BINDINGS];
 	GLuint 							UniformBufferOffsets[CrossCompiler::NUM_SHADER_STAGES*OGL_MAX_UNIFORM_BUFFER_BINDINGS];
 	TArray<FOpenGLSamplerState*>	CachedSamplerStates;
@@ -380,7 +379,6 @@ struct FOpenGLContextState : public FOpenGLCommonState
 	:	StencilRef(0)
 	,	Framebuffer(0)
 	,	Program(0)
-	,	bUsingTessellation(false)
 	,	ActiveTexture(GL_TEXTURE0)
 	,	bScissorEnabled(false)
 	,	DepthMinZ(0.0f)
