Fix a crash in TAA pass on mobile OpenGL platform due to UAV as R11G11B10 is not supported.

diff --git a/Engine/Config/Android/AndroidEngine.ini b/Engine/Config/Android/AndroidEngine.ini
--- a/Engine/Config/Android/AndroidEngine.ini
+++ b/Engine/Config/Android/AndroidEngine.ini
@@ -83,6 +83,8 @@ r.Forward.MaxCulledLightsPerCell = 8
 r.HairStrands.Simulation=0
 r.HairStrands.Strands=0
 r.HairStrands.Binding=0
+; R11G11B10 UAV is not supported on Android
+r.TemporalAA.R11G11B10History=0
 
 [OnlineSubsystemGooglePlay.Store]
 bSupportsInAppPurchasing=true
diff --git a/Engine/Source/Runtime/OpenGLDrv/Private/OpenGLCommands.cpp b/Engine/Source/Runtime/OpenGLDrv/Private/OpenGLCommands.cpp
--- a/Engine/Source/Runtime/OpenGLDrv/Private/OpenGLCommands.cpp
+++ b/Engine/Source/Runtime/OpenGLDrv/Private/OpenGLCommands.cpp
@@ -882,6 +882,7 @@ void FOpenGLDynamicRHI::SetupUAVsForProgram(FOpenGLContextState& ContextState, c
 
 static bool IsImageTextureFormatSupported(GLenum Format)
 {
+#if PLATFORM_ANDROID
 	if (GMaxRHIFeatureLevel == ERHIFeatureLevel::ES3_1)
 	{
 		// from GLES 3.1 spec
@@ -905,6 +906,7 @@ static bool IsImageTextureFormatSupported(GLenum Format)
 				return false;
 		}
 	}
+#endif
 	return true;
 }
 
