Changed WaitFence to only wait if we are on the rendering thread

diff --git a/Engine/Source/Runtime/OpenGLDrv/Private/OpenGLQuery.cpp b/Engine/Source/Runtime/OpenGLDrv/Private/OpenGLQuery.cpp
--- a/Engine/Source/Runtime/OpenGLDrv/Private/OpenGLQuery.cpp
+++ b/Engine/Source/Runtime/OpenGLDrv/Private/OpenGLQuery.cpp
@@ -663,7 +663,7 @@ FOpenGLRenderQuery::~FOpenGLRenderQuery()
 		}
 		else
 		{
-			CreationFence.WaitFence();
+			CreationFence.WaitFenceRenderThreadOnly();
 			ALLOC_COMMAND_CL(RHICmdList, FRHICommandGLCommand)([Resource = Resource, ResourceContext = ResourceContext]() {VERIFY_GL_SCOPE(); ReleaseResource(Resource, ResourceContext); });
 		}
 	}
diff --git a/Engine/Source/Runtime/OpenGLDrv/Private/OpenGLState.cpp b/Engine/Source/Runtime/OpenGLDrv/Private/OpenGLState.cpp
--- a/Engine/Source/Runtime/OpenGLDrv/Private/OpenGLState.cpp
+++ b/Engine/Source/Runtime/OpenGLDrv/Private/OpenGLState.cpp
@@ -217,7 +217,7 @@ static EBlendFactor TranslateBlendFactor(GLenum BlendFactor)
 
 FOpenGLSamplerState::~FOpenGLSamplerState()
 {
-	CreationFence.WaitFence();
+	CreationFence.WaitFenceRenderThreadOnly();
 	VERIFY_GL_SCOPE();
 	FOpenGL::DeleteSamplers(1,&Resource);
 }
diff --git a/Engine/Source/Runtime/OpenGLDrv/Private/OpenGLUniformBuffer.cpp b/Engine/Source/Runtime/OpenGLDrv/Private/OpenGLUniformBuffer.cpp
--- a/Engine/Source/Runtime/OpenGLDrv/Private/OpenGLUniformBuffer.cpp
+++ b/Engine/Source/Runtime/OpenGLDrv/Private/OpenGLUniformBuffer.cpp
@@ -400,8 +400,8 @@ void FOpenGLUniformBuffer::SetGLUniformBufferParams(GLuint InResource, uint32 In
 
 FOpenGLUniformBuffer::~FOpenGLUniformBuffer()
 {
-	AccessFence.WaitFence();
-	CopyFence.WaitFence();
+	AccessFence.WaitFenceRenderThreadOnly();
+	CopyFence.WaitFenceRenderThreadOnly();
 
 	if (Resource != 0)
 	{
diff --git a/Engine/Source/Runtime/OpenGLDrv/Public/OpenGLResources.h b/Engine/Source/Runtime/OpenGLDrv/Public/OpenGLResources.h
--- a/Engine/Source/Runtime/OpenGLDrv/Public/OpenGLResources.h
+++ b/Engine/Source/Runtime/OpenGLDrv/Public/OpenGLResources.h
@@ -136,29 +136,7 @@ public:
 			GLAF_CHECK((IsInRenderingThread() && !IsRunningRHIInSeparateThread()) || (IsInRHIThread() && IsRunningRHIInSeparateThread()));
 		}
 	}
-	FORCEINLINE_DEBUGGABLE void WaitFence()
-	{
-		if (IsRunningRHIInSeparateThread())
-		{
-			GLAF_CHECK(IsInRenderingThread());
-			if (!IsRunningRHIInSeparateThread() && !FRHICommandListExecutor::GetImmediateCommandList().Bypass() && !GRHINeedsExtraDeletionLatency) // if we don't have an RHI thread, but we are doing parallel rendering, then we need to flush now because we are not deferring resource destruction
-			{
-				QUICK_SCOPE_CYCLE_COUNTER(STAT_FOpenGLRHIThreadResourceFence_Flush);
-				FRHICommandListExecutor::GetImmediateCommandList().ImmediateFlush(EImmediateFlushType::FlushRHIThread);
-			}
-			if (RealRHIFence.GetReference() && RealRHIFence->IsComplete())
-			{
-				RealRHIFence = nullptr;
-			}
-			else if (RealRHIFence.GetReference())
-			{
-				UE_LOG(LogRHI, Warning, TEXT("FOpenGLRHIThreadResourceFence waited.")); 
-				QUICK_SCOPE_CYCLE_COUNTER(STAT_FOpenGLRHIThreadResourceFence_Wait);
-				FRHICommandListExecutor::WaitOnRHIThreadFence(RealRHIFence);
-				RealRHIFence = nullptr;
-			}
-		}
-	}
+
 	FORCEINLINE_DEBUGGABLE void WaitFenceRenderThreadOnly()
 	{
 		if (IsRunningRHIInSeparateThread())
@@ -167,7 +145,17 @@ public:
 			// all rhi thread operations will be in order, check for RHIT isnt required.
 			if (IsInRenderingThread())
 			{
-				WaitFence();
+				if (RealRHIFence.GetReference() && RealRHIFence->IsComplete())
+				{
+					RealRHIFence = nullptr;
+				}
+				else if (RealRHIFence.GetReference())
+				{
+					UE_LOG(LogRHI, Warning, TEXT("FOpenGLRHIThreadResourceFence waited."));
+					QUICK_SCOPE_CYCLE_COUNTER(STAT_FOpenGLRHIThreadResourceFence_Wait);
+					FRHICommandListExecutor::WaitOnRHIThreadFence(RealRHIFence);
+					RealRHIFence = nullptr;
+				}
 			}
 		}
 	}
@@ -232,28 +220,6 @@ public:
 #endif
 		}
 	}
-	FORCEINLINE_DEBUGGABLE void WaitFence()
-	{
-		if (IsRunningRHIInSeparateThread())
-		{
-			check(IsInRenderingThread() || IsInRHIThread());
-			if (!IsRunningRHIInSeparateThread() && !FRHICommandListExecutor::GetImmediateCommandList().Bypass() && !GRHINeedsExtraDeletionLatency) // if we don't have an RHI thread, but we are doing parallel rendering, then we need to flush now because we are not deferring resource destruction
-			{
-				FRHICommandListExecutor::GetImmediateCommandList().ImmediateFlush(EImmediateFlushType::FlushRHIThread);
-			}
-#if USE_CHEAP_ASSERTONLY_RHI_FENCES
-			GLAF_CHECK(AssertFence.GetValue() == 0 || AssertFence.GetValue() == 2);
-#endif
-#if USE_REAL_RHI_FENCES
-			GLAF_CHECK(!RealRHIFence.GetReference() || RealRHIFence->IsComplete());
-			if (RealRHIFence.GetReference())
-			{
-				FRHICommandListExecutor::WaitOnRHIThreadFence(RealRHIFence);
-				RealRHIFence = nullptr;
-			}
-#endif
-		}
-	}
 
 	FORCEINLINE_DEBUGGABLE void WaitFenceRenderThreadOnly()
 	{
@@ -263,7 +229,17 @@ public:
 			// all rhi thread operations will be in order, check for RHIT isnt required.
 			if (IsInRenderingThread())
 			{
-				WaitFence();
+#if USE_CHEAP_ASSERTONLY_RHI_FENCES
+				GLAF_CHECK(AssertFence.GetValue() == 0 || AssertFence.GetValue() == 2);
+#endif
+#if USE_REAL_RHI_FENCES
+				GLAF_CHECK(!RealRHIFence.GetReference() || RealRHIFence->IsComplete());
+				if (RealRHIFence.GetReference())
+				{
+					FRHICommandListExecutor::WaitOnRHIThreadFence(RealRHIFence);
+					RealRHIFence = nullptr;
+				}
+#endif
 			}
 		}
 	}
@@ -345,7 +321,7 @@ private:
 	{
 		if (bQueuedCreation)
 		{
-			CreationFence.WaitFence();
+			CreationFence.WaitFenceRenderThreadOnly();
 		}
 	}
 
@@ -1582,7 +1558,7 @@ public:
 		{
 			if (IsInActualRenderingThread())
 			{
-				this->CreationFence.WaitFence();
+				this->CreationFence.WaitFenceRenderThreadOnly();
 			}
 
 			if(!CanCreateAsEvicted())
