Fixed: clearing of structured buffers in GL

diff --git a/Engine/Source/Runtime/OpenGLDrv/Private/OpenGLUAV.cpp b/Engine/Source/Runtime/OpenGLDrv/Private/OpenGLUAV.cpp
--- a/Engine/Source/Runtime/OpenGLDrv/Private/OpenGLUAV.cpp
+++ b/Engine/Source/Runtime/OpenGLDrv/Private/OpenGLUAV.cpp
@@ -274,26 +274,6 @@ void FOpenGLDynamicRHI::RHIClearUAVFloat(FRHIUnorderedAccessView* UnorderedAcces
 {
 	FOpenGLUnorderedAccessView* Texture = ResourceCast(UnorderedAccessViewRHI);
 
-#if OPENGL_GL4 || PLATFORM_LUMINGL4
-	if (GMaxRHIFeatureLevel >= ERHIFeatureLevel::SM5)
-	{
-		glBindBuffer(GL_TEXTURE_BUFFER, Texture->BufferResource);
-		FOpenGL::ClearBufferData(GL_TEXTURE_BUFFER, Texture->Format, GL_RGBA_INTEGER, GL_FLOAT, reinterpret_cast<const uint32*>(&Values));
-		GPUProfilingData.RegisterGPUWork(1);
-		return;
-	}
-#elif defined(OPENGL_ESDEFERRED)
-	glBindBuffer(GL_TEXTURE_BUFFER, Texture->BufferResource);
-	uint32 BufferSize = Texture->GetBufferSize();
-	if (BufferSize > 0)
-	{
-		void* BufferData = FOpenGL::MapBufferRange(GL_TEXTURE_BUFFER, 0, BufferSize, FOpenGLBase::RLM_WriteOnly);
-		uint8 ClearValue = uint8(Values[0] & 0xff);
-		FPlatformMemory::Memset(BufferData, ClearValue, BufferSize);
-		FOpenGL::UnmapBufferRange(GL_TEXTURE_BUFFER, 0, BufferSize);
-		GPUProfilingData.RegisterGPUWork(1);
-	}
-#endif
 	// Use compute on ES3.1
 	TRHICommandList_RecursiveHazardous<FOpenGLDynamicRHI> RHICmdList(this);
 
@@ -316,8 +296,19 @@ void FOpenGLDynamicRHI::RHIClearUAVFloat(FRHIUnorderedAccessView* UnorderedAcces
 	{
 		check(Texture->BufferResource);
 		{
-			int32 NumComponents = GPixelFormats[Texture->UnrealFormat].NumComponents;
-			uint32 NumElements = Texture->GetBufferSize() / GPixelFormats[Texture->UnrealFormat].BlockBytes;
+			int32 NumComponents = 0;
+			uint32 NumElements = 0;
+
+			if (Texture->UnrealFormat != 0)
+			{
+				NumComponents = GPixelFormats[Texture->UnrealFormat].NumComponents;
+				NumElements = Texture->GetBufferSize() / GPixelFormats[Texture->UnrealFormat].BlockBytes;
+			}
+			else
+			{
+				NumElements = Texture->GetBufferSize() / sizeof(float);
+				NumComponents = 1;
+			}
 					
 			switch (NumComponents)
 			{
@@ -359,8 +350,19 @@ void FOpenGLDynamicRHI::RHIClearUAVUint(FRHIUnorderedAccessView* UnorderedAccess
 	{
 		check(Texture->BufferResource);
 		{
-			int32 NumComponents = GPixelFormats[Texture->UnrealFormat].NumComponents;
-			uint32 NumElements = Texture->GetBufferSize() / GPixelFormats[Texture->UnrealFormat].BlockBytes;
+			int32 NumComponents = 0;
+			uint32 NumElements = 0;
+
+			if (Texture->UnrealFormat != 0)
+			{
+				NumComponents = GPixelFormats[Texture->UnrealFormat].NumComponents;
+				NumElements = Texture->GetBufferSize() / GPixelFormats[Texture->UnrealFormat].BlockBytes;
+			}
+			else
+			{
+				NumElements = Texture->GetBufferSize() / sizeof(uint32);
+				NumComponents = 1;
+			}
 
 			switch (NumComponents)
 			{
