Fixed OGL RHI thread ownership task tag scope

diff --git a/Engine/Source/Runtime/OpenGLDrv/Private/OpenGLDevice.cpp b/Engine/Source/Runtime/OpenGLDrv/Private/OpenGLDevice.cpp
--- a/Engine/Source/Runtime/OpenGLDrv/Private/OpenGLDevice.cpp
+++ b/Engine/Source/Runtime/OpenGLDrv/Private/OpenGLDevice.cpp
@@ -1262,7 +1262,6 @@ void FOpenGLDynamicRHI::RHIReleaseThreadOwnership()
 			CustomPresent->OnReleaseThreadOwnership();
 		}
 	}
-	FTaskTagScope Scope(ETaskTag::ERhiThread);
 	VERIFY_GL(RHIReleaseThreadOwnership);
 	bIsRenderingContextAcquired = false;
 	PlatformNULLContextSetup();
