Fixed: Incorrect assert for GL buffer memory counting

diff --git a/Engine/Source/Runtime/OpenGLDrv/Private/OpenGLUtil.cpp b/Engine/Source/Runtime/OpenGLDrv/Private/OpenGLUtil.cpp
--- a/Engine/Source/Runtime/OpenGLDrv/Private/OpenGLUtil.cpp
+++ b/Engine/Source/Runtime/OpenGLDrv/Private/OpenGLUtil.cpp
@@ -145,7 +145,7 @@ void DecrementBufferMemory(GLenum Type, uint32 NumBytes)
 	{
 		DEC_MEMORY_STAT_BY(STAT_PixelBufferMemory,NumBytes);
 	}
-	else
+	else if (Type != 0) // CreateInfo.bWithoutNativeResource
 	{
 		check(Type == GL_ARRAY_BUFFER);
 		DEC_MEMORY_STAT_BY(STAT_VertexBufferMemory,NumBytes);
