Binary Cache fix for Compute shader and fix calculation of ShaderLib Processing time

diff --git a/Engine/Source/Runtime/OpenGLDrv/Private/OpenGLShaders.cpp b/Engine/Source/Runtime/OpenGLDrv/Private/OpenGLShaders.cpp
--- a/Engine/Source/Runtime/OpenGLDrv/Private/OpenGLShaders.cpp
+++ b/Engine/Source/Runtime/OpenGLDrv/Private/OpenGLShaders.cpp
@@ -2786,6 +2786,15 @@ FOpenGLLinkedProgram* FOpenGLDynamicRHI::GetLinkedComputeProgram(FRHIComputeShad
 	Config.ProgramKey.ShaderHashes[CrossCompiler::SHADER_STAGE_COMPUTE] = ComputeShaderRHI->GetHash();
 
 	FOpenGLLinkedProgram* LinkedProgram = GetOpenGLProgramsCache().Find(Config.ProgramKey, true);
+	if (!LinkedProgram)
+	{
+		// ensure that pending request for this program has been completed before
+		if (FOpenGLProgramBinaryCache::CheckSinglePendingGLProgramCreateRequest(Config.ProgramKey))
+		{
+			LinkedProgram = GetOpenGLProgramsCache().Find(Config.ProgramKey, true);
+		}
+	}
+
 	if (LinkedProgram == nullptr)
 	{
 		// Not in the cache. Create and add the program here.
@@ -4689,10 +4698,11 @@ void FOpenGLProgramBinaryCache::CheckPendingGLProgramCreateRequests_internal()
 	float TimeRemainingS = (float)GMaxShaderLibProcessingTimeMS / 1000.0f;
 	double StartTime = FPlatformTime::Seconds();
 	int32 Count = 0;
-	while(PendingGLProgramCreateRequests.Num() && TimeRemainingS > 0.0f)
+	while (PendingGLProgramCreateRequests.Num() && TimeRemainingS > 0.0f)
 	{
 		CompleteLoadedGLProgramRequest_internal(PendingGLProgramCreateRequests.Pop());
 		TimeRemainingS -= (float)(FPlatformTime::Seconds() - StartTime);
+		StartTime = FPlatformTime::Seconds();
 		Count++;
 	}
 	UE_CLOG(PendingGLProgramCreateRequests.Num()>0, LogRHI, Log, TEXT("CheckPendingGLProgramCreateRequests : iter count = %d, time taken = %d ms (remaining %d)"), Count, GMaxShaderLibProcessingTimeMS - (int32)(TimeRemainingS*1000.0f), PendingGLProgramCreateRequests.Num());
