Make sure we allocate enough UAV slots in OpenGL RHI as they are shared between Images and SSBOs

diff --git a/Engine/Source/Runtime/OpenGLDrv/Private/OpenGL4.cpp b/Engine/Source/Runtime/OpenGLDrv/Private/OpenGL4.cpp
--- a/Engine/Source/Runtime/OpenGLDrv/Private/OpenGL4.cpp
+++ b/Engine/Source/Runtime/OpenGLDrv/Private/OpenGL4.cpp
@@ -29,8 +29,13 @@ void FOpenGL4::ProcessQueryGLInt()
 	GET_GL_INT(GL_MAX_COMPUTE_UNIFORM_COMPONENTS, 0, MaxComputeUniformComponents);
 	GET_GL_INT(GL_MAX_COMPUTE_IMAGE_UNIFORMS, 0, MaxComputeUAVUnits);
 	
+	GLint MaxCombinedSSBOUnits = 0;
+	GET_GL_INT(GL_MAX_COMBINED_SHADER_STORAGE_BLOCKS, 0, MaxCombinedSSBOUnits);
+	// UAVs slots in UE are shared between Images and SSBO, so this should be max(GL_MAX_COMBINED_IMAGE_UNIFORMS, GL_MAX_COMBINED_SHADER_STORAGE_BLOCKS)
+	MaxCombinedUAVUnits = FMath::Max(MaxCombinedUAVUnits, MaxCombinedSSBOUnits);
+	
 	// clamp UAV units to a sensible limit
-	MaxCombinedUAVUnits = FMath::Min(MaxCombinedUAVUnits, 8);
+	MaxCombinedUAVUnits = FMath::Min(MaxCombinedUAVUnits, 16);
 	MaxComputeUAVUnits = FMath::Min(MaxComputeUAVUnits, 16);
 	// this is split between VS and PS, 4 to each stage
 	MaxPixelUAVUnits = FMath::Min(MaxPixelUAVUnits, 4);
diff --git a/Engine/Source/Runtime/OpenGLDrv/Private/OpenGLES.cpp b/Engine/Source/Runtime/OpenGLDrv/Private/OpenGLES.cpp
--- a/Engine/Source/Runtime/OpenGLDrv/Private/OpenGLES.cpp
+++ b/Engine/Source/Runtime/OpenGLDrv/Private/OpenGLES.cpp
@@ -173,8 +173,13 @@ void FOpenGLES::ProcessQueryGLInt()
 	LOG_AND_GET_GL_INT(GL_MAX_COMPUTE_IMAGE_UNIFORMS, 0, MaxComputeUAVUnits);
 	LOG_AND_GET_GL_INT(GL_MAX_FRAGMENT_IMAGE_UNIFORMS, 0, MaxPixelUAVUnits);
 
+	GLint MaxCombinedSSBOUnits = 0;
+	GET_GL_INT(GL_MAX_COMBINED_SHADER_STORAGE_BLOCKS, 0, MaxCombinedSSBOUnits);
+	// UAVs slots in UE are shared between Images and SSBO, so this should be max(GL_MAX_COMBINED_IMAGE_UNIFORMS, GL_MAX_COMBINED_SHADER_STORAGE_BLOCKS)
+	MaxCombinedUAVUnits = FMath::Max(MaxCombinedUAVUnits, MaxCombinedSSBOUnits);
+
 	// clamp UAV units to a sensible limit
-	MaxCombinedUAVUnits = FMath::Min(MaxCombinedUAVUnits, 8);
+	MaxCombinedUAVUnits = FMath::Min(MaxCombinedUAVUnits, 16);
 	MaxComputeUAVUnits = FMath::Min(MaxComputeUAVUnits, 16);
 	// this is split between VS and PS, 4 to each stage
 	MaxPixelUAVUnits = FMath::Min(MaxPixelUAVUnits, 4);
