Add GL_MAX_TEXTURE_BUFFER_SIZE to OpenGLES.h to not depend on extension GL_MAX_TEXTURE_BUFFER_SIZE_EXT

diff --git a/Engine/Source/Runtime/OpenGLDrv/Public/OpenGLES.h b/Engine/Source/Runtime/OpenGLDrv/Public/OpenGLES.h
--- a/Engine/Source/Runtime/OpenGLDrv/Public/OpenGLES.h
+++ b/Engine/Source/Runtime/OpenGLDrv/Public/OpenGLES.h
@@ -94,6 +94,9 @@ typedef GLfloat GLdouble;
 #ifndef GL_TEXTURE_RECTANGLE
 #define GL_TEXTURE_RECTANGLE	0x84F5
 #endif
+#ifndef GL_MAX_TEXTURE_BUFFER_SIZE
+#define GL_MAX_TEXTURE_BUFFER_SIZE 0x8C2B
+#endif
 
 /** For the shader stage bits that don't exist just use 0 */
 #define GL_GEOMETRY_SHADER_BIT				0x00000000
