GL: Always use 'orphaining' on buffer updates as current buffer data could be in use by GPU

diff --git a/Engine/Source/Runtime/OpenGLDrv/Private/OpenGLBuffer.cpp b/Engine/Source/Runtime/OpenGLDrv/Private/OpenGLBuffer.cpp
--- a/Engine/Source/Runtime/OpenGLDrv/Private/OpenGLBuffer.cpp
+++ b/Engine/Source/Runtime/OpenGLDrv/Private/OpenGLBuffer.cpp
@@ -227,7 +227,10 @@ void* FOpenGLDynamicRHI::LockBuffer_BottomOfPipe(FRHICommandListImmediate& RHICm
 			return Staging;
 		}
 	}
-	return (void*)Buffer->Lock(Offset, Size, LockMode == EResourceLockMode::RLM_ReadOnly, Buffer->IsDynamic());
+
+	const bool bReadOnly = (LockMode == EResourceLockMode::RLM_ReadOnly);
+	const bool bDiscard = !bReadOnly; // Always use 'orphaning' on write as buffer could be in use by GPU atm
+	return (void*)Buffer->Lock(Offset, Size, bReadOnly, bDiscard);
 	RHITHREAD_GLCOMMAND_EPILOGUE_RETURN(void*);
 }
 
diff --git a/Engine/Source/Runtime/OpenGLDrv/Public/OpenGLResources.h b/Engine/Source/Runtime/OpenGLDrv/Public/OpenGLResources.h
--- a/Engine/Source/Runtime/OpenGLDrv/Public/OpenGLResources.h
+++ b/Engine/Source/Runtime/OpenGLDrv/Public/OpenGLResources.h
@@ -454,48 +454,48 @@ public:
 
 	void CreateGLBuffer(const void *InData, const GLuint ResourceToUse, const uint32 ResourceSize)
 	{
-			VERIFY_GL_SCOPE();
+		VERIFY_GL_SCOPE();
 		uint32 InSize = BaseType::GetSize();
-			RealSize = ResourceSize ? ResourceSize : InSize;
-			if( ResourceToUse )
+		RealSize = ResourceSize ? ResourceSize : InSize;
+		if( ResourceToUse )
+		{
+			Resource = ResourceToUse;
+			check( Type != GL_UNIFORM_BUFFER || !IsUniformBufferBound(Resource) );
+			Bind();
+			FOpenGL::BufferSubData(Type, 0, InSize, InData);
+		}
+		else
+		{
+			if (BaseType::GLSupportsType())
 			{
-				Resource = ResourceToUse;
+				FOpenGL::GenBuffers(1, &Resource);
 				check( Type != GL_UNIFORM_BUFFER || !IsUniformBufferBound(Resource) );
 				Bind();
-				FOpenGL::BufferSubData(Type, 0, InSize, InData);
-			}
-			else
-			{
-				if (BaseType::GLSupportsType())
-				{
-					FOpenGL::GenBuffers(1, &Resource);
-					check( Type != GL_UNIFORM_BUFFER || !IsUniformBufferBound(Resource) );
-					Bind();
 #if !RESTRICT_SUBDATA_SIZE
-					if( InData == NULL || RealSize <= InSize )
-					{
-						glBufferData(Type, RealSize, InData, GetAccess());
-					}
-					else
-					{
-						glBufferData(Type, RealSize, NULL, GetAccess());
-						FOpenGL::BufferSubData(Type, 0, InSize, InData);
-					}
-#else
-					glBufferData(Type, RealSize, NULL, GetAccess());
-					if ( InData != NULL )
-					{
-						LoadData( 0, FMath::Min<uint32>(InSize,RealSize), InData);
-					}
-#endif
-					IncrementBufferMemory(Type, RealSize);
+				if( InData == NULL || RealSize <= InSize )
+				{
+					glBufferData(Type, RealSize, InData, GetAccess());
 				}
 				else
 				{
-					BaseType::CreateType(Resource, InData, InSize);
+					glBufferData(Type, RealSize, NULL, GetAccess());
+					FOpenGL::BufferSubData(Type, 0, InSize, InData);
 				}
+#else
+				glBufferData(Type, RealSize, NULL, GetAccess());
+				if ( InData != NULL )
+				{
+					LoadData( 0, FMath::Min<uint32>(InSize,RealSize), InData);
+				}
+#endif
+				IncrementBufferMemory(Type, RealSize);
+			}
+			else
+			{
+				BaseType::CreateType(Resource, InData, InSize);
 			}
 		}
+	}
 
 	virtual ~TOpenGLBuffer()
 	{
@@ -564,7 +564,7 @@ public:
 		uint32 DiscardSize = (bDiscard && !bUseMapBuffer && InSize == RealSize && !RESTRICT_SUBDATA_SIZE) ? 0 : RealSize;
 
 		// Don't call BufferData if Bindless is on, as bindless texture buffers make buffers immutable
-		if ( bDiscard && !OpenGLConsoleVariables::bBindlessTexture && OpenGLConsoleVariables::bUseBufferDiscard)
+		if (bDiscard && !OpenGLConsoleVariables::bBindlessTexture && OpenGLConsoleVariables::bUseBufferDiscard)
 		{
 			if (BaseType::GLSupportsType())
 			{
@@ -709,14 +709,7 @@ public:
 			if (BaseType::GLSupportsType() && (OpenGLConsoleVariables::bUseMapBuffer || bIsLockReadOnly))
 			{
 				check(!bLockBufferWasAllocated);
-				if (Type == GL_ARRAY_BUFFER || Type == GL_ELEMENT_ARRAY_BUFFER)
-				{
-					FOpenGL::UnmapBufferRange(Type, LockOffset, LockSize);
-				}
-				else
-				{
-					FOpenGL::UnmapBuffer(Type);
-				}
+				FOpenGL::UnmapBufferRange(Type, LockOffset, LockSize);
 				LockBuffer = NULL;
 			}
 			else
