Fixed fallback for ES31 support with hlslcc

diff --git a/Engine/Source/Runtime/OpenGLDrv/Private/OpenGLShaders.cpp b/Engine/Source/Runtime/OpenGLDrv/Private/OpenGLShaders.cpp
--- a/Engine/Source/Runtime/OpenGLDrv/Private/OpenGLShaders.cpp
+++ b/Engine/Source/Runtime/OpenGLDrv/Private/OpenGLShaders.cpp
@@ -855,11 +855,18 @@ void OPENGLDRV_API GLSLToDeviceCompatibleGLSL(FAnsiCharArray& GlslCodeOriginal,
 	
 	if (Capabilities.TargetPlatform == EOpenGLShaderTargetPlatform::OGLSTP_Android)
 	{
-		const ANSICHAR* ES320Version = "#version 320 es";
+		const ANSICHAR* ESVersion = "#version 320 es";
 
-		AppendCString(GlslCode, ES320Version);
+		bool FoundVersion = (FCStringAnsi::Strstr(GlslCodeOriginal.GetData(), ESVersion)) != nullptr;
+
+		if (!FoundVersion)
+		{
+			ESVersion = "#version 310 es";
+		}
+		
+		AppendCString(GlslCode, ESVersion);
 		AppendCString(GlslCode, "\n");
-		ReplaceCString(GlslCodeOriginal, ES320Version, "");
+		ReplaceCString(GlslCodeOriginal, ESVersion, "");
 	}
 
 	if (TypeEnum == GL_FRAGMENT_SHADER && Capabilities.bRequiresDisabledEarlyFragmentTests)
