Avoid likely extraneous driver side memcpy when restoring evicted OpenGL textures. Halves texture restore cost on Pixel3a.
#rnx

diff --git a/Engine/Source/Runtime/OpenGLDrv/Private/OpenGLTexture.cpp b/Engine/Source/Runtime/OpenGLDrv/Private/OpenGLTexture.cpp
--- a/Engine/Source/Runtime/OpenGLDrv/Private/OpenGLTexture.cpp
+++ b/Engine/Source/Runtime/OpenGLDrv/Private/OpenGLTexture.cpp
@@ -866,6 +866,57 @@ uint32 TOpenGLTexture<RHIResourceType>::GetLockSize(uint32 InMipIndex, uint32 Ar
 	return MipBytes;
 }
 
+template<typename RHIResourceType>
+void TOpenGLTexture<RHIResourceType>::Fill2DGLTextureImage(const FOpenGLTextureFormat& GLFormat, const bool bSRGB, uint32 MipIndex, const void* BufferOrPBOOffset, uint32 ImageSize, uint32 ArrayIndex)
+{
+	if (GLFormat.bCompressed)
+	{
+		if (GetAllocatedStorageForMip(MipIndex,ArrayIndex))
+		{
+			glCompressedTexSubImage2D(
+				bCubemap ? GL_TEXTURE_CUBE_MAP_POSITIVE_X + ArrayIndex : Target,
+				MipIndex,
+				0,
+				0,
+				FMath::Max<uint32>(1,(this->GetSizeX() >> MipIndex)),
+				FMath::Max<uint32>(1,(this->GetSizeY() >> MipIndex)),
+				GLFormat.InternalFormat[bSRGB],
+				ImageSize,
+				BufferOrPBOOffset);	// offset into PBO
+		}
+		else
+		{
+			glCompressedTexImage2D(
+				bCubemap ? GL_TEXTURE_CUBE_MAP_POSITIVE_X + ArrayIndex : Target,
+				MipIndex,
+				GLFormat.InternalFormat[bSRGB],
+				FMath::Max<uint32>(1,(this->GetSizeX() >> MipIndex)),
+				FMath::Max<uint32>(1,(this->GetSizeY() >> MipIndex)),
+				0,
+				ImageSize,
+				BufferOrPBOOffset);	// offset into PBO
+			SetAllocatedStorageForMip(MipIndex,ArrayIndex);
+		}
+	}
+	else
+	{
+		// All construction paths should have called TexStorage2D or TexImage2D. So we will
+		// always call TexSubImage2D.
+		check(GetAllocatedStorageForMip(MipIndex,ArrayIndex) == true);
+		glPixelStorei(GL_UNPACK_ALIGNMENT, 1);
+		glTexSubImage2D(
+			bCubemap ? GL_TEXTURE_CUBE_MAP_POSITIVE_X + ArrayIndex : Target,
+			MipIndex,
+			0,
+			0,
+			FMath::Max<uint32>(1,(this->GetSizeX() >> MipIndex)),
+			FMath::Max<uint32>(1,(this->GetSizeY() >> MipIndex)),
+			GLFormat.Format,
+			GLFormat.Type,
+			BufferOrPBOOffset);	// offset into PBO
+		glPixelStorei(GL_UNPACK_ALIGNMENT, 4);
+	}
+}
 
 template<typename RHIResourceType>
 void* TOpenGLTexture<RHIResourceType>::Lock(uint32 InMipIndex,uint32 ArrayIndex,EResourceLockMode LockMode,uint32& DestStride)
@@ -1114,53 +1165,7 @@ void TOpenGLTexture<RHIResourceType>::Unlock(uint32 MipIndex,uint32 ArrayIndex)
 		}
 		else
 		{
-			if (GLFormat.bCompressed)
-			{
-				if (GetAllocatedStorageForMip(MipIndex,ArrayIndex))
-				{
-					glCompressedTexSubImage2D(
-						bCubemap ? GL_TEXTURE_CUBE_MAP_POSITIVE_X + ArrayIndex : Target,
-						MipIndex,
-						0,
-						0,
-						FMath::Max<uint32>(1,(this->GetSizeX() >> MipIndex)),
-						FMath::Max<uint32>(1,(this->GetSizeY() >> MipIndex)),
-						GLFormat.InternalFormat[bSRGB],
-						PixelBuffer->GetSize(),
-						0);	// offset into PBO
-				}
-				else
-				{
-					glCompressedTexImage2D(
-						bCubemap ? GL_TEXTURE_CUBE_MAP_POSITIVE_X + ArrayIndex : Target,
-						MipIndex,
-						GLFormat.InternalFormat[bSRGB],
-						FMath::Max<uint32>(1,(this->GetSizeX() >> MipIndex)),
-						FMath::Max<uint32>(1,(this->GetSizeY() >> MipIndex)),
-						0,
-						PixelBuffer->GetSize(),
-						0);	// offset into PBO
-					SetAllocatedStorageForMip(MipIndex,ArrayIndex);
-				}
-			}
-			else
-			{
-				// All construction paths should have called TexStorage2D or TexImage2D. So we will
-				// always call TexSubImage2D.
-				check(GetAllocatedStorageForMip(MipIndex,ArrayIndex) == true);
-				glPixelStorei(GL_UNPACK_ALIGNMENT, 1);
-				glTexSubImage2D(
-					bCubemap ? GL_TEXTURE_CUBE_MAP_POSITIVE_X + ArrayIndex : Target,
-					MipIndex,
-					0,
-					0,
-					FMath::Max<uint32>(1,(this->GetSizeX() >> MipIndex)),
-					FMath::Max<uint32>(1,(this->GetSizeY() >> MipIndex)),
-					GLFormat.Format,
-					GLFormat.Type,
-					0);	// offset into PBO
-				glPixelStorei(GL_UNPACK_ALIGNMENT, 4);
-			}
+			Fill2DGLTextureImage(GLFormat, bSRGB, MipIndex, 0, PixelBuffer->GetSize(), ArrayIndex);
 		}
 	}
 
@@ -1204,18 +1209,16 @@ void TOpenGLTexture<RHIResourceType>::RestoreEvictedGLResource(bool bAttemptToRe
 	const FOpenGLTextureFormat& GLFormat = GOpenGLTextureFormats[PixelFormat];
 	const bool bSRGB = (this->GetFlags() & TexCreate_SRGB) != 0;
 	checkf(EvictionParamsPtr->MipImageData.Num() == this->GetNumMips(), TEXT("EvictionParamsPtr->MipImageData.Num() =%d, this->GetNumMips() = %d"), EvictionParamsPtr->MipImageData.Num(), this->GetNumMips());
-	
+
+	CachedBindPixelUnpackBuffer(0);
+
 	for (int i = EvictionParamsPtr->MipImageData.Num() - 1; i >= 0; i--)
 	{
 		auto& MipMem = EvictionParamsPtr->MipImageData[i];
 		if(MipMem.Num())
 		{
-			uint32 DestStride;
 			check(MipMem.Num() == GetLockSize(i, 0, EResourceLockMode::RLM_WriteOnly, DestStride));
-			void* pDest = Lock(i, 0, EResourceLockMode::RLM_WriteOnly, DestStride);
-			check(DestStride)
-			FMemory::Memcpy(pDest, MipMem.GetData(), MipMem.Num());
-			Unlock(i, 0);
+			Fill2DGLTextureImage(GLFormat, bSRGB, i, MipMem.GetData(), MipMem.Num(), 0);
 		}
 	}
 
@@ -3072,8 +3075,8 @@ void LogTextureEvictionDebugInfo()
 	{
 		UE_LOG(LogRHI, Warning, TEXT("txdbg: Texture mipmem %d. GTotalTexStorageSkipped %d, GTotalCompressedTexStorageSkipped %d, Total noncompressed = %d"), GTotalEvictedMipMemStored, GTotalTexStorageSkipped, GTotalCompressedTexStorageSkipped, GTotalTexStorageSkipped - GTotalCompressedTexStorageSkipped);
 		UE_LOG(LogRHI, Warning, TEXT("txdbg: Texture GTotalEvictedMipMemDuplicated %d"), GTotalEvictedMipMemDuplicated);
-	UE_LOG(LogRHI, Warning, TEXT("txdbg: Texture GTotalMipRestores %d, GTotalMipStoredCount %d"), GTotalMipRestores, GTotalMipStoredCount);
-		UE_LOG(LogRHI, Warning, TEXT("txdbg: Texture GAvgRestoreTime %f (%d), GMaxRestoreTime %f"), GAvgRestoreCount ? GAvgRestoreTime / GAvgRestoreCount : 0.f, GAvgRestoreCount, GMaxRestoreTime);
+		UE_LOG(LogRHI, Warning, TEXT("txdbg: Texture GTotalMipRestores %d, GTotalMipStoredCount %d"), GTotalMipRestores, GTotalMipStoredCount);
+		UE_LOG(LogRHI, Warning, TEXT("txdbg: Texture GAvgRestoreTime %f (%d), GMaxRestoreTime %f, TotalRestoreTime %f"), GAvgRestoreCount ? (float)(GAvgRestoreTime / GAvgRestoreCount) : 0.f, GAvgRestoreCount, (float)GMaxRestoreTime, (float)GAvgRestoreTime);
 		UE_LOG(LogRHI, Warning, TEXT("txdbg: Texture LRU %d"), FTextureEvictionLRU::Get().Num());
 
 		GAvgRestoreCount = 0;
diff --git a/Engine/Source/Runtime/OpenGLDrv/Public/OpenGLResources.h b/Engine/Source/Runtime/OpenGLDrv/Public/OpenGLResources.h
--- a/Engine/Source/Runtime/OpenGLDrv/Public/OpenGLResources.h
+++ b/Engine/Source/Runtime/OpenGLDrv/Public/OpenGLResources.h
@@ -1670,6 +1670,8 @@ public:
 	virtual bool CanBeEvicted() override;
 	virtual void TryEvictGLResource() override;
 private:
+	void Fill2DGLTextureImage(const FOpenGLTextureFormat& GLFormat, const bool bSRGB, uint32 MipIndex, const void* LockedBuffer, uint32 LockedSize, uint32 ArrayIndex);
+
 	TArray< TRefCountPtr<FOpenGLPixelBuffer> > PixelBuffers;
 
 	uint32 GetEffectiveSizeZ( void ) { return this->GetSizeZ() ? this->GetSizeZ() : 1; }
