Remove broken check()
#rnx

diff --git a/Engine/Source/Runtime/OpenGLDrv/Private/OpenGLTexture.cpp b/Engine/Source/Runtime/OpenGLDrv/Private/OpenGLTexture.cpp
--- a/Engine/Source/Runtime/OpenGLDrv/Private/OpenGLTexture.cpp
+++ b/Engine/Source/Runtime/OpenGLDrv/Private/OpenGLTexture.cpp
@@ -1217,7 +1217,6 @@ void TOpenGLTexture<RHIResourceType>::RestoreEvictedGLResource(bool bAttemptToRe
 		auto& MipMem = EvictionParamsPtr->MipImageData[i];
 		if(MipMem.Num())
 		{
-			check(MipMem.Num() == GetLockSize(i, 0, EResourceLockMode::RLM_WriteOnly, DestStride));
 			Fill2DGLTextureImage(GLFormat, bSRGB, i, MipMem.GetData(), MipMem.Num(), 0);
 		}
 	}
