Removed GL ES2 path for vertex attributes setup
Removed BUF_ZeroStride buffer flag

diff --git a/Engine/Plugins/Enterprise/LidarPointCloud/Source/LidarPointCloudRuntime/Private/Rendering/LidarPointCloudRenderBuffers.cpp b/Engine/Plugins/Enterprise/LidarPointCloud/Source/LidarPointCloudRuntime/Private/Rendering/LidarPointCloudRenderBuffers.cpp
--- a/Engine/Plugins/Enterprise/LidarPointCloud/Source/LidarPointCloudRuntime/Private/Rendering/LidarPointCloudRenderBuffers.cpp
+++ b/Engine/Plugins/Enterprise/LidarPointCloud/Source/LidarPointCloudRuntime/Private/Rendering/LidarPointCloudRenderBuffers.cpp
@@ -273,7 +273,7 @@ void FLidarPointCloudVertexFactory::ReleaseRHI()
 void FLidarPointCloudSharedVertexFactory::FPointCloudVertexBuffer::InitRHI()
 {
 	FRHIResourceCreateInfo CreateInfo(TEXT("FPointCloudVertexBuffer"));
-	VertexBufferRHI = RHICreateBuffer(sizeof(FVector), BUF_Static | BUF_VertexBuffer | BUF_ZeroStride, 0, ERHIAccess::VertexOrIndexBuffer, CreateInfo);
+	VertexBufferRHI = RHICreateBuffer(sizeof(FVector), BUF_Static | BUF_VertexBuffer, 0, ERHIAccess::VertexOrIndexBuffer, CreateInfo);
 	void* Buffer = RHILockBuffer(VertexBufferRHI, 0, sizeof(FVector), RLM_WriteOnly);
 	FMemory::Memzero(Buffer, sizeof(FVector));
 	RHIUnlockBuffer(VertexBufferRHI);
diff --git a/Engine/Source/Runtime/Engine/Public/ParticleHelper.h b/Engine/Source/Runtime/Engine/Public/ParticleHelper.h
--- a/Engine/Source/Runtime/Engine/Public/ParticleHelper.h
+++ b/Engine/Source/Runtime/Engine/Public/ParticleHelper.h
@@ -2614,7 +2614,7 @@ public:
 	{
 		// create a static vertex buffer
 		FRHIResourceCreateInfo CreateInfo(TEXT("FNullDynamicParameterVertexBuffer"));
-		VertexBufferRHI = RHICreateBuffer(sizeof(FParticleVertexDynamicParameter), BUF_Static | BUF_VertexBuffer | BUF_ZeroStride, 0, ERHIAccess::VertexOrIndexBuffer, CreateInfo);
+		VertexBufferRHI = RHICreateBuffer(sizeof(FParticleVertexDynamicParameter), BUF_Static | BUF_VertexBuffer, 0, ERHIAccess::VertexOrIndexBuffer, CreateInfo);
 		FParticleVertexDynamicParameter* Vertices = (FParticleVertexDynamicParameter*)RHILockBuffer(VertexBufferRHI, 0, sizeof(FParticleVertexDynamicParameter), RLM_WriteOnly);
 		Vertices[0].DynamicValue[0] = Vertices[0].DynamicValue[1] = Vertices[0].DynamicValue[2] = Vertices[0].DynamicValue[3] = 1.0f;
 		RHIUnlockBuffer(VertexBufferRHI);
diff --git a/Engine/Source/Runtime/OpenGLDrv/Private/OpenGL4.cpp b/Engine/Source/Runtime/OpenGLDrv/Private/OpenGL4.cpp
--- a/Engine/Source/Runtime/OpenGLDrv/Private/OpenGL4.cpp
+++ b/Engine/Source/Runtime/OpenGLDrv/Private/OpenGL4.cpp
@@ -17,7 +17,6 @@ GLint FOpenGL4::MaxComputeUAVUnits = -1;
 GLint FOpenGL4::MaxPixelUAVUnits = -1;
 
 bool FOpenGL4::bSupportsGPUMemoryInfo = false;
-bool FOpenGL4::bSupportsVertexAttribBinding = true;
 bool FOpenGL4::bSupportsTextureView = true;
 
 void FOpenGL4::ProcessQueryGLInt()
@@ -53,7 +52,6 @@ void FOpenGL4::ProcessExtensions( const FString& ExtensionsString )
 
 
 	bSupportsGPUMemoryInfo = ExtensionsString.Contains(TEXT("GL_NVX_gpu_memory_info"));
-	bSupportsVertexAttribBinding = ExtensionsString.Contains(TEXT("GL_ARB_vertex_attrib_binding")) || (MajorVersion == 4 && MinorVersion >= 3) || (MajorVersion > 4);
 	bSupportsTextureView = ExtensionsString.Contains(TEXT("GL_ARB_texture_view")) || (MajorVersion == 4 && MinorVersion >= 3) || (MajorVersion > 4);
 
 	//Process Queries after extensions to avoid queries that use functionality that might not be present
@@ -87,11 +85,6 @@ uint64 FOpenGL4::GetVideoMemorySize()
 	return VideoMemorySize;
 }
 
-static FORCEINLINE bool SupportsVertexAttribBinding()
-{
-	return glVertexAttribBinding != NULL;
-}
-
 static FORCEINLINE void BindVertexBuffer(GLuint BindingIndex, GLuint Buffer, GLintptr Offset, GLsizei Stride)
 {
 	glBindVertexBuffer(BindingIndex, Buffer, Offset, Stride);
diff --git a/Engine/Source/Runtime/OpenGLDrv/Private/OpenGLBuffer.cpp b/Engine/Source/Runtime/OpenGLDrv/Private/OpenGLBuffer.cpp
--- a/Engine/Source/Runtime/OpenGLDrv/Private/OpenGLBuffer.cpp
+++ b/Engine/Source/Runtime/OpenGLDrv/Private/OpenGLBuffer.cpp
@@ -219,14 +219,6 @@ void* FOpenGLDynamicRHI::LockBuffer_BottomOfPipe(FRHICommandListImmediate& RHICm
 
 	VERIFY_GL_SCOPE();
 	FOpenGLBuffer* Buffer = ResourceCast(BufferRHI);
-	if( !(FOpenGL::SupportsVertexAttribBinding() && OpenGLConsoleVariables::bUseVAB) && Buffer->GetUsage() & BUF_ZeroStride )
-	{
-		check( Offset + Size <= Buffer->GetSize() );
-		// We assume we are only using the first elements from the VB, so we can later on read this memory and create an expanded version of this zero stride buffer
-		check(Offset == 0);
-		return (void*)( (uint8*)Buffer->GetZeroStrideBuffer() + Offset );
-	}
-	else
 	{
 		if (Buffer->IsDynamic() && LockMode == EResourceLockMode::RLM_WriteOnly)
 		{
@@ -246,12 +238,9 @@ void FOpenGLDynamicRHI::UnlockBuffer_BottomOfPipe(FRHICommandListImmediate& RHIC
 	RHITHREAD_GLCOMMAND_PROLOGUE();
 	VERIFY_GL_SCOPE();
 	FOpenGLBuffer* Buffer = ResourceCast(BufferRHI);
-	if( (FOpenGL::SupportsVertexAttribBinding() && OpenGLConsoleVariables::bUseVAB) || !( Buffer->GetUsage() & BUF_ZeroStride ) )
+	if (!RetireAllocation(Buffer))
 	{
-		if (!RetireAllocation(Buffer))
-		{
-			Buffer->Unlock();
-		}
+		Buffer->Unlock();
 	}
 	RHITHREAD_GLCOMMAND_EPILOGUE();
 }
diff --git a/Engine/Source/Runtime/OpenGLDrv/Private/OpenGLCommands.cpp b/Engine/Source/Runtime/OpenGLDrv/Private/OpenGLCommands.cpp
--- a/Engine/Source/Runtime/OpenGLDrv/Private/OpenGLCommands.cpp
+++ b/Engine/Source/Runtime/OpenGLDrv/Private/OpenGLCommands.cpp
@@ -54,14 +54,6 @@ namespace OpenGLConsoleVariables
 		TEXT("If true, don't issue dispatch work.")
 		);
 
-	int32 bUseVAB = 1;
-	static FAutoConsoleVariableRef CVarUseVAB(
-		TEXT("OpenGL.UseVAB"),
-		bUseVAB,
-		TEXT("If true, use GL_VERTEX_ATTRIB_BINDING instead of traditional vertex array setup."),
-		ECVF_ReadOnly
-		);
-
 #if PLATFORM_WINDOWS || PLATFORM_LINUX
 	int32 MaxSubDataSize = 256*1024;
 #else
@@ -137,132 +129,6 @@ struct FVertexBufferPair
 	FOpenGLBuffer*				Source;
 	TRefCountPtr<FOpenGLBuffer>	Dest;
 };
-static TArray<FVertexBufferPair> ZeroStrideExpandedBuffersList;
-
-
-static int FindVertexBuffer(FOpenGLBuffer* Source)
-{
-	for (int32 Index = 0; Index < ZeroStrideExpandedBuffersList.Num(); ++Index)
-	{
-		if (ZeroStrideExpandedBuffersList[Index].Source == Source)
-		{
-			return Index;
-		}
-	}
-	return -1;
-}
-
-static FOpenGLBuffer* FindExpandedZeroStrideBuffer(FOpenGLBuffer* ZeroStrideVertexBuffer, uint32 Stride, uint32 NumVertices, const FOpenGLVertexElement& VertexElement)
-{
-	uint32 Size = NumVertices * Stride;
-	int32 FoundExpandedVBIndex = FindVertexBuffer(ZeroStrideVertexBuffer);
-	if (FoundExpandedVBIndex != -1)
-	{
-		// Check if the current size is big enough
-		FOpenGLBuffer* ExpandedVB = ZeroStrideExpandedBuffersList[FoundExpandedVBIndex].Dest;
-		if (Size <= ExpandedVB->GetSize())
-		{
-			return ExpandedVB;
-		}
-	}
-	else
-	{
-		FVertexBufferPair NewPair;
-		NewPair.Source = ZeroStrideVertexBuffer;
-		NewPair.Dest = NULL;
-		FoundExpandedVBIndex = ZeroStrideExpandedBuffersList.Num();
-		ZeroStrideExpandedBuffersList.Add(NewPair);
-	}
-
-	int32 VertexTypeSize = 0;
-	switch( VertexElement.Type )
-	{
-	case GL_FLOAT:
-	case GL_UNSIGNED_INT:
-	case GL_INT:
-		VertexTypeSize = 4;
-		break;
-	case GL_SHORT:
-	case GL_UNSIGNED_SHORT:
-	case GL_HALF_FLOAT:
-		VertexTypeSize = 2;
-		break;
-	case GL_BYTE:
-	case GL_UNSIGNED_BYTE:
-		VertexTypeSize = 1;
-		break;
-	case GL_DOUBLE:
-		VertexTypeSize = 8;
-		break;
-	default:
-		check(0);
-		break;
-	}
-
-	const int32 VertexElementSize = ( VertexElement.Size == GL_BGRA ) ? 4 : VertexElement.Size;
-	const int32 SizeToFill = VertexElementSize * VertexTypeSize;
-	void* RESTRICT SourceData = ZeroStrideVertexBuffer->GetZeroStrideBuffer();
-	check(SourceData);
-	TRefCountPtr<FOpenGLBuffer> ExpandedVB = new FOpenGLBuffer(GL_ARRAY_BUFFER, 0, Size, BUF_Static, nullptr);
-	uint8* RESTRICT Data = ExpandedVB->Lock(0, Size, false, true);
-
-	switch (SizeToFill)
-	{
-	case 4:
-		{
-			uint32 Source = *(uint32*)SourceData;
-			uint32* RESTRICT Dest = (uint32*)Data;
-			for (uint32 Index = 0; Index < Size / sizeof(uint32); ++Index)
-			{
-				*Dest++ = Source;
-			}
-		}
-		break;
-	case 8:
-		{
-			uint64 Source = *(uint64*)SourceData;
-			uint64* RESTRICT Dest = (uint64*)Data;
-			for (uint32 Index = 0; Index < Size / sizeof(uint64); ++Index)
-			{
-				*Dest++ = Source;
-			}
-		}
-		break;
-	case 12:
-		{
-			uint64 SourceA = *(uint64*)SourceData;
-			uint32 SourceB = *((uint32*)SourceData + 2);
-			uint32* RESTRICT Dest = (uint32*)Data;
-			for (uint32 Index = 0; Index < Size / (3 * sizeof(uint32)); ++Index)
-			{
-				*((uint64*)Dest) = SourceA;
-				Dest = Dest + 2;
-				*Dest++ = SourceB;
-			}
-		}
-		break;
-	case 16:
-		{
-			uint64 SourceA = *(uint64*)SourceData;
-			uint64 SourceB = *((uint64*)SourceData + 1);
-			uint64* RESTRICT Dest = (uint64*)Data;
-			for (uint32 Index = 0; Index < Size / (2 * sizeof(uint64)); ++Index)
-			{
-				*Dest++ = SourceA;
-				*Dest++ = SourceB;
-			}
-		}
-		break;
-	default:
-		check(0);
-	}
-
-	ExpandedVB->Unlock();
-
-	ZeroStrideExpandedBuffersList[FoundExpandedVBIndex].Dest = ExpandedVB;
-
-	return ExpandedVB;
-}
 
 static FORCEINLINE GLint ModifyFilterByMips(GLint Filter, bool bHasMips)
 {
@@ -1837,182 +1703,7 @@ void FOpenGLDynamicRHI::SetRenderTargetsAndClear(const FRHISetRenderTargetsInfo&
 
 // Primitive drawing.
 
-void FOpenGLDynamicRHI::EnableVertexElementCached(
-	FOpenGLContextState& ContextState,
-	GLuint AttributeIndex,
-	const FOpenGLVertexElement &VertexElement,
-	GLsizei Stride,
-	void *Pointer,
-	GLuint Buffer)
-{
-	VERIFY_GL_SCOPE();
-
-	check( !(FOpenGL::SupportsVertexAttribBinding() && OpenGLConsoleVariables::bUseVAB));
-
-	if (!ContextState.GetVertexAttrEnabled(AttributeIndex))
-	{
-		ContextState.SetVertexAttrEnabled(AttributeIndex, true);
-		glEnableVertexAttribArray(AttributeIndex);
-	}
-	FOpenGLCachedAttr &Attr = ContextState.VertexAttrs[AttributeIndex];
-
-	bool bAnyDifferent = //bitwise ors to get rid of the branches
-		(Attr.Pointer != Pointer) |
-		(Attr.Buffer != Buffer) |
-		(Attr.Size != VertexElement.Size) |
-		(Attr.Type != VertexElement.Type) |
-		(Attr.bNormalized != VertexElement.bNormalized) |
-		(Attr.Stride != Stride) |
-		(Attr.bShouldConvertToFloat != VertexElement.bShouldConvertToFloat); 
-
-	if (bAnyDifferent)
-	{
-		CachedBindArrayBuffer(ContextState, Buffer);
-		if( !VertexElement.bShouldConvertToFloat )
-		{
-			FOpenGL::VertexAttribIPointer(
-				AttributeIndex,
-				VertexElement.Size,
-				VertexElement.Type,
-				Stride,
-				Pointer
-				);
-		}
-		else
-		{
-			FOpenGL::VertexAttribPointer(
-				AttributeIndex,
-				VertexElement.Size,
-				VertexElement.Type,
-				VertexElement.bNormalized,
-				Stride,
-				Pointer
-				);
-		}
-		
-		Attr.Pointer = Pointer;
-		Attr.Buffer = Buffer;
-		Attr.Size = VertexElement.Size;
-		Attr.Type = VertexElement.Type;
-		Attr.bNormalized = VertexElement.bNormalized;
-		Attr.Stride = Stride;
-		Attr.bShouldConvertToFloat = VertexElement.bShouldConvertToFloat;
-	}
-
-	if (Attr.Divisor != VertexElement.Divisor)
-	{
-		FOpenGL::VertexAttribDivisor(AttributeIndex, VertexElement.Divisor);
-		Attr.Divisor = VertexElement.Divisor;
-	}
-}
-
-FORCEINLINE void FOpenGLDynamicRHI::EnableVertexElementCachedZeroStride(FOpenGLContextState& ContextState, GLuint AttributeIndex, const FOpenGLVertexElement& VertexElement, uint32 NumVertices, FOpenGLBuffer* ZeroStrideVertexBuffer)
-{
-	uint32 Stride = ZeroStrideVertexBuffer->GetSize();
-	FOpenGLBuffer* ExpandedVertexBuffer = FindExpandedZeroStrideBuffer(ZeroStrideVertexBuffer, Stride, NumVertices, VertexElement);
-	EnableVertexElementCached(ContextState, AttributeIndex, VertexElement, Stride, 0, ExpandedVertexBuffer->Resource);
-}
-
-void FOpenGLDynamicRHI::FreeZeroStrideBuffers()
-{
-	// Forces releasing references to expanded zero stride vertex buffers
-	ZeroStrideExpandedBuffersList.Empty();
-}
-
 void FOpenGLDynamicRHI::SetupVertexArrays(FOpenGLContextState& ContextState, uint32 BaseVertexIndex, FOpenGLStream* Streams, uint32 NumStreams, uint32 MaxVertices)
-{
-	SCOPE_CYCLE_COUNTER_DETAILED(STAT_OpenGLVBOSetupTime);
-	if (FOpenGL::SupportsVertexAttribBinding() && OpenGLConsoleVariables::bUseVAB)
-	{
-		SetupVertexArraysVAB(ContextState, BaseVertexIndex, Streams, NumStreams, MaxVertices);
-		return;
-	}
-	VERIFY_GL_SCOPE();
-	
-	uint32 UsedAttributes = 0;
-	static_assert(NUM_OPENGL_VERTEX_STREAMS <= sizeof(UsedAttributes) * 8, "Not enough bits in UsedAttributes to store NUM_OPENGL_VERTEX_STREAMS");
-
-	check(IsValidRef(PendingState.BoundShaderState));
-
-	const FOpenGLShaderBindings& Bindings = PendingState.BoundShaderState->GetVertexShader()->Bindings;
-
-	FOpenGLVertexDeclaration* VertexDeclaration = PendingState.BoundShaderState->VertexDeclaration;
-	for (int32 ElementIndex = 0; ElementIndex < VertexDeclaration->VertexElements.Num(); ElementIndex++)
-	{
-		FOpenGLVertexElement& VertexElement = VertexDeclaration->VertexElements[ElementIndex];
-		uint32 AttributeIndex = VertexElement.AttributeIndex;
-		if (!Bindings.InOutMask.IsFieldEnabled((int32)AttributeIndex))
-		{
-			continue; // skip unused attributes.
-		}
-
-		if (VertexElement.StreamIndex < NumStreams)
-		{
-			FOpenGLStream* Stream = &Streams[VertexElement.StreamIndex];
-			uint32 Stride = Stream->Stride;
-
-			if( Stream->VertexBuffer->GetUsage() & BUF_ZeroStride )
-			{
-				check(Stride == 0);
-				check(Stream->Offset == 0);
-				check(VertexElement.Offset == 0);
-				check(Stream->VertexBuffer->GetZeroStrideBuffer());
-				EnableVertexElementCachedZeroStride(
-					ContextState,
-					AttributeIndex,
-					VertexElement,
-					MaxVertices,
-					Stream->VertexBuffer
-					);
-			}
-			else
-			{
-				check( Stride > 0 );
-				EnableVertexElementCached(
-					ContextState,
-					AttributeIndex,
-					VertexElement,
-					Stride,
-					INDEX_TO_VOID(BaseVertexIndex * Stride + Stream->Offset + VertexElement.Offset),
-					Stream->VertexBuffer->Resource
-					);
-			}
-
-			UsedAttributes |= (1<<AttributeIndex);
-		}
-		else
-		{
-			//workaround attributes with no streams
-			VERIFY_GL_SCOPE();
-
-			if (ContextState.GetVertexAttrEnabled(AttributeIndex))
-			{
-				ContextState.SetVertexAttrEnabled(AttributeIndex, false);
-				glDisableVertexAttribArray(AttributeIndex);
-			}
-
-			float data[4] = { 0.0f};
-
-			glVertexAttrib4fv(AttributeIndex, data);
-		}
-	}
-
-	uint32 NotUsedButEnabledMask = (ContextState.VertexAttrs_EnabledBits & ~(UsedAttributes));
-
-	for (GLuint AttribIndex = 0; AttribIndex < NUM_OPENGL_VERTEX_STREAMS && NotUsedButEnabledMask; AttribIndex++)
-	{
-		if (NotUsedButEnabledMask & 1)
-		{
-			glDisableVertexAttribArray(AttribIndex);
-			ContextState.SetVertexAttrEnabled(AttribIndex, false);
-		}
-		NotUsedButEnabledMask >>= 1;
-	}
-
-	// Disable remaining vertex arrays
-}
-
-void FOpenGLDynamicRHI::SetupVertexArraysVAB(FOpenGLContextState& ContextState, uint32 BaseVertexIndex, FOpenGLStream* Streams, uint32 NumStreams, uint32 MaxVertices)
 {
 	VERIFY_GL_SCOPE();
 	bool KnowsDivisor[NUM_OPENGL_VERTEX_STREAMS] = { 0 };
@@ -2070,6 +1761,7 @@ void FOpenGLDynamicRHI::SetupVertexArraysVAB(FOpenGLContextState& ContextState,
 						Attr.Size = VertexElement.Size;
 						Attr.Type = VertexElement.Type;
 						Attr.bNormalized = VertexElement.bNormalized;
+						Attr.bShouldConvertToFloat = VertexElement.bShouldConvertToFloat;
 					}
 
 					if (Attr.StreamIndex != StreamIndex)
@@ -2187,53 +1879,6 @@ void FOpenGLDynamicRHI::SetupVertexArraysVAB(FOpenGLContextState& ContextState,
 	check(NotUsedButActiveStreamMask == 0);
 }
 
-// Used by default on ES2 for immediate mode rendering.
-void FOpenGLDynamicRHI::SetupVertexArraysUP(FOpenGLContextState& ContextState, void* Buffer, uint32 Stride)
-{
-	VERIFY_GL_SCOPE();
-
-	uint32 UsedAttributes = 0;
-	static_assert(NUM_OPENGL_VERTEX_STREAMS <= sizeof(UsedAttributes) * 8, "Not enough bits in UsedAttributes to store NUM_OPENGL_VERTEX_STREAMS");
-
-	check(IsValidRef(PendingState.BoundShaderState));
-	FOpenGLVertexDeclaration* VertexDeclaration = PendingState.BoundShaderState->VertexDeclaration;
-
-	const FOpenGLShaderBindings& Bindings = PendingState.BoundShaderState->GetVertexShader()->Bindings;
-
-	for (int32 ElementIndex = 0; ElementIndex < VertexDeclaration->VertexElements.Num(); ElementIndex++)
-	{
-		FOpenGLVertexElement &VertexElement = VertexDeclaration->VertexElements[ElementIndex];
-		check(VertexElement.StreamIndex < 1);
-
-		uint32 AttributeIndex = VertexElement.AttributeIndex;
-		if (Bindings.InOutMask.IsFieldEnabled(AttributeIndex))
-		{
-			check(Stride > 0);
-			EnableVertexElementCached(
-				ContextState,
-				AttributeIndex,
-				VertexElement,
-				Stride,
-				(void*)(((char*)Buffer) + VertexElement.Offset),
-				0
-				);
-			UsedAttributes |= (1 << AttributeIndex);
-		}
-	}
-
-	uint32 NotUsedButEnabledMask = (ContextState.VertexAttrs_EnabledBits & ~(UsedAttributes));
-
-	for (GLuint AttribIndex = 0; AttribIndex < NUM_OPENGL_VERTEX_STREAMS && NotUsedButEnabledMask; AttribIndex++)
-	{
-		if (NotUsedButEnabledMask & 1)
-		{
-			glDisableVertexAttribArray(AttribIndex);
-			ContextState.SetVertexAttrEnabled(AttribIndex, false);
-		}
-		NotUsedButEnabledMask >>= 1;
-	}
-}
-
 void FOpenGLDynamicRHI::OnProgramDeletion( GLint ProgramResource )
 {
 	VERIFY_GL_SCOPE();
@@ -2272,60 +1917,39 @@ void FOpenGLDynamicRHI::OnBufferDeletion( GLuint BufferResource )
 		RenderingContextState.StorageBufferBound = -1;	// will force refresh
 	}
 
-	if (FOpenGL::SupportsVertexAttribBinding() && OpenGLConsoleVariables::bUseVAB)
+	// loop through active streams
+	uint32 ActiveStreamMask = SharedContextState.ActiveStreamMask;
+	for (GLuint StreamIndex = 0; StreamIndex < NUM_OPENGL_VERTEX_STREAMS && ActiveStreamMask; StreamIndex++)
 	{
-		// loop through active streams
-		uint32 ActiveStreamMask = SharedContextState.ActiveStreamMask;
-		for (GLuint StreamIndex = 0; StreamIndex < NUM_OPENGL_VERTEX_STREAMS && ActiveStreamMask; StreamIndex++)
+		FOpenGLStream& CachedStream = SharedContextState.VertexStreams[StreamIndex];
+		if ((ActiveStreamMask & 0x1) && 
+			CachedStream.VertexBuffer && 
+			CachedStream.VertexBuffer->Resource == BufferResource)
 		{
-			FOpenGLStream& CachedStream = SharedContextState.VertexStreams[StreamIndex];
-			if ((ActiveStreamMask & 0x1) && 
-				CachedStream.VertexBuffer && 
-				CachedStream.VertexBuffer->Resource == BufferResource)
-			{
-				FOpenGL::BindVertexBuffer(StreamIndex, 0, 0, 0); // brianh@nvidia: work around driver bug 1809000
-				CachedStream.VertexBuffer = nullptr;
-				CachedStream.Offset = 0;
-				CachedStream.Stride = 0;
-			}
-			ActiveStreamMask >>= 1;
-		}
-	
-		// loop through active streams
-		ActiveStreamMask = RenderingContextState.ActiveStreamMask;
-		for (GLuint StreamIndex = 0; StreamIndex < NUM_OPENGL_VERTEX_STREAMS && ActiveStreamMask; StreamIndex++)
-		{
-			FOpenGLStream& CachedStream = RenderingContextState.VertexStreams[StreamIndex];
-			if ((ActiveStreamMask & 0x1) && 
-				CachedStream.VertexBuffer && 
-				CachedStream.VertexBuffer->Resource == BufferResource)
-			{
-				FOpenGL::BindVertexBuffer(StreamIndex, 0, 0, 0); // brianh@nvidia: work around driver bug 1809000
-				CachedStream.VertexBuffer = nullptr;
-				CachedStream.Offset = 0;
-				CachedStream.Stride = 0;
-			}
-			ActiveStreamMask >>= 1;
+			FOpenGL::BindVertexBuffer(StreamIndex, 0, 0, 0); // brianh@nvidia: work around driver bug 1809000
+			CachedStream.VertexBuffer = nullptr;
+			CachedStream.Offset = 0;
+			CachedStream.Stride = 0;
 		}
+		ActiveStreamMask >>= 1;
 	}
-	else
-	{
-		for (GLuint AttribIndex = 0; AttribIndex < NUM_OPENGL_VERTEX_STREAMS; AttribIndex++)
-		{
-			if (SharedContextState.VertexAttrs[AttribIndex].Buffer == BufferResource)
-			{
-				SharedContextState.VertexAttrs[AttribIndex].Pointer = FOpenGLCachedAttr_Invalid;	// that'll enforce state update on next cache test		}
-			}
-		}
 	
-		for (GLuint AttribIndex = 0; AttribIndex < NUM_OPENGL_VERTEX_STREAMS; AttribIndex++)
-		{
-			if (RenderingContextState.VertexAttrs[AttribIndex].Buffer == BufferResource)
-			{
-				RenderingContextState.VertexAttrs[AttribIndex].Pointer = FOpenGLCachedAttr_Invalid;	// that'll enforce state update on next cache test		}
-			}
-		}
-	}
+	// loop through active streams
+	ActiveStreamMask = RenderingContextState.ActiveStreamMask;
+	for (GLuint StreamIndex = 0; StreamIndex < NUM_OPENGL_VERTEX_STREAMS && ActiveStreamMask; StreamIndex++)
+	{
+		FOpenGLStream& CachedStream = RenderingContextState.VertexStreams[StreamIndex];
+		if ((ActiveStreamMask & 0x1) && 
+			CachedStream.VertexBuffer && 
+			CachedStream.VertexBuffer->Resource == BufferResource)
+		{
+			FOpenGL::BindVertexBuffer(StreamIndex, 0, 0, 0); // brianh@nvidia: work around driver bug 1809000
+			CachedStream.VertexBuffer = nullptr;
+			CachedStream.Offset = 0;
+			CachedStream.Stride = 0;
+		}
+		ActiveStreamMask >>= 1;
+	}	
 	
 	// Storage buffer
 	{
diff --git a/Engine/Source/Runtime/OpenGLDrv/Private/OpenGLDevice.cpp b/Engine/Source/Runtime/OpenGLDrv/Private/OpenGLDevice.cpp
--- a/Engine/Source/Runtime/OpenGLDrv/Private/OpenGLDevice.cpp
+++ b/Engine/Source/Runtime/OpenGLDrv/Private/OpenGLDevice.cpp
@@ -232,10 +232,10 @@ JNI_METHOD void Java_com_epicgames_unreal_MediaPlayer14_nativeClearCachedAttribu
 
 	// update vertex attributes state
 	ContextState.SetVertexAttrEnabled(PositionAttrib, false);
-	ContextState.VertexAttrs[PositionAttrib].Stride = -1;
+	ContextState.VertexAttrs[PositionAttrib].Size = -1; // will force attribute update
 
 	ContextState.SetVertexAttrEnabled(TexCoordsAttrib, false);
-	ContextState.VertexAttrs[TexCoordsAttrib].Stride = -1;
+	ContextState.VertexAttrs[TexCoordsAttrib].Size = -1; // will force attribute update
 
 	// make sure the texture is set again
 	ContextState.ActiveTexture = 0;
@@ -1208,8 +1208,6 @@ void FOpenGLDynamicRHI::Cleanup()
 		FRenderResource::ReleaseRHIForAllResources();
 	}
 
-	FreeZeroStrideBuffers();
-
 	// Release the point sampler state.
 	PointSamplerState.SafeRelease();
 
diff --git a/Engine/Source/Runtime/OpenGLDrv/Public/OpenGL.h b/Engine/Source/Runtime/OpenGLDrv/Public/OpenGL.h
--- a/Engine/Source/Runtime/OpenGLDrv/Public/OpenGL.h
+++ b/Engine/Source/Runtime/OpenGLDrv/Public/OpenGL.h
@@ -128,7 +128,6 @@ public:
 	static FORCEINLINE bool SupportsSeamlessCubeMap()					{ return false; }
 	static FORCEINLINE bool SupportsDrawIndirect()						{ return false; }
 	static FORCEINLINE bool SupportsGenerateMipmap()					{ return false; }
-	static FORCEINLINE bool SupportsVertexAttribBinding()				{ return false; }
 	static FORCEINLINE bool SupportsBufferStorage()						{ return false; }
 	static FORCEINLINE bool SupportsDepthBoundsTest()					{ return false; }
 	static FORCEINLINE bool SupportsTextureRange()						{ return false; }
diff --git a/Engine/Source/Runtime/OpenGLDrv/Public/OpenGL4.h b/Engine/Source/Runtime/OpenGLDrv/Public/OpenGL4.h
--- a/Engine/Source/Runtime/OpenGLDrv/Public/OpenGL4.h
+++ b/Engine/Source/Runtime/OpenGLDrv/Public/OpenGL4.h
@@ -15,7 +15,6 @@
 struct FOpenGL4 : public FOpenGL3
 {
 	static FORCEINLINE bool SupportsDrawIndirect()						{ return true; }
-	static FORCEINLINE bool SupportsVertexAttribBinding()				{ return bSupportsVertexAttribBinding; }
 	static FORCEINLINE bool SupportsTextureView()						{ return bSupportsTextureView; }
 
 
@@ -125,7 +124,6 @@ protected:
 	static GLint MaxPixelUAVUnits;
 
 	static bool bSupportsGPUMemoryInfo;
-	static bool bSupportsVertexAttribBinding;
 	static bool bSupportsTextureView;
 
 };
diff --git a/Engine/Source/Runtime/OpenGLDrv/Public/OpenGLDrv.h b/Engine/Source/Runtime/OpenGLDrv/Public/OpenGLDrv.h
--- a/Engine/Source/Runtime/OpenGLDrv/Public/OpenGLDrv.h
+++ b/Engine/Source/Runtime/OpenGLDrv/Public/OpenGLDrv.h
@@ -1083,14 +1083,7 @@ private:
 
 	void InitializeStateResources();
 
-	/** needs to be called before each dispatch call */
-	
-
-	void EnableVertexElementCached(FOpenGLContextState& ContextCache, GLuint AttributeIndex, const FOpenGLVertexElement &VertexElement, GLsizei Stride, void *Pointer, GLuint Buffer);
-	void EnableVertexElementCachedZeroStride(FOpenGLContextState& ContextCache, GLuint AttributeIndex, const FOpenGLVertexElement &VertexElement, uint32 NumVertices, FOpenGLBuffer* VertexBuffer);
 	void SetupVertexArrays(FOpenGLContextState& ContextCache, uint32 BaseVertexIndex, FOpenGLStream* Streams, uint32 NumStreams, uint32 MaxVertices);
-	void SetupVertexArraysVAB(FOpenGLContextState& ContextCache, uint32 BaseVertexIndex, FOpenGLStream* Streams, uint32 NumStreams, uint32 MaxVertices);
-	void SetupVertexArraysUP(FOpenGLContextState& ContextState, void* Buffer, uint32 Stride);
 
 	void SetupBindlessTextures( FOpenGLContextState& ContextState, const TArray<FOpenGLBindlessSamplerInfo> &Samplers );
 
@@ -1151,8 +1144,6 @@ private:
 
 	void ClearCurrentFramebufferWithCurrentScissor(FOpenGLContextState& ContextState, int8 ClearType, int32 NumClearColors, const FLinearColor* ClearColorArray, float Depth, uint32 Stencil);
 
-	void FreeZeroStrideBuffers();
-
 	FTextureLockTracker GLLockTracker;
 
 };
diff --git a/Engine/Source/Runtime/OpenGLDrv/Public/OpenGLES.h b/Engine/Source/Runtime/OpenGLDrv/Public/OpenGLES.h
--- a/Engine/Source/Runtime/OpenGLDrv/Public/OpenGLES.h
+++ b/Engine/Source/Runtime/OpenGLDrv/Public/OpenGLES.h
@@ -215,7 +215,6 @@ struct FOpenGLES : public FOpenGLBase
 	static FORCEINLINE bool SupportsFramebufferSRGBEnable() { return false; }
 	static FORCEINLINE bool SupportsRGB10A2() { return bSupportsRGB10A2; }
 	static FORCEINLINE bool SupportsDrawIndirect() { return true; }
-	static FORCEINLINE bool SupportsVertexAttribBinding() { return true; }
 	static FORCEINLINE bool SupportsBufferStorage() { return bSupportsBufferStorage; }
 	
 
diff --git a/Engine/Source/Runtime/OpenGLDrv/Public/OpenGLResources.h b/Engine/Source/Runtime/OpenGLDrv/Public/OpenGLResources.h
--- a/Engine/Source/Runtime/OpenGLDrv/Public/OpenGLResources.h
+++ b/Engine/Source/Runtime/OpenGLDrv/Public/OpenGLResources.h
@@ -43,7 +43,6 @@ extern bool IsUniformBufferBound( GLuint Buffer );
 namespace OpenGLConsoleVariables
 {
 	extern int32 bUseMapBuffer;
-	extern int32 bUseVAB;
 	extern int32 MaxSubDataSize;
 	extern int32 bUseStagingBuffer;
 	extern int32 bBindlessTexture;
@@ -428,31 +427,28 @@ public:
 
 		RealSize = ResourceSize ? ResourceSize : InSize;
 
-		if( (FOpenGL::SupportsVertexAttribBinding() && OpenGLConsoleVariables::bUseVAB) || !( InUsage & BUF_ZeroStride ) )
-		{
-			FRHICommandListImmediate& RHICmdList = FRHICommandListExecutor::GetImmediateCommandList();
+		FRHICommandListImmediate& RHICmdList = FRHICommandListExecutor::GetImmediateCommandList();
 
-			if (ShouldRunGLRenderContextOpOnThisThread(RHICmdList))
+		if (ShouldRunGLRenderContextOpOnThisThread(RHICmdList))
+		{
+			CreateGLBuffer(InData, ResourceToUse, ResourceSize);
+		}
+		else
+		{
+			void* BuffData = nullptr;
+			if (InData)
 			{
-				CreateGLBuffer(InData, ResourceToUse, ResourceSize);
+				BuffData = RHICmdList.Alloc(RealSize, 16);
+				FMemory::Memcpy(BuffData, InData, RealSize);
 			}
-			else
+			TransitionFence.Reset();
+			ALLOC_COMMAND_CL(RHICmdList, FRHICommandGLCommand)([=]() 
 			{
-				void* BuffData = nullptr;
-				if (InData)
-				{
-					BuffData = RHICmdList.Alloc(RealSize, 16);
-					FMemory::Memcpy(BuffData, InData, RealSize);
-				}
-				TransitionFence.Reset();
-				ALLOC_COMMAND_CL(RHICmdList, FRHICommandGLCommand)([=]() 
-				{
-					CreateGLBuffer(BuffData, ResourceToUse, ResourceSize); 
-					TransitionFence.WriteAssertFence();
-				});
-				TransitionFence.SetRHIThreadFence();
+				CreateGLBuffer(BuffData, ResourceToUse, ResourceSize); 
+				TransitionFence.WriteAssertFence();
+			});
+			TransitionFence.SetRHIThreadFence();
 
-			}
 		}
 	}
 
@@ -540,14 +536,12 @@ public:
 	void Bind()
 	{
 		VERIFY_GL_SCOPE();
-		check( (FOpenGL::SupportsVertexAttribBinding() && OpenGLConsoleVariables::bUseVAB) || ( this->GetUsage() & BUF_ZeroStride ) == 0 );
 		BufBind(Type, Resource);
 	}
 
 	uint8 *Lock(uint32 InOffset, uint32 InSize, bool bReadOnly, bool bDiscard)
 	{
 		SCOPE_CYCLE_COUNTER_DETAILED(STAT_OpenGLMapBufferTime);
-		check( (FOpenGL::SupportsVertexAttribBinding() && OpenGLConsoleVariables::bUseVAB) || ( this->GetUsage() & BUF_ZeroStride ) == 0 );
 		check(InOffset + InSize <= this->GetSize());
 		//check( LockBuffer == NULL );	// Only one outstanding lock is allowed at a time!
 		check( !bIsLocked );	// Only one outstanding lock is allowed at a time!
@@ -640,7 +634,6 @@ public:
 	uint8 *LockWriteOnlyUnsynchronized(uint32 InOffset, uint32 InSize, bool bDiscard)
 	{
 		//SCOPE_CYCLE_COUNTER_DETAILED(STAT_OpenGLMapBufferTime);
-		check( (FOpenGL::SupportsVertexAttribBinding() && OpenGLConsoleVariables::bUseVAB) || ( this->GetUsage() & BUF_ZeroStride ) == 0 );
 		check(InOffset + InSize <= this->GetSize());
 		//check( LockBuffer == NULL );	// Only one outstanding lock is allowed at a time!
 		check( !bIsLocked );	// Only one outstanding lock is allowed at a time!
@@ -708,7 +701,6 @@ public:
 	void Unlock()
 	{
 		//SCOPE_CYCLE_COUNTER_DETAILED(STAT_OpenGLUnmapBufferTime);
-		check( (FOpenGL::SupportsVertexAttribBinding() && OpenGLConsoleVariables::bUseVAB) || ( this->GetUsage() & BUF_ZeroStride ) == 0 );
 		VERIFY_GL_SCOPE();
 		if (bIsLocked)
 		{
@@ -781,7 +773,6 @@ public:
 
 	void Update(void *InData, uint32 InOffset, uint32 InSize, bool bDiscard)
 	{
-		check( (FOpenGL::SupportsVertexAttribBinding() && OpenGLConsoleVariables::bUseVAB) || ( this->GetUsage() & BUF_ZeroStride ) == 0 );
 		check(InOffset + InSize <= this->GetSize());
 		VERIFY_GL_SCOPE();
 		Bind();
@@ -870,43 +861,21 @@ private:
 class FOpenGLBaseBuffer : public FRHIBuffer
 {
 public:
-	FOpenGLBaseBuffer() : ZeroStrideBuffer(0)
+	FOpenGLBaseBuffer()
 	{}
 
-	FOpenGLBaseBuffer(uint32 InStride, uint32 InSize, uint32 InUsage): FRHIBuffer(InSize, InUsage, InStride), ZeroStrideBuffer(0)
+	FOpenGLBaseBuffer(uint32 InStride, uint32 InSize, uint32 InUsage): FRHIBuffer(InSize, InUsage, InStride)
 	{
-		if(!(FOpenGL::SupportsVertexAttribBinding() && OpenGLConsoleVariables::bUseVAB) && (InUsage & BUF_ZeroStride))
-		{
-			ZeroStrideBuffer = FMemory::Malloc( InSize );
-		}
-
 		LLM_SCOPED_PAUSE_TRACKING_WITH_ENUM_AND_AMOUNT(ELLMTag::GraphicsPlatform, InSize, ELLMTracker::Platform, ELLMAllocType::None);
 		LLM_SCOPED_PAUSE_TRACKING_WITH_ENUM_AND_AMOUNT(ELLMTag::Meshes, InSize, ELLMTracker::Default, ELLMAllocType::None);
 	}
 
 	~FOpenGLBaseBuffer( void )
 	{
-		if( ZeroStrideBuffer )
-		{
-			FMemory::Free( ZeroStrideBuffer );
-		}
-
 		LLM_SCOPED_PAUSE_TRACKING_WITH_ENUM_AND_AMOUNT(ELLMTag::GraphicsPlatform, -(int64)GetSize(), ELLMTracker::Platform, ELLMAllocType::None);
 		LLM_SCOPED_PAUSE_TRACKING_WITH_ENUM_AND_AMOUNT(ELLMTag::Meshes, -(int64)GetSize(), ELLMTracker::Default, ELLMAllocType::None);
 	}
 
-	void* GetZeroStrideBuffer( void )
-	{
-		check( ZeroStrideBuffer );
-		return ZeroStrideBuffer;
-	}
-
-	void Swap(FOpenGLBaseBuffer& Other)
-	{
-		FRHIBuffer::Swap(Other);
-		::Swap(ZeroStrideBuffer, Other.ZeroStrideBuffer);
-	}
-
 	static bool OnDelete(GLuint Resource,uint32 Size,bool bStreamDraw,uint32 Offset)
 	{
 		OnBufferDeletion(Resource);
@@ -923,8 +892,6 @@ public:
 		// @todo-mobile
 	}
 
-private:
-	void*	ZeroStrideBuffer;
 };
 
 struct FOpenGLEUniformBufferData : public FRefCountedObject
diff --git a/Engine/Source/Runtime/OpenGLDrv/Public/OpenGLState.h b/Engine/Source/Runtime/OpenGLDrv/Public/OpenGLState.h
--- a/Engine/Source/Runtime/OpenGLDrv/Public/OpenGLState.h
+++ b/Engine/Source/Runtime/OpenGLDrv/Public/OpenGLState.h
@@ -241,11 +241,7 @@ struct FUAVStage
 
 struct FOpenGLCachedAttr
 {
-	void* Pointer;
-	GLsizei Stride;
-	GLuint Buffer;
 	GLuint Size;
-	GLuint Divisor;
 	GLenum Type;
 	GLuint StreamOffset;
 	GLuint StreamIndex;
@@ -253,11 +249,7 @@ struct FOpenGLCachedAttr
 	GLboolean bShouldConvertToFloat;
 
 	FOpenGLCachedAttr() : 
-		Pointer(FOpenGLCachedAttr_Invalid), 
-		Stride(-1), 
-		Buffer(),
 		Size(),
-		Divisor(0xFFFFFFFF), 
 		Type(0), 
 		StreamOffset(),
 		StreamIndex(0xFFFFFFFF),
diff --git a/Engine/Source/Runtime/RHI/Private/RHI.cpp b/Engine/Source/Runtime/RHI/Private/RHI.cpp
--- a/Engine/Source/Runtime/RHI/Private/RHI.cpp
+++ b/Engine/Source/Runtime/RHI/Private/RHI.cpp
@@ -189,7 +189,6 @@ FString GetBufferUsageFlagsName(EBufferUsageFlags BufferUsage)
 		case BUF_DrawIndirect:			return TEXT("BUF_DrawIndirect");
 		case BUF_ShaderResource:		return TEXT("BUF_ShaderResource");
 		case BUF_KeepCPUAccessible:		return TEXT("BUF_KeepCPUAccessible");
-		case BUF_ZeroStride:			return TEXT("BUF_ZeroStride");
 		case BUF_FastVRAM:				return TEXT("BUF_FastVRAM");
 		case BUF_Transient:				return TEXT("BUF_Transient");
 		case BUF_Shared:				return TEXT("BUF_Shared");
diff --git a/Engine/Source/Runtime/RHI/Public/RHIDefinitions.h b/Engine/Source/Runtime/RHI/Public/RHIDefinitions.h
--- a/Engine/Source/Runtime/RHI/Public/RHIDefinitions.h
+++ b/Engine/Source/Runtime/RHI/Public/RHIDefinitions.h
@@ -1426,11 +1426,8 @@ enum EBufferUsageFlags
 	 */
 	BUF_KeepCPUAccessible		= 0x0400,
 
-	/**
-	 * Provide information that this buffer will contain only one vertex, which should be delivered to every primitive drawn.
-	 * This is necessary for OpenGL implementations, which need to handle this case very differently (and can't handle GL_HALF_FLOAT in such vertices at all).
-	 */
-	BUF_ZeroStride				= 0x0800,
+	// Unused
+	//BUF_ZeroStride			= 0x0800,
 
 	/** Buffer should go in fast vram (hint only). Requires BUF_Transient */
 	BUF_FastVRAM				= 0x1000,
diff --git a/Engine/Source/Runtime/RenderCore/Public/RenderResource.h b/Engine/Source/Runtime/RenderCore/Public/RenderResource.h
--- a/Engine/Source/Runtime/RenderCore/Public/RenderResource.h
+++ b/Engine/Source/Runtime/RenderCore/Public/RenderResource.h
@@ -543,7 +543,7 @@ public:
 		// create a static vertex buffer
 		FRHIResourceCreateInfo CreateInfo(TEXT("FNullColorVertexBuffer"));
 
-		VertexBufferRHI = RHICreateBuffer(sizeof(uint32) * 4, BUF_Static | BUF_ZeroStride | BUF_VertexBuffer | BUF_ShaderResource, 0, ERHIAccess::VertexOrIndexBuffer | ERHIAccess::SRVMask, CreateInfo);
+		VertexBufferRHI = RHICreateBuffer(sizeof(uint32) * 4, BUF_Static | BUF_VertexBuffer | BUF_ShaderResource, 0, ERHIAccess::VertexOrIndexBuffer | ERHIAccess::SRVMask, CreateInfo);
 		uint32* Vertices = (uint32*)RHILockBuffer(VertexBufferRHI, 0, sizeof(uint32) * 4, RLM_WriteOnly);
 		Vertices[0] = FColor(255, 255, 255, 255).DWColor();
 		Vertices[1] = FColor(255, 255, 255, 255).DWColor();
@@ -578,7 +578,7 @@ public:
 	{
 		// create a static vertex buffer
 		FRHIResourceCreateInfo CreateInfo(TEXT("FNullVertexBuffer"));
-		VertexBufferRHI = RHICreateBuffer(sizeof(FVector3f), BUF_Static | BUF_ZeroStride | BUF_VertexBuffer | BUF_ShaderResource, 0, ERHIAccess::VertexOrIndexBuffer | ERHIAccess::SRVMask, CreateInfo);
+		VertexBufferRHI = RHICreateBuffer(sizeof(FVector3f), BUF_Static | BUF_VertexBuffer | BUF_ShaderResource, 0, ERHIAccess::VertexOrIndexBuffer | ERHIAccess::SRVMask, CreateInfo);
 		FVector3f* LockedData = (FVector3f*)RHILockBuffer(VertexBufferRHI, 0, sizeof(FVector3f), RLM_WriteOnly);
 		*LockedData = FVector3f(0.0f);
 		RHIUnlockBuffer(VertexBufferRHI);
