Made ES 3.2 a requirement for OpenGL on Android

diff --git a/Engine/Source/Programs/UnrealBuildTool/Platform/Android/AndroidToolChain.cs b/Engine/Source/Programs/UnrealBuildTool/Platform/Android/AndroidToolChain.cs
--- a/Engine/Source/Programs/UnrealBuildTool/Platform/Android/AndroidToolChain.cs
+++ b/Engine/Source/Programs/UnrealBuildTool/Platform/Android/AndroidToolChain.cs
@@ -397,7 +397,7 @@ namespace UnrealBuildTool
 
 			if (bBuildForES31)
 			{
-				GLESversion = "0x00030001";
+				GLESversion = "0x00030002";
 			}
 
 			return GLESversion;
diff --git a/Engine/Source/Runtime/OpenGLDrv/Private/Android/AndroidOpenGL.cpp b/Engine/Source/Runtime/OpenGLDrv/Private/Android/AndroidOpenGL.cpp
--- a/Engine/Source/Runtime/OpenGLDrv/Private/Android/AndroidOpenGL.cpp
+++ b/Engine/Source/Runtime/OpenGLDrv/Private/Android/AndroidOpenGL.cpp
@@ -243,7 +243,7 @@ bool PlatformInitOpenGL()
 		FAndroidOpenGL::GLMajorVerion = FCString::Atoi(*VersionString);
 		FAndroidOpenGL::GLMinorVersion = FCString::Atoi(*SubVersionString);
 
-		bool bES31Supported = FAndroidOpenGL::GLMajorVerion == 3 && FAndroidOpenGL::GLMinorVersion >= 1;
+		bool bES32Supported = FAndroidOpenGL::GLMajorVerion == 3 && FAndroidOpenGL::GLMinorVersion >= 2;
 		static const auto CVarDisableES31 = IConsoleManager::Get().FindTConsoleVariableDataInt(TEXT("r.Android.DisableOpenGLES31Support"));
 
 		bool bBuildForES31 = false;
@@ -251,7 +251,7 @@ bool PlatformInitOpenGL()
 
 		const bool bSupportsFloatingPointRTs = FAndroidMisc::SupportsFloatingPointRenderTargets();
 
-		if (bBuildForES31 && bES31Supported)
+		if (bBuildForES31 && bES32Supported)
 		{
 			FOpenGLES::CurrentFeatureLevelSupport = FAndroidOpenGL::GLMinorVersion >= 2 ? FOpenGLES::EFeatureLevelSupport::ES32 : FOpenGLES::EFeatureLevelSupport::ES31;
 			UE_LOG(LogRHI, Log, TEXT("App is packaged for OpenGL ES 3.1 and an ES %d.%d-capable device was detected."), FAndroidOpenGL::GLMajorVerion, FAndroidOpenGL::GLMinorVersion);
@@ -260,7 +260,7 @@ bool PlatformInitOpenGL()
 		{
 			FString Message = TEXT("");
 
-			if (bES31Supported)
+			if (bES32Supported)
 			{
 				Message.Append(TEXT("This device does not support Vulkan but the app was not packaged with ES 3.1 support."));
 				if (FAndroidMisc::GetAndroidBuildVersion() < 26)
@@ -272,7 +272,7 @@ bool PlatformInitOpenGL()
 			}
 			else
 			{
-				Message.Append(TEXT("This device only supports OpenGL ES 2/3 which is not supported, only supports ES 3.1+ "));
+				Message.Append(TEXT("This device only supports OpenGL ES 2/3/3.1 which is not supported, only supports ES 3.2+ "));
 				FPlatformMisc::LowLevelOutputDebugString(*Message);
 				FAndroidMisc::MessageBoxExt(EAppMsgType::Ok, *Message, TEXT("Unable to run on this device!"));
 			}
