Android EGL: Bind newly created shared surface on game thread after resume

After resume, we create new surfaces for the game and render threads (shared and rendering context surfaces). However, we only immediately bind the rendering thread context.
TECHNICALLY, we should be safe to keep the previous surface handy (EGL spec says old surface is pending deletion once made un-current). I'd be surprised if drivers actually like this.

This was causing a problem with ARCore where they used our game thread context (shared context) in order to create their own internal context (after resume). Once we call eglMakeCurrent with the shared context and the newly created shared surface, ARCore happily survives pause/resume.

diff --git a/Engine/Source/Runtime/OpenGLDrv/Private/Android/AndroidEGL.cpp b/Engine/Source/Runtime/OpenGLDrv/Private/Android/AndroidEGL.cpp
--- a/Engine/Source/Runtime/OpenGLDrv/Private/Android/AndroidEGL.cpp
+++ b/Engine/Source/Runtime/OpenGLDrv/Private/Android/AndroidEGL.cpp
@@ -1174,6 +1174,15 @@ void BlockOnLostWindowRenderCommand(TSharedPtr<FEvent, ESPMode::ThreadSafe> RTBl
 	UE_LOG(LogAndroid, Log, TEXT("RendererBlock released window lock"));
 }
 
+void SetSharedContextGameCommand(TSharedPtr<FEvent, ESPMode::ThreadSafe> GTBlockedTrigger)
+{
+	check(IsInGameThread());
+	AndroidEGL* EGL = AndroidEGL::GetInstance();
+	EGL->SetCurrentContext(EGL_NO_CONTEXT, EGL_NO_SURFACE);
+	EGL->SetCurrentSharedContext();
+	GTBlockedTrigger->Trigger();
+}
+
 extern bool IsInAndroidEventThread();
 void BlockRendering()
 {
@@ -1184,20 +1193,27 @@ void BlockRendering()
 	
 	// Wait for GC to complete and prevent further GCs
 	FGCScopeGuard GCGuard;
-	TSharedPtr<FEvent, ESPMode::ThreadSafe> RTBlockedTrigger = MakeShareable(FPlatformProcess::GetSynchEventFromPool(), [](FEvent* EventToDelete)
+	TSharedPtr<FEvent, ESPMode::ThreadSafe> BlockedTrigger = MakeShareable(FPlatformProcess::GetSynchEventFromPool(), [](FEvent* EventToDelete)
 	{
 		FPlatformProcess::ReturnSynchEventToPool(EventToDelete);
 	});
 
-	FGraphEventRef RTBlockTask = FFunctionGraphTask::CreateAndDispatchWhenReady([RTBlockedTrigger]()
+	FGraphEventRef RTBlockTask = FFunctionGraphTask::CreateAndDispatchWhenReady([BlockedTrigger]()
 	{
-		FTaskTagScope TagScope(ETaskTag::ERenderingThread);
-		BlockOnLostWindowRenderCommand(RTBlockedTrigger);
+		BlockOnLostWindowRenderCommand(BlockedTrigger);
 	}, TStatId(), NULL, ENamedThreads::GetRenderThread());
 
 	// wait for the render thread to process.
 	UE_LOG(LogAndroid, Log, TEXT("Waiting for renderer to encounter blocking command."));
-	RTBlockedTrigger->Wait();
+	BlockedTrigger->Wait();
+
+	FGraphEventRef GTBlockTask = FFunctionGraphTask::CreateAndDispatchWhenReady([BlockedTrigger]()
+	{
+		SetSharedContextGameCommand(BlockedTrigger);
+	}, TStatId(), NULL, ENamedThreads::GameThread);
+
+	UE_LOG(LogAndroid, Log, TEXT("Waiting for game thread to encounter blocking command."));
+	BlockedTrigger->Wait();
 }
 
 #endif
