Enqueue FrameSyncEvent ReleaseResource on the RT when in the RHIT in OpenGLViewport

diff --git a/Engine/Source/Runtime/OpenGLDrv/Private/OpenGLViewport.cpp b/Engine/Source/Runtime/OpenGLDrv/Private/OpenGLViewport.cpp
--- a/Engine/Source/Runtime/OpenGLDrv/Private/OpenGLViewport.cpp
+++ b/Engine/Source/Runtime/OpenGLDrv/Private/OpenGLViewport.cpp
@@ -253,7 +253,7 @@ FOpenGLViewport::~FOpenGLViewport()
 		PlatformRestoreDesktopDisplayMode();
 	}
 
-	FrameSyncEvent.ReleaseResource();
+	BeginReleaseResource(&FrameSyncEvent);
 
 	// Release back buffer, before OpenGL context becomes invalid, making it impossible
 	BackBuffer.SafeRelease();
