Removing deprecated RHI functions

diff --git a/Engine/Source/Runtime/Apple/AGXRHI/Private/AGXDynamicRHI.h b/Engine/Source/Runtime/Apple/AGXRHI/Private/AGXDynamicRHI.h
--- a/Engine/Source/Runtime/Apple/AGXRHI/Private/AGXDynamicRHI.h
+++ b/Engine/Source/Runtime/Apple/AGXRHI/Private/AGXDynamicRHI.h
@@ -167,10 +167,6 @@ public:
 	virtual void* RHILockStagingBuffer(FRHIStagingBuffer* StagingBuffer, FRHIGPUFence* Fence, uint32 Offset, uint32 SizeRHI) final override;
 	virtual void RHIUnlockStagingBuffer(FRHIStagingBuffer* StagingBuffer) final override;
 	
-	/* Copy the source box pixels in the destination box texture, return true if implemented for the current platform*/
-	virtual void RHICopySubTextureRegion_RenderThread(class FRHICommandListImmediate& RHICmdList, FRHITexture2D* SourceTexture, FRHITexture2D* DestinationTexture, FBox2D SourceBox, FBox2D DestinationBox) final override;
-	virtual void RHICopySubTextureRegion(FRHITexture2D* SourceTexture, FRHITexture2D* DestinationTexture, FBox2D SourceBox, FBox2D DestinationBox) final override;
-
 	virtual FVertexShaderRHIRef CreateVertexShader_RenderThread(class FRHICommandListImmediate& RHICmdList, TArrayView<const uint8> Code, const FSHAHash& Hash) final override;
 	virtual FPixelShaderRHIRef CreatePixelShader_RenderThread(class FRHICommandListImmediate& RHICmdList, TArrayView<const uint8> Code, const FSHAHash& Hash) final override;
 	virtual FGeometryShaderRHIRef CreateGeometryShader_RenderThread(class FRHICommandListImmediate& RHICmdList, TArrayView<const uint8> Code, const FSHAHash& Hash) final override;
diff --git a/Engine/Source/Runtime/Apple/AGXRHI/Private/AGXTexture.cpp b/Engine/Source/Runtime/Apple/AGXRHI/Private/AGXTexture.cpp
--- a/Engine/Source/Runtime/Apple/AGXRHI/Private/AGXTexture.cpp
+++ b/Engine/Source/Runtime/Apple/AGXRHI/Private/AGXTexture.cpp
@@ -2370,124 +2370,6 @@ void FAGXDynamicRHI::RHIVirtualTextureSetFirstMipVisible(FRHITexture2D* TextureR
 	NOT_SUPPORTED("RHIVirtualTextureSetFirstMipVisible");
 }
 
-struct FRHICopySubTextureRegion final : public FRHICommand<FRHICopySubTextureRegion>
-{
-	FRHITexture2D* SourceTexture;
-	FRHITexture2D* DestinationTexture;
-	FBox2D SourceBox;
-	FBox2D DestinationBox;
-	
-	FORCEINLINE_DEBUGGABLE FRHICopySubTextureRegion(FRHITexture2D* InSourceTexture, FRHITexture2D* InDestinationTexture, FBox2D InSourceBox, FBox2D InDestinationBox)
-	: SourceTexture(InSourceTexture)
-	, DestinationTexture(InDestinationTexture)
-	, SourceBox(InSourceBox)
-	, DestinationBox(InDestinationBox)
-	{
-	}
-	
-	void Execute(FRHICommandListBase& CmdList)
-	{
-		GDynamicRHI->RHICopySubTextureRegion(SourceTexture, DestinationTexture, SourceBox, DestinationBox);
-	}
-};
-
-void FAGXDynamicRHI::RHICopySubTextureRegion_RenderThread(class FRHICommandListImmediate& RHICmdList, FRHITexture2D* SourceTexture, FRHITexture2D* DestinationTexture, FBox2D SourceBox, FBox2D DestinationBox)
-{
-	@autoreleasepool
-	{
-		if (RHICmdList.Bypass() || !IsRunningRHIInSeparateThread())
-		{
-			GDynamicRHI->RHICopySubTextureRegion(SourceTexture, DestinationTexture, SourceBox, DestinationBox);
-		}
-		else
-		{
-			new (RHICmdList.AllocCommand<FRHICopySubTextureRegion>()) FRHICopySubTextureRegion(SourceTexture, DestinationTexture, SourceBox, DestinationBox);
-		}
-	}
-}
-
-void FAGXDynamicRHI::RHICopySubTextureRegion(FRHITexture2D* SourceTexture, FRHITexture2D* DestinationTexture, FBox2D SourceBox, FBox2D DestinationBox)
-{
-	@autoreleasepool {
-		check(SourceTexture);
-		check(DestinationTexture);
-		
-		FAGXTexture2D* MetalSrcTexture = ResourceCast(SourceTexture);
-		FAGXTexture2D* MetalDestTexture = ResourceCast(DestinationTexture);
-		
-		if(SourceTexture->GetFormat() == DestinationTexture->GetFormat())
-		{
-			FVector2D SourceSizeVector = SourceBox.GetSize();
-			FVector2D DestinatioSizeVector = DestinationBox.GetSize();
-			
-			mtlpp::Origin SourceOrigin(SourceBox.Min.X, SourceBox.Min.Y, 0);
-			mtlpp::Size SourceSize(SourceSizeVector.X, SourceSizeVector.Y, 1);
-			
-			mtlpp::Origin DestinationOrigin(DestinationBox.Min.X, DestinationBox.Min.Y, 0);
-			mtlpp::Size DestinationSize(DestinatioSizeVector.X, DestinatioSizeVector.Y, 1);
-			
-			check(DestinationSize.width == SourceSize.width);
-			check(DestinationSize.height == SourceSize.height);
-			
-			FAGXTexture SrcTexture;
-			mtlpp::TextureUsage Usage = MetalSrcTexture->Surface.Texture.GetUsage();
-			if(Usage & mtlpp::TextureUsage::PixelFormatView)
-			{
-				ns::Range Slices(0, MetalSrcTexture->Surface.Texture.GetArrayLength() * (MetalSrcTexture->Surface.bIsCubemap ? 6 : 1));
-				if(MetalSrcTexture->Surface.Texture.GetPixelFormat() != MetalDestTexture->Surface.Texture.GetPixelFormat())
-				{
-					SrcTexture = MetalSrcTexture->Surface.Texture.NewTextureView(MetalDestTexture->Surface.Texture.GetPixelFormat(), MetalSrcTexture->Surface.Texture.GetTextureType(), ns::Range(0, MetalSrcTexture->Surface.Texture.GetMipmapLevelCount()), Slices);
-				}
-			}
-			if (!SrcTexture)
-			{
-				SrcTexture = MetalSrcTexture->Surface.Texture;
-			}
-			
-			// Account for create with TexCreate_SRGB flag which could make these different
-			if(SrcTexture.GetPixelFormat() == MetalDestTexture->Surface.Texture.GetPixelFormat())
-			{
-				ImmediateContext.GetInternalContext().CopyFromTextureToTexture(SrcTexture, 0, 0, SourceOrigin,SourceSize,MetalDestTexture->Surface.Texture, 0, 0, DestinationOrigin);
-			}
-			else
-			{
-				// Linear and sRGB mismatch then try to go via metal buffer
-				// Modified clone of logic from MetalRenderTarget.cpp
-				uint32 BytesPerPixel = (MetalSrcTexture->Surface.PixelFormat != PF_DepthStencil) ? GPixelFormats[MetalSrcTexture->Surface.PixelFormat].BlockBytes : 1;
-				const uint32 Stride = BytesPerPixel * SourceSize.width;
-				const uint32 Alignment = PLATFORM_MAC ? 1u : 64u;
-				const uint32 AlignedStride = ((Stride - 1) & ~(Alignment - 1)) + Alignment;
-				const uint32 BytesPerImage = AlignedStride *  SourceSize.height;
-				
-				FAGXBuffer Buffer = GetAGXDeviceContext().CreatePooledBuffer(FAGXPooledBufferArgs(BytesPerImage, BUF_Dynamic, mtlpp::StorageMode::Shared));
-				
-				check(Buffer);
-				
-				mtlpp::BlitOption Options = mtlpp::BlitOption::None;
-#if !PLATFORM_MAC
-				if (MetalSrcTexture->Surface.Texture.GetPixelFormat() >= mtlpp::PixelFormat::PVRTC_RGB_2BPP && MetalSrcTexture->Surface.Texture.GetPixelFormat() <= mtlpp::PixelFormat::PVRTC_RGBA_4BPP_sRGB)
-				{
-					Options = mtlpp::BlitOption::RowLinearPVRTC;
-				}
-#endif
-				ImmediateContext.GetInternalContext().CopyFromTextureToBuffer(MetalSrcTexture->Surface.Texture, 0, 0, SourceOrigin, SourceSize, Buffer, 0, AlignedStride, BytesPerImage, Options);
-				ImmediateContext.GetInternalContext().CopyFromBufferToTexture(Buffer, 0, Stride, BytesPerImage, SourceSize, MetalDestTexture->Surface.Texture, 0, 0, DestinationOrigin, Options);
-				
-				GetAGXDeviceContext().ReleaseBuffer(Buffer);
-			}
-			
-			if (SrcTexture != MetalSrcTexture->Surface.Texture)
-			{
-				AGXSafeReleaseMetalTexture(SrcTexture);
-			}
-		}
-		else
-		{
-			UE_LOG(LogAGX, Error, TEXT("RHICopySubTextureRegion Source (UnrealEngine %d: MTL %d) <-> Destination (UnrealEngine %d: MTL %d) texture format mismatch"), (uint32)SourceTexture->GetFormat(), (uint32)MetalSrcTexture->Surface.Texture.GetPixelFormat(), (uint32)DestinationTexture->GetFormat(), (uint32)MetalDestTexture->Surface.Texture.GetPixelFormat());
-		}
-	}
-}
-
 inline bool AGXRHICopyTexutre_IsTextureFormatCompatible(EPixelFormat SrcFmt, EPixelFormat DstFmt)
 {
 	//
diff --git a/Engine/Source/Runtime/Apple/MetalRHI/Private/MetalDynamicRHI.h b/Engine/Source/Runtime/Apple/MetalRHI/Private/MetalDynamicRHI.h
--- a/Engine/Source/Runtime/Apple/MetalRHI/Private/MetalDynamicRHI.h
+++ b/Engine/Source/Runtime/Apple/MetalRHI/Private/MetalDynamicRHI.h
@@ -170,10 +170,6 @@ public:
 	virtual void* RHILockStagingBuffer(FRHIStagingBuffer* StagingBuffer, FRHIGPUFence* Fence, uint32 Offset, uint32 SizeRHI) final override;
 	virtual void RHIUnlockStagingBuffer(FRHIStagingBuffer* StagingBuffer) final override;
 	
-	/* Copy the source box pixels in the destination box texture, return true if implemented for the current platform*/
-	virtual void RHICopySubTextureRegion_RenderThread(class FRHICommandListImmediate& RHICmdList, FRHITexture2D* SourceTexture, FRHITexture2D* DestinationTexture, FBox2D SourceBox, FBox2D DestinationBox) final override;
-	virtual void RHICopySubTextureRegion(FRHITexture2D* SourceTexture, FRHITexture2D* DestinationTexture, FBox2D SourceBox, FBox2D DestinationBox) final override;
-
 	virtual FVertexShaderRHIRef CreateVertexShader_RenderThread(class FRHICommandListImmediate& RHICmdList, TArrayView<const uint8> Code, const FSHAHash& Hash) final override;
 	virtual FPixelShaderRHIRef CreatePixelShader_RenderThread(class FRHICommandListImmediate& RHICmdList, TArrayView<const uint8> Code, const FSHAHash& Hash) final override;
 	virtual FGeometryShaderRHIRef CreateGeometryShader_RenderThread(class FRHICommandListImmediate& RHICmdList, TArrayView<const uint8> Code, const FSHAHash& Hash) final override;
diff --git a/Engine/Source/Runtime/Apple/MetalRHI/Private/MetalTexture.cpp b/Engine/Source/Runtime/Apple/MetalRHI/Private/MetalTexture.cpp
--- a/Engine/Source/Runtime/Apple/MetalRHI/Private/MetalTexture.cpp
+++ b/Engine/Source/Runtime/Apple/MetalRHI/Private/MetalTexture.cpp
@@ -2463,124 +2463,6 @@ void FMetalDynamicRHI::RHIVirtualTextureSetFirstMipVisible(FRHITexture2D* Textur
 	NOT_SUPPORTED("RHIVirtualTextureSetFirstMipVisible");
 }
 
-struct FRHICopySubTextureRegion final : public FRHICommand<FRHICopySubTextureRegion>
-{
-	FRHITexture2D* SourceTexture;
-	FRHITexture2D* DestinationTexture;
-	FBox2D SourceBox;
-	FBox2D DestinationBox;
-	
-	FORCEINLINE_DEBUGGABLE FRHICopySubTextureRegion(FRHITexture2D* InSourceTexture, FRHITexture2D* InDestinationTexture, FBox2D InSourceBox, FBox2D InDestinationBox)
-	: SourceTexture(InSourceTexture)
-	, DestinationTexture(InDestinationTexture)
-	, SourceBox(InSourceBox)
-	, DestinationBox(InDestinationBox)
-	{
-	}
-	
-	void Execute(FRHICommandListBase& CmdList)
-	{
-		GDynamicRHI->RHICopySubTextureRegion(SourceTexture, DestinationTexture, SourceBox, DestinationBox);
-	}
-};
-
-void FMetalDynamicRHI::RHICopySubTextureRegion_RenderThread(class FRHICommandListImmediate& RHICmdList, FRHITexture2D* SourceTexture, FRHITexture2D* DestinationTexture, FBox2D SourceBox, FBox2D DestinationBox)
-{
-	@autoreleasepool
-	{
-		if (RHICmdList.Bypass() || !IsRunningRHIInSeparateThread())
-		{
-			GDynamicRHI->RHICopySubTextureRegion(SourceTexture, DestinationTexture, SourceBox, DestinationBox);
-		}
-		else
-		{
-			new (RHICmdList.AllocCommand<FRHICopySubTextureRegion>()) FRHICopySubTextureRegion(SourceTexture, DestinationTexture, SourceBox, DestinationBox);
-		}
-	}
-}
-
-void FMetalDynamicRHI::RHICopySubTextureRegion(FRHITexture2D* SourceTexture, FRHITexture2D* DestinationTexture, FBox2D SourceBox, FBox2D DestinationBox)
-{
-	@autoreleasepool {
-		check(SourceTexture);
-		check(DestinationTexture);
-		
-		FMetalTexture2D* MetalSrcTexture = ResourceCast(SourceTexture);
-		FMetalTexture2D* MetalDestTexture = ResourceCast(DestinationTexture);
-		
-		if(SourceTexture->GetFormat() == DestinationTexture->GetFormat())
-		{
-			FVector2D SourceSizeVector = SourceBox.GetSize();
-			FVector2D DestinatioSizeVector = DestinationBox.GetSize();
-			
-			mtlpp::Origin SourceOrigin(SourceBox.Min.X, SourceBox.Min.Y, 0);
-			mtlpp::Size SourceSize(SourceSizeVector.X, SourceSizeVector.Y, 1);
-			
-			mtlpp::Origin DestinationOrigin(DestinationBox.Min.X, DestinationBox.Min.Y, 0);
-			mtlpp::Size DestinationSize(DestinatioSizeVector.X, DestinatioSizeVector.Y, 1);
-			
-			check(DestinationSize.width == SourceSize.width);
-			check(DestinationSize.height == SourceSize.height);
-			
-			FMetalTexture SrcTexture;
-			mtlpp::TextureUsage Usage = MetalSrcTexture->Surface.Texture.GetUsage();
-			if(Usage & mtlpp::TextureUsage::PixelFormatView)
-			{
-				ns::Range Slices(0, MetalSrcTexture->Surface.Texture.GetArrayLength() * (MetalSrcTexture->Surface.bIsCubemap ? 6 : 1));
-				if(MetalSrcTexture->Surface.Texture.GetPixelFormat() != MetalDestTexture->Surface.Texture.GetPixelFormat())
-				{
-					SrcTexture = MetalSrcTexture->Surface.Texture.NewTextureView(MetalDestTexture->Surface.Texture.GetPixelFormat(), MetalSrcTexture->Surface.Texture.GetTextureType(), ns::Range(0, MetalSrcTexture->Surface.Texture.GetMipmapLevelCount()), Slices);
-				}
-			}
-			if (!SrcTexture)
-			{
-				SrcTexture = MetalSrcTexture->Surface.Texture;
-			}
-			
-			// Account for create with TexCreate_SRGB flag which could make these different
-			if(SrcTexture.GetPixelFormat() == MetalDestTexture->Surface.Texture.GetPixelFormat())
-			{
-				ImmediateContext.GetInternalContext().CopyFromTextureToTexture(SrcTexture, 0, 0, SourceOrigin,SourceSize,MetalDestTexture->Surface.Texture, 0, 0, DestinationOrigin);
-			}
-			else
-			{
-				// Linear and sRGB mismatch then try to go via metal buffer
-				// Modified clone of logic from MetalRenderTarget.cpp
-				uint32 BytesPerPixel = (MetalSrcTexture->Surface.PixelFormat != PF_DepthStencil) ? GPixelFormats[MetalSrcTexture->Surface.PixelFormat].BlockBytes : 1;
-				const uint32 Stride = BytesPerPixel * SourceSize.width;
-				const uint32 Alignment = PLATFORM_MAC ? 1u : 64u;
-				const uint32 AlignedStride = ((Stride - 1) & ~(Alignment - 1)) + Alignment;
-				const uint32 BytesPerImage = AlignedStride *  SourceSize.height;
-				
-				FMetalBuffer Buffer = GetMetalDeviceContext().CreatePooledBuffer(FMetalPooledBufferArgs(ImmediateContext.Context->GetDevice(), BytesPerImage, BUF_Dynamic, mtlpp::StorageMode::Shared));
-				
-				check(Buffer);
-				
-				mtlpp::BlitOption Options = mtlpp::BlitOption::None;
-#if !PLATFORM_MAC
-				if (MetalSrcTexture->Surface.Texture.GetPixelFormat() >= mtlpp::PixelFormat::PVRTC_RGB_2BPP && MetalSrcTexture->Surface.Texture.GetPixelFormat() <= mtlpp::PixelFormat::PVRTC_RGBA_4BPP_sRGB)
-				{
-					Options = mtlpp::BlitOption::RowLinearPVRTC;
-				}
-#endif
-				ImmediateContext.GetInternalContext().CopyFromTextureToBuffer(MetalSrcTexture->Surface.Texture, 0, 0, SourceOrigin, SourceSize, Buffer, 0, AlignedStride, BytesPerImage, Options);
-				ImmediateContext.GetInternalContext().CopyFromBufferToTexture(Buffer, 0, Stride, BytesPerImage, SourceSize, MetalDestTexture->Surface.Texture, 0, 0, DestinationOrigin, Options);
-				
-				GetMetalDeviceContext().ReleaseBuffer(Buffer);
-			}
-			
-			if (SrcTexture != MetalSrcTexture->Surface.Texture)
-			{
-				SafeReleaseMetalTexture(SrcTexture);
-			}
-		}
-		else
-		{
-			UE_LOG(LogMetal, Error, TEXT("RHICopySubTextureRegion Source (UnrealEngine %d: MTL %d) <-> Destination (UnrealEngine %d: MTL %d) texture format mismatch"), (uint32)SourceTexture->GetFormat(), (uint32)MetalSrcTexture->Surface.Texture.GetPixelFormat(), (uint32)DestinationTexture->GetFormat(), (uint32)MetalDestTexture->Surface.Texture.GetPixelFormat());
-		}
-	}
-}
-
 inline bool MetalRHICopyTexutre_IsTextureFormatCompatible(EPixelFormat SrcFmt, EPixelFormat DstFmt)
 {
 	//
diff --git a/Engine/Source/Runtime/D3D12RHI/Private/D3D12Buffer.cpp b/Engine/Source/Runtime/D3D12RHI/Private/D3D12Buffer.cpp
--- a/Engine/Source/Runtime/D3D12RHI/Private/D3D12Buffer.cpp
+++ b/Engine/Source/Runtime/D3D12RHI/Private/D3D12Buffer.cpp
@@ -943,145 +943,3 @@ void FD3D12CommandContext::RHICopyBufferRegion(FRHIBuffer* DestBufferRHI, uint64
 
 	Device->RegisterGPUWork(1);
 }
-
-#if D3D12_RHI_RAYTRACING
-void FD3D12CommandContext::RHICopyBufferRegions(const TArrayView<const FCopyBufferRegionParams> Params)
-{
-	// Batched buffer copy finds unique source and destination buffer resources, performs transitions
-	// to copy source / dest state, then performs copies and finally restores original state.
-
-	using FLocalResourceArray = TArray<FD3D12Resource*, TInlineAllocator<16, TMemStackAllocator<>>>;
-	FLocalResourceArray SrcBuffers;
-	FLocalResourceArray DstBuffers;
-
-	SrcBuffers.Reserve(Params.Num());
-	DstBuffers.Reserve(Params.Num());
-
-	// Transition buffers to copy states
-	for (auto& Param : Params)
-	{
-		FD3D12Buffer* SourceBuffer = RetrieveObject<FD3D12Buffer>(Param.SourceBuffer);
-		FD3D12Buffer* DestBuffer = RetrieveObject<FD3D12Buffer>(Param.DestBuffer);
-		check(SourceBuffer);
-		check(DestBuffer);
-
-		FD3D12Device* Device = SourceBuffer->GetParentDevice();
-		check(Device == DestBuffer->GetParentDevice());
-		check(Device == GetParentDevice());
-
-		FD3D12Resource* pSourceResource = SourceBuffer->ResourceLocation.GetResource();
-		FD3D12Resource* pDestResource = DestBuffer->ResourceLocation.GetResource();
-
-		checkf(pSourceResource != pDestResource, TEXT("CopyBufferRegion cannot be used on the same resource. This can happen when both the source and the dest are suballocated from the same resource."));
-
-		SrcBuffers.Add(pSourceResource);
-		DstBuffers.Add(pDestResource);
-	}
-
-	Algo::Sort(SrcBuffers);
-	Algo::Sort(DstBuffers);
-
-	enum class EBatchCopyState
-	{
-		CopySource,
-		CopyDest,
-		FinalizeSource,
-		FinalizeDest,
-	};
-
-	auto TransitionResources = [](FD3D12CommandListHandle& InCommandListHandle, FLocalResourceArray& SortedResources, EBatchCopyState State)
-	{
-		const uint32 Subresource = 0; // Buffers only have one subresource
-
-		FD3D12Resource* PrevResource = nullptr;
-		for (FD3D12Resource* Resource : SortedResources)
-		{
-			if (Resource == PrevResource) continue; // Skip duplicate resource barriers
-
-			const bool bUseDefaultState = !Resource->RequiresResourceStateTracking();
-
-			D3D12_RESOURCE_STATES DesiredState = D3D12_RESOURCE_STATE_CORRUPT;
-			D3D12_RESOURCE_STATES CurrentState = D3D12_RESOURCE_STATE_CORRUPT;
-			switch (State)
-			{
-			case EBatchCopyState::CopySource:
-				DesiredState = D3D12_RESOURCE_STATE_COPY_SOURCE;
-				CurrentState = bUseDefaultState ? Resource->GetDefaultResourceState() : CurrentState;
-				break;
-			case EBatchCopyState::CopyDest:
-				DesiredState = D3D12_RESOURCE_STATE_COPY_DEST;
-				CurrentState = bUseDefaultState ? Resource->GetDefaultResourceState() : CurrentState;
-				break;
-			case EBatchCopyState::FinalizeSource:
-				CurrentState = D3D12_RESOURCE_STATE_COPY_SOURCE;
-				DesiredState = bUseDefaultState ? Resource->GetDefaultResourceState() : D3D12_RESOURCE_STATE_GENERIC_READ;
-				break;
-			case EBatchCopyState::FinalizeDest:
-				CurrentState = D3D12_RESOURCE_STATE_COPY_DEST;
-				DesiredState = bUseDefaultState ? Resource->GetDefaultResourceState() : D3D12_RESOURCE_STATE_GENERIC_READ;
-				break;
-			default:
-				checkf(false, TEXT("Unexpected batch copy state"));
-				break;
-			}
-
-			if (bUseDefaultState)
-			{
-				check(CurrentState != D3D12_RESOURCE_STATE_CORRUPT);
-				InCommandListHandle.AddTransitionBarrier(Resource, CurrentState, DesiredState, Subresource);
-			}
-			else
-			{
-				FD3D12DynamicRHI::TransitionResource(InCommandListHandle, Resource, D3D12_RESOURCE_STATE_TBD, DesiredState, Subresource, FD3D12DynamicRHI::ETransitionMode::Validate);
-			}
-
-			PrevResource = Resource;
-		}
-	};
-
-	// Ensure that all previously pending barriers have been processed to avoid incorrect/conflicting transitions for non-tracked resources
-	CommandListHandle.FlushResourceBarriers();
-
-	TransitionResources(CommandListHandle, SrcBuffers, EBatchCopyState::CopySource);
-	TransitionResources(CommandListHandle, DstBuffers, EBatchCopyState::CopyDest);
-
-	// Issue all copy source/dest barriers before performing actual copies
-	CommandListHandle.FlushResourceBarriers();
-
-	for (auto& Param : Params)
-	{
-		FD3D12Buffer* SourceBuffer = RetrieveObject<FD3D12Buffer>(Param.SourceBuffer);
-		FD3D12Buffer* DestBuffer = RetrieveObject<FD3D12Buffer>(Param.DestBuffer);
-		uint64 SrcOffset = Param.SrcOffset;
-		uint64 DstOffset = Param.DstOffset;
-		uint64 NumBytes = Param.NumBytes;
-
-		FD3D12Device* Device = SourceBuffer->GetParentDevice();
-		check(Device == DestBuffer->GetParentDevice());
-
-		FD3D12Resource* pSourceResource = SourceBuffer->ResourceLocation.GetResource();
-		D3D12_RESOURCE_DESC const& SourceBufferDesc = pSourceResource->GetDesc();
-
-		FD3D12Resource* pDestResource = DestBuffer->ResourceLocation.GetResource();
-		D3D12_RESOURCE_DESC const& DestBufferDesc = pDestResource->GetDesc();
-
-		check(DstOffset + NumBytes <= DestBufferDesc.Width);
-		check(SrcOffset + NumBytes <= SourceBufferDesc.Width);
-
-		numCopies++;
-
-		CommandListHandle->CopyBufferRegion(pDestResource->GetResource(), DestBuffer->ResourceLocation.GetOffsetFromBaseOfResource() + DstOffset, pSourceResource->GetResource(), SourceBuffer->ResourceLocation.GetOffsetFromBaseOfResource() + SrcOffset, NumBytes);
-		CommandListHandle.UpdateResidency(pDestResource);
-		CommandListHandle.UpdateResidency(pSourceResource);
-
-		Device->RegisterGPUWork(1);
-	}
-
-	// Transition buffers back to default readable state
-
-	TransitionResources(CommandListHandle, SrcBuffers, EBatchCopyState::FinalizeSource);
-	TransitionResources(CommandListHandle, DstBuffers, EBatchCopyState::FinalizeDest);
-	
-	ConditionalFlushCommandList();
-}
-#endif // D3D12_RHI_RAYTRACING
diff --git a/Engine/Source/Runtime/D3D12RHI/Private/D3D12CommandContext.h b/Engine/Source/Runtime/D3D12RHI/Private/D3D12CommandContext.h
--- a/Engine/Source/Runtime/D3D12RHI/Private/D3D12CommandContext.h
+++ b/Engine/Source/Runtime/D3D12RHI/Private/D3D12CommandContext.h
@@ -385,7 +385,6 @@ public:
 	virtual void RHIBroadcastTemporalEffect(const FName& InEffectName, const TArrayView<FRHIBuffer*> InBuffers) final AFR_API_OVERRIDE;
 
 #if D3D12_RHI_RAYTRACING
-	virtual void RHICopyBufferRegions(const TArrayView<const FCopyBufferRegionParams> Params) final override;
 	virtual void RHIBindAccelerationStructureMemory(FRHIRayTracingScene* Scene, FRHIBuffer* Buffer, uint32 BufferOffset) final override;
 	virtual void BuildAccelerationStructuresInternal(const TArrayView<D3D12_BUILD_RAYTRACING_ACCELERATION_STRUCTURE_DESC> BuildDesc);
 	virtual void RHIBuildAccelerationStructures(const TArrayView<const FRayTracingGeometryBuildParams> Params, const FRHIBufferRange& ScratchBufferRange) final override;
diff --git a/Engine/Source/Runtime/D3D12RHI/Private/D3D12RHIPrivate.h b/Engine/Source/Runtime/D3D12RHI/Private/D3D12RHIPrivate.h
--- a/Engine/Source/Runtime/D3D12RHI/Private/D3D12RHIPrivate.h
+++ b/Engine/Source/Runtime/D3D12RHI/Private/D3D12RHIPrivate.h
@@ -379,10 +379,6 @@ public:
 	virtual bool RHIGetRenderQueryResult(FRHIRenderQuery* RenderQuery, uint64& OutResult, bool bWait, uint32 GPUIndex = INDEX_NONE) final override;
 	virtual uint32 RHIGetViewportNextPresentGPUIndex(FRHIViewport* Viewport) final override;
 	virtual FTexture2DRHIRef RHIGetViewportBackBuffer(FRHIViewport* Viewport) final override;
-	UE_DEPRECATED(4.25, "RHIAliasTextureResources now takes references to FTextureRHIRef objects as parameters")
-	virtual void RHIAliasTextureResources(FRHITexture* DestTexture, FRHITexture* SrcTexture) final override;
-	UE_DEPRECATED(4.25, "RHICreateAliasedTexture now takes a reference to an FTextureRHIRef object")
-	virtual FTextureRHIRef RHICreateAliasedTexture(FRHITexture* SourceTexture) final override;
 	virtual void RHIAliasTextureResources(FTextureRHIRef& DestTexture, FTextureRHIRef& SrcTexture) final override;
 	virtual FTextureRHIRef RHICreateAliasedTexture(FTextureRHIRef& SourceTexture) final override;
 	virtual void RHIAdvanceFrameFence() final override;
@@ -503,8 +499,6 @@ public:
 		return RHICreateRenderQuery(QueryType);
 	}
 
-	virtual void RHICopySubTextureRegion(FRHITexture2D* SourceTexture, FRHITexture2D* DestinationTexture, FBox2D SourceBox, FBox2D DestinationBox) final override;
-
 	void RHICalibrateTimers() override;
 
 #if D3D12_RHI_RAYTRACING
diff --git a/Engine/Source/Runtime/D3D12RHI/Private/D3D12Texture.cpp b/Engine/Source/Runtime/D3D12RHI/Private/D3D12Texture.cpp
--- a/Engine/Source/Runtime/D3D12RHI/Private/D3D12Texture.cpp
+++ b/Engine/Source/Runtime/D3D12RHI/Private/D3D12Texture.cpp
@@ -127,39 +127,6 @@ struct FRHICommandUpdateTexture final : public FRHICommand<FRHICommandUpdateText
 	}
 };
 
-struct FRHICommandCopySubTextureRegionString
-{
-	static const TCHAR* TStr() { return TEXT("FRHICommandCopySubTextureRegion"); }
-};
-struct FRHICommandCopySubTextureRegion final : public FRHICommand<FRHICommandCopySubTextureRegion, FRHICommandCopySubTextureRegionString>
-{
-	FD3D12TextureBase* DestTexture;
-	uint32 DestX;
-	uint32 DestY;
-	uint32 DestZ;
-	FD3D12TextureBase* SourceTexture;
-	D3D12_BOX SourceBox;
-
-	FORCEINLINE_DEBUGGABLE FRHICommandCopySubTextureRegion(FD3D12TextureBase* InDestTexture, uint32 InDestX, uint32 InDestY, uint32 InDestZ, FD3D12TextureBase* InSourceTexture, const D3D12_BOX& InSourceBox)
-		: DestTexture(InDestTexture)
-		, DestX(InDestX)
-		, DestY(InDestY)
-		, DestZ(InDestZ)
-		, SourceTexture(InSourceTexture)
-		, SourceBox(InSourceBox)
-	{
-	}
-
-	~FRHICommandCopySubTextureRegion()
-	{
-	}
-
-	void Execute(FRHICommandListBase& CmdList)
-	{
-		DestTexture->CopyTextureRegion(DestX, DestY, DestZ, SourceTexture, SourceBox);
-	}
-};
-
 struct FD3D12RHICommandInitializeTextureString
 {
 	static const TCHAR* TStr() { return TEXT("FD3D12RHICommandInitializeTexture"); }
@@ -3300,24 +3267,6 @@ FTextureCubeRHIRef FD3D12DynamicRHI::RHICreateTextureCubeFromResource(EPixelForm
 	return CreateTextureFromResource<FD3D12BaseTextureCube>(false, true, Format, TexCreateFlags, ClearValueBinding, Resource);
 }
 
-void FD3D12DynamicRHI::RHIAliasTextureResources(FRHITexture* DestTextureRHI, FRHITexture* SrcTextureRHI)
-{
-	FD3D12TextureBase* DestTexture = GetD3D12TextureFromRHITexture(DestTextureRHI);
-	FD3D12TextureBase* SrcTexture = GetD3D12TextureFromRHITexture(SrcTextureRHI);
-
-	// This path will potentially cause crashes, if the source texture is destroyed and we're still being used. This
-	// API path will be deprecated post 4.25. To avoid issues, use the version that takes FTextureRHIRef references instead.
-	check(false);
-
-	for (FD3D12TextureBase::FDualLinkedObjectIterator It(DestTexture, SrcTexture); It; ++It)
-	{
-		DestTexture = It.GetFirst();
-		SrcTexture = It.GetSecond();
-
-		DestTexture->AliasResources(SrcTexture);
-	}
-}
-
 void FD3D12DynamicRHI::RHIAliasTextureResources(FTextureRHIRef& DestTextureRHI, FTextureRHIRef& SrcTextureRHI)
 {
 	FD3D12TextureBase* DestTexture = GetD3D12TextureFromRHITexture(DestTextureRHI);
@@ -3423,26 +3372,6 @@ TD3D12Texture2D<BaseResourceType>* FD3D12DynamicRHI::CreateAliasedD3D12Texture2D
 	return Texture2D;
 }
 
-FTextureRHIRef FD3D12DynamicRHI::RHICreateAliasedTexture(FRHITexture* SourceTextureRHI)
-{
-	FD3D12TextureBase* SourceTexture = GetD3D12TextureFromRHITexture(SourceTextureRHI);
-	if (SourceTextureRHI->GetTexture2D() != nullptr)
-	{
-		return CreateAliasedD3D12Texture2D<FD3D12BaseTexture2D>(ResourceCast(SourceTextureRHI->GetTexture2D()));
-	}
-	else if (SourceTextureRHI->GetTexture2DArray() != nullptr)
-	{
-		return CreateAliasedD3D12Texture2D<FD3D12BaseTexture2DArray>(ResourceCast(SourceTextureRHI->GetTexture2DArray()));
-	}
-	else if (SourceTextureRHI->GetTextureCube() != nullptr)
-	{
-		return CreateAliasedD3D12Texture2D<FD3D12BaseTextureCube>(ResourceCast(SourceTextureRHI->GetTextureCube()));
-	}
-
-	UE_LOG(LogD3D12RHI, Error, TEXT("Currently FD3D12DynamicRHI::RHICreateAliasedTexture only supports 2D, 2D Array and Cube textures."));
-	return nullptr;
-}
-
 FTextureRHIRef FD3D12DynamicRHI::RHICreateAliasedTexture(FTextureRHIRef& SourceTextureRHI)
 {
 	FD3D12TextureBase* SourceTexture = GetD3D12TextureFromRHITexture(SourceTextureRHI);
@@ -3472,33 +3401,6 @@ FTextureRHIRef FD3D12DynamicRHI::RHICreateAliasedTexture(FTextureRHIRef& SourceT
 	return ReturnTexture;
 }
 
-void FD3D12DynamicRHI::RHICopySubTextureRegion(FRHITexture2D* SourceTextureRHI, FRHITexture2D* DestTextureRHI, FBox2D SourceBox, FBox2D DestinationBox)
-{
-	FD3D12TextureBase* SourceTexture = GetD3D12TextureFromRHITexture(SourceTextureRHI);
-	FD3D12TextureBase* DestTexture = GetD3D12TextureFromRHITexture(DestTextureRHI);
-
-	const uint32 XOffset = (uint32)(DestinationBox.Min.X);
-	const uint32 YOffset = (uint32)(DestinationBox.Min.Y);
-	const uint32 Width = (uint32)(SourceBox.Max.X - SourceBox.Min.X);
-	const uint32 Height = (uint32)(SourceBox.Max.Y - SourceBox.Min.Y);
-
-	const CD3DX12_BOX SourceBoxD3D((LONG)SourceBox.Min.X, (LONG)SourceBox.Min.Y, (LONG)SourceBox.Max.X, (LONG)SourceBox.Max.Y);
-
-	CD3DX12_TEXTURE_COPY_LOCATION DestCopyLocation(DestTexture->GetResource()->GetResource(), 0);
-	CD3DX12_TEXTURE_COPY_LOCATION SourceCopyLocation(SourceTexture->GetResource()->GetResource(), 0);
-
-	FRHICommandListImmediate& RHICmdList = FRHICommandListExecutor::GetImmediateCommandList();
-	if (RHICmdList.Bypass())
-	{
-		FRHICommandCopySubTextureRegion RHICmd(DestTexture, XOffset, YOffset, 0, SourceTexture, SourceBoxD3D);
-		RHICmd.Execute(RHICmdList);
-	}
-	else
-	{
-		ALLOC_COMMAND_CL(RHICmdList, FRHICommandCopySubTextureRegion)(DestTexture, XOffset, YOffset, 0, SourceTexture, SourceBoxD3D);
-	}
-}
-
 void FD3D12CommandContext::RHICopyTexture(FRHITexture* SourceTextureRHI, FRHITexture* DestTextureRHI, const FRHICopyTextureInfo& CopyInfo)
 {
 	FD3D12TextureBase* SourceTexture = RetrieveTextureBase(SourceTextureRHI);
diff --git a/Engine/Source/Runtime/OpenGLDrv/Private/OpenGLTexture.cpp b/Engine/Source/Runtime/OpenGLDrv/Private/OpenGLTexture.cpp
--- a/Engine/Source/Runtime/OpenGLDrv/Private/OpenGLTexture.cpp
+++ b/Engine/Source/Runtime/OpenGLDrv/Private/OpenGLTexture.cpp
@@ -2545,42 +2545,6 @@ void FOpenGLDynamicRHI::RHIVirtualTextureSetFirstMipVisible(FRHITexture2D* Textu
 {
 }
 
-void FOpenGLDynamicRHI::RHICopySubTextureRegion(FRHITexture2D* SourceTextureRHI, FRHITexture2D* DestinationTextureRHI, FBox2D SourceBox, FBox2D DestinationBox)
-{
-	VERIFY_GL_SCOPE();
-	FOpenGLTexture2D* SourceTexture = ResourceCast(SourceTextureRHI);
-	FOpenGLTexture2D* DestinationTexture = ResourceCast(DestinationTextureRHI);
-
-	check(SourceTexture->Target == DestinationTexture->Target);
-
-	// Use a texture stage that's not likely to be used for draws, to avoid waiting
-	FOpenGLContextState& ContextState = GetContextStateForCurrentContext();
-	CachedSetupTextureStage(ContextState, FOpenGL::GetMaxCombinedTextureImageUnits() - 1, DestinationTexture->Target, DestinationTexture->GetResource(), 0, DestinationTexture->GetNumMips());
-	CachedBindPixelUnpackBuffer(ContextState, 0);
-
-	// Convert sub texture regions to GL types
-	GLint XOffset = DestinationBox.Min.X;
-	GLint YOffset = DestinationBox.Min.Y;
-	GLint X = SourceBox.Min.X;
-	GLint Y = SourceBox.Min.Y;
-	GLsizei Width = DestinationBox.Max.X - DestinationBox.Min.X;
-	GLsizei Height = DestinationBox.Max.Y - DestinationBox.Min.Y;
-
-	// Bind source texture to an FBO to read from
-	FOpenGLTextureBase* RenderTarget[] = { SourceTexture };
-	uint32 MipLevel = 0;
-	GLuint SourceFBO = GetOpenGLFramebuffer(1, RenderTarget, NULL, &MipLevel, NULL);
-	check(SourceFBO != 0);
-
-	glBindFramebuffer(GL_FRAMEBUFFER, SourceFBO);
-
-	FOpenGL::ReadBuffer(GL_COLOR_ATTACHMENT0);
-	FOpenGL::CopyTexSubImage2D(DestinationTexture->Target, 0, XOffset, YOffset, X, Y, Width, Height);
-
-	ContextState.Framebuffer = (GLuint)-1;
-}
-
-
 void FOpenGLDynamicRHI::RHICopyTexture(FRHITexture* SourceTextureRHI, FRHITexture* DestTextureRHI, const FRHICopyTextureInfo& CopyInfo)
 {
 	VERIFY_GL_SCOPE();
@@ -2814,21 +2778,7 @@ FTextureCubeRHIRef FOpenGLDynamicRHI::RHICreateTextureCubeFromResource(EPixelFor
 	return TextureCube;
 }
 
-PRAGMA_DISABLE_DEPRECATION_WARNINGS
-
 void FOpenGLDynamicRHI::RHIAliasTextureResources(FTextureRHIRef& DestRHITexture, FTextureRHIRef& SrcRHITexture)
-{
-	// @todo: Move the raw-pointer implementation down here when it's deprecation is completed.
-	RHIAliasTextureResources((FRHITexture*)DestRHITexture, (FRHITexture*)SrcRHITexture);
-}
-
-FTextureRHIRef FOpenGLDynamicRHI::RHICreateAliasedTexture(FTextureRHIRef& SourceTexture)
-{
-	// @todo: Move the raw-pointer implementation down here when it's deprecation is completed.
-	return RHICreateAliasedTexture((FRHITexture*)SourceTexture);
-}
-
-void FOpenGLDynamicRHI::RHIAliasTextureResources(FRHITexture* DestRHITexture, FRHITexture* SrcRHITexture)
 {
 	VERIFY_GL_SCOPE();
 	FOpenGLTextureBase* DestTexture = GetOpenGLTextureFromRHITexture(DestRHITexture);
@@ -2840,77 +2790,55 @@ void FOpenGLDynamicRHI::RHIAliasTextureResources(FRHITexture* DestRHITexture, FR
 	}
 }
 
-FRHITexture* FOpenGLDynamicRHI::CreateTexture2DAliased(FOpenGLTexture2D* SourceTexture)
-{
-	SCOPE_CYCLE_COUNTER(STAT_OpenGLCreateTextureTime);
-
-	FRHITexture* Result = new FOpenGLTexture2D(this, 0, SourceTexture->Target, -1, SourceTexture->GetSizeX(), SourceTexture->GetSizeY(), 0, 
-		SourceTexture->GetNumMips(), SourceTexture->GetNumSamples(), SourceTexture->GetNumSamplesTileMem(), 1, SourceTexture->GetFormat(), 
-		false, false, SourceTexture->GetFlags(), SourceTexture->GetClearBinding());
-
-	RHIAliasTextureResources(Result, SourceTexture);
-
-	return Result;
-}
-
-FRHITexture* FOpenGLDynamicRHI::CreateTexture2DArrayAliased(FOpenGLTexture2DArray* SourceTexture)
-{
-	SCOPE_CYCLE_COUNTER(STAT_OpenGLCreateTextureTime);
-
-	FRHITexture* Result = new FOpenGLTexture2DArray(this, 0, SourceTexture->Target, -1, SourceTexture->GetSizeX(), SourceTexture->GetSizeY(), SourceTexture->GetSizeZ(),
-		SourceTexture->GetNumMips(), SourceTexture->GetNumSamples(), 1 /* aka check(InNumSamplesTileMem == 1) in OpenGLResource.h FOpenGLBaseTexture2DArray constructor */,
-		1, SourceTexture->GetFormat(), false, false, SourceTexture->GetFlags(), SourceTexture->GetClearBinding());
-
-	RHIAliasTextureResources(Result, SourceTexture);
-
-	return Result;
-}
-
-FRHITexture* FOpenGLDynamicRHI::CreateTextureCubeAliased(FOpenGLTextureCube* SourceTexture)
+FTextureRHIRef FOpenGLDynamicRHI::RHICreateAliasedTexture(FTextureRHIRef& SourceTexture)
 {
 	SCOPE_CYCLE_COUNTER(STAT_OpenGLCreateTextureTime);
 
-	FRHITexture* Result = new FOpenGLTextureCube(this, 0, SourceTexture->Target, -1, SourceTexture->GetSizeX(), SourceTexture->GetSizeY(), SourceTexture->GetSizeZ(),
-		SourceTexture->GetNumMips(), SourceTexture->GetNumSamples(), 1 /* OpenGL currently doesn't support multisample cube textures, per OpenGLResource.h FOpenGLBaseTextureCube */,
-		1, SourceTexture->GetFormat(), true, false, SourceTexture->GetFlags(), SourceTexture->GetClearBinding());
-
-	RHIAliasTextureResources(Result, SourceTexture);
-
-	return Result;
-}
-
-FTextureRHIRef FOpenGLDynamicRHI::RHICreateAliasedTexture(FRHITexture* SourceTexture)
-{
-	FTextureRHIRef AliasedTexture;
+	FTextureRHIRef Result;
 	if (SourceTexture->GetTexture2D() != nullptr)
 	{
-		AliasedTexture = CreateTexture2DAliased(static_cast<FOpenGLTexture2D*>(SourceTexture));
+		FOpenGLTexture2D* SourceTexture2D = static_cast<FOpenGLTexture2D*>(SourceTexture.GetReference());
+
+		Result = new FOpenGLTexture2D(this, 0, SourceTexture2D->Target, -1, SourceTexture2D->GetSizeX(), SourceTexture2D->GetSizeY(), 0,
+			SourceTexture2D->GetNumMips(), SourceTexture2D->GetNumSamples(), SourceTexture2D->GetNumSamplesTileMem(), 1, SourceTexture2D->GetFormat(),
+			false, false, SourceTexture2D->GetFlags(), SourceTexture2D->GetClearBinding());
 	}
 	else if (SourceTexture->GetTexture2DArray() != nullptr)
 	{
-		AliasedTexture = CreateTexture2DArrayAliased(static_cast<FOpenGLTexture2DArray*>(SourceTexture));
+		FOpenGLTexture2DArray* SourceTexture2DArray = static_cast<FOpenGLTexture2DArray*>(SourceTexture.GetReference());
+
+		Result = new FOpenGLTexture2DArray(this, 0, SourceTexture2DArray->Target, -1, SourceTexture2DArray->GetSizeX(), SourceTexture2DArray->GetSizeY(), SourceTexture2DArray->GetSizeZ(),
+			SourceTexture2DArray->GetNumMips(), SourceTexture2DArray->GetNumSamples(), 1 /* aka check(InNumSamplesTileMem == 1) in OpenGLResource.h FOpenGLBaseTexture2DArray constructor */,
+			1, SourceTexture2DArray->GetFormat(), false, false, SourceTexture2DArray->GetFlags(), SourceTexture2DArray->GetClearBinding());
 	}
 	else if (SourceTexture->GetTextureCube() != nullptr)
 	{
-		AliasedTexture = CreateTextureCubeAliased(static_cast<FOpenGLTextureCube*>(SourceTexture));
+		FOpenGLTextureCube* SourceTextureCube = static_cast<FOpenGLTextureCube*>(SourceTexture.GetReference());
+
+		Result = new FOpenGLTextureCube(this, 0, SourceTextureCube->Target, -1, SourceTextureCube->GetSizeX(), SourceTextureCube->GetSizeY(), SourceTextureCube->GetSizeZ(),
+			SourceTextureCube->GetNumMips(), SourceTextureCube->GetNumSamples(), 1 /* OpenGL currently doesn't support multisample cube textures, per OpenGLResource.h FOpenGLBaseTextureCube */,
+			1, SourceTextureCube->GetFormat(), true, false, SourceTextureCube->GetFlags(), SourceTextureCube->GetClearBinding());
 	}
 	else
 	{
 		UE_LOG(LogRHI, Error, TEXT("Currently FOpenGLDynamicRHI::RHICreateAliasedTexture only supports 2D, 2D Array and Cube textures."));
 	}
 
-	if (AliasedTexture)
+	if (Result != nullptr)
 	{
-		FOpenGLTextureBase* BaseTexture = GetOpenGLTextureFromRHITexture(AliasedTexture);
+		RHIAliasTextureResources(Result, SourceTexture);
+	}
+
+	if (Result)
+	{
+		FOpenGLTextureBase* BaseTexture = GetOpenGLTextureFromRHITexture(Result);
 		// Init memory size to zero, since we're aliased.
 		BaseTexture->SetMemorySize(0);
 	}
 
-	return AliasedTexture;
+	return Result;
 }
 
-PRAGMA_ENABLE_DEPRECATION_WARNINGS
-
 void* FOpenGLDynamicRHI::LockTexture2D_RenderThread(class FRHICommandListImmediate& RHICmdList, FRHITexture2D* Texture, uint32 MipIndex, EResourceLockMode LockMode, uint32& DestStride, bool bLockWithinMiptail, bool bNeedsDefaultRHIFlush)
 {
 	check(IsInRenderingThread());
diff --git a/Engine/Source/Runtime/OpenGLDrv/Public/OpenGLDrv.h b/Engine/Source/Runtime/OpenGLDrv/Public/OpenGLDrv.h
--- a/Engine/Source/Runtime/OpenGLDrv/Public/OpenGLDrv.h
+++ b/Engine/Source/Runtime/OpenGLDrv/Public/OpenGLDrv.h
@@ -394,10 +394,6 @@ public:
 	virtual FTexture2DRHIRef RHIGetViewportBackBuffer(FRHIViewport* Viewport) final override;
 	virtual void RHIAliasTextureResources(FTextureRHIRef& DestTextureRHI, FTextureRHIRef& SrcTextureRHI) final override;
 	virtual FTextureRHIRef RHICreateAliasedTexture(FTextureRHIRef& SourceTexture) final override;
-	UE_DEPRECATED(4.25, "RHIAliasTextureResources now takes references to FTextureRHIRef objects as parameters")
-	virtual void RHIAliasTextureResources(FRHITexture* DestTexture, FRHITexture* SrcTexture) final override;
-	UE_DEPRECATED(4.25, "RHICreateAliasedTexture now takes a reference to an FTextureRHIRef object")
-	virtual FTextureRHIRef RHICreateAliasedTexture(FRHITexture* SourceTexture) final override;
 	virtual void RHIAdvanceFrameForGetViewportBackBuffer(FRHIViewport* Viewport) final override;
 	virtual void RHIAcquireThreadOwnership() final override;
 	virtual void RHIReleaseThreadOwnership() final override;
@@ -485,7 +481,6 @@ public:
 	virtual void RHIInvalidateCachedState() final override;
 	virtual void RHIDiscardRenderTargets(bool Depth,bool Stencil,uint32 ColorBitMask) final override;
 
-	virtual void RHICopySubTextureRegion(FRHITexture2D* SourceTexture, FRHITexture2D* DestinationTexture, FBox2D SourceBox, FBox2D DestinationBox) final override;
 	// FIXME: Broken on Android for cubemaps
 	virtual void RHICopyTexture(FRHITexture* SourceTexture, FRHITexture* DestTexture, const FRHICopyTextureInfo& CopyInfo) final override;
 
@@ -619,9 +614,6 @@ public:
 	FRHITexture* CreateOpenGLTexture(uint32 SizeX, uint32 SizeY, bool CubeTexture, bool ArrayTexture, bool bIsExternal, uint8 Format, uint32 NumMips, uint32 NumSamples, uint32 ArraySize, ETextureCreateFlags Flags, const FClearValueBinding& InClearValue, FResourceBulkDataInterface* BulkData = NULL);
 	
 	FRHITexture* CreateOpenGLRHITextureOnly(const uint32 SizeX, const uint32 SizeY, const bool bCubeTexture, const bool bArrayTexture, const bool bIsExternal, uint8& Format, uint32& NumMips, uint32& NumSamples, const uint32 ArraySize, ETextureCreateFlags Flags, const FClearValueBinding& InClearValue, FResourceBulkDataInterface* BulkData = nullptr);
-	FRHITexture* CreateTexture2DAliased(FOpenGLTexture2D* SourceTexture);
-	FRHITexture* CreateTexture2DArrayAliased(FOpenGLTexture2DArray* SourceTexture);
-	FRHITexture* CreateTextureCubeAliased(FOpenGLTextureCube* SourceTexture);
 
 	void InitializeGLTexture(FRHITexture* Texture, uint32 SizeX, const uint32 SizeY, const bool bCubeTexture, const bool bArrayTexture, const bool bIsExternal, const uint8 Format, const uint32 NumMips, const uint32 NumSamples, const uint32 ArraySize, const ETextureCreateFlags Flags, const FClearValueBinding& InClearValue, FResourceBulkDataInterface* BulkData = nullptr);
 
diff --git a/Engine/Source/Runtime/RHI/Private/RHICommandList.cpp b/Engine/Source/Runtime/RHI/Private/RHICommandList.cpp
--- a/Engine/Source/Runtime/RHI/Private/RHICommandList.cpp
+++ b/Engine/Source/Runtime/RHI/Private/RHICommandList.cpp
@@ -2754,10 +2754,3 @@ void FRHICommandListImmediate::UpdateRHIResources(FRHIResourceUpdateInfo* Update
 		}
 	}
 }
-
-void FDynamicRHI::RHICopySubTextureRegion_RenderThread(class FRHICommandListImmediate& RHICmdList, FRHITexture2D* SourceTexture, FRHITexture2D* DestinationTexture, FBox2D SourceBox, FBox2D DestinationBox)
-{
-	CSV_SCOPED_TIMING_STAT(RHITStalls, RHICopySubTextureRegion_RenderThread);
-	FScopedRHIThreadStaller StallRHIThread(RHICmdList);
-	return GDynamicRHI->RHICopySubTextureRegion(SourceTexture, DestinationTexture, SourceBox, DestinationBox);
-}
diff --git a/Engine/Source/Runtime/RHI/Public/DynamicRHI.h b/Engine/Source/Runtime/RHI/Public/DynamicRHI.h
--- a/Engine/Source/Runtime/RHI/Public/DynamicRHI.h
+++ b/Engine/Source/Runtime/RHI/Public/DynamicRHI.h
@@ -943,19 +943,6 @@ public:
 		return nullptr;
 	}
 
-	UE_DEPRECATED(4.25, "RHIAliasTextureResources now takes references to FTextureRHIRef objects as parameters")
-	virtual void RHIAliasTextureResources(FRHITexture* DestTexture, FRHITexture* SrcTexture)
-	{
-		checkNoEntry();
-	}
-
-	UE_DEPRECATED(4.25, "RHICreateAliasedTexture now takes a reference to an FTextureRHIRef object")
-	virtual FTextureRHIRef RHICreateAliasedTexture(FRHITexture* SourceTexture)
-	{
-		checkNoEntry();
-		return nullptr;
-	}
-
 	virtual void RHIAliasTextureResources(FTextureRHIRef& DestTexture, FTextureRHIRef& SrcTexture)
 	{
 		checkNoEntry();
@@ -1270,10 +1257,6 @@ public:
 	virtual void VirtualTextureSetFirstMipInMemory_RenderThread(class FRHICommandListImmediate& RHICmdList, FRHITexture2D* Texture, uint32 FirstMip);
 	virtual void VirtualTextureSetFirstMipVisible_RenderThread(class FRHICommandListImmediate& RHICmdList, FRHITexture2D* Texture, uint32 FirstMip);
 
-	/* Copy the source box pixels in the destination box texture, return true if implemented for the current platform*/
-	virtual void RHICopySubTextureRegion_RenderThread(class FRHICommandListImmediate& RHICmdList, FRHITexture2D* SourceTexture, FRHITexture2D* DestinationTexture, FBox2D SourceBox, FBox2D DestinationBox);
-	virtual void RHICopySubTextureRegion(FRHITexture2D* SourceTexture, FRHITexture2D* DestinationTexture, FBox2D SourceBox, FBox2D DestinationBox) { }
-
 	virtual FRHIFlipDetails RHIWaitForFlip(double TimeoutInSeconds) { return FRHIFlipDetails(); }
 	virtual void RHISignalFlipEvent() { }
 
diff --git a/Engine/Source/Runtime/RHI/Public/RHICommandList.h b/Engine/Source/Runtime/RHI/Public/RHICommandList.h
--- a/Engine/Source/Runtime/RHI/Public/RHICommandList.h
+++ b/Engine/Source/Runtime/RHI/Public/RHICommandList.h
@@ -2148,17 +2148,6 @@ FRHICOMMAND_MACRO(FRHICommandCopyBufferRegion)
 
 #if RHI_RAYTRACING
 
-FRHICOMMAND_MACRO(FRHICommandCopyBufferRegions)
-{
-	const TArrayView<const FCopyBufferRegionParams> Params;
-
-	explicit FRHICommandCopyBufferRegions(const TArrayView<const FCopyBufferRegionParams> InParams)
-		: Params(InParams)
-	{}
-
-	RHI_API void Execute(FRHICommandListBase& CmdList);
-};
-
 struct FRHICommandBindAccelerationStructureMemory final : public FRHICommand<FRHICommandBindAccelerationStructureMemory>
 {
 	FRHIRayTracingScene* Scene;
@@ -3584,13 +3573,6 @@ public:
 		ALLOC_COMMAND(FRHICommandResummarizeHTile)(DepthTexture);
 	}
 
-	UE_DEPRECATED(4.25, "RHIClearTinyUAV is deprecated. Use RHIClearUAVUint or RHIClearUAVFloat instead.")
-	FORCEINLINE_DEBUGGABLE void ClearTinyUAV(FRHIUnorderedAccessView* UnorderedAccessViewRHI, const uint32(&Values)[4])
-	{
-		// Forward to the new uint clear implementation.
-		ClearUAVUint(UnorderedAccessViewRHI, FUintVector4(Values[0], Values[1], Values[2], Values[3]));
-	}
-
 	FORCEINLINE_DEBUGGABLE void BeginRenderQuery(FRHIRenderQuery* RenderQuery)
 	{
 		if (Bypass())
@@ -3769,21 +3751,6 @@ public:
 
 #if RHI_RAYTRACING
 	// Ray tracing API
-	UE_DEPRECATED(4.25, "CopyBufferRegions API is deprecated. Use an explicit compute shader copy dispatch instead.")
-	FORCEINLINE_DEBUGGABLE void CopyBufferRegions(const TArrayView<const FCopyBufferRegionParams> Params)
-	{
-		// No copy/DMA operation inside render passes
-		check(IsOutsideRenderPass());
-
-		if (Bypass())
-		{
-			GetContext().RHICopyBufferRegions(Params);
-		}
-		else
-		{
-			ALLOC_COMMAND(FRHICommandCopyBufferRegions)(AllocArray(Params));
-		}
-	}
 
 	FORCEINLINE_DEBUGGABLE void ClearRayTracingBindings(FRHIRayTracingScene* Scene)
 	{
@@ -4623,19 +4590,6 @@ PRAGMA_ENABLE_DEPRECATION_WARNINGS
 		ImmediateFlush(EImmediateFlushType::FlushRHIThread);  
 		GDynamicRHI->RHIRead3DSurfaceFloatData(Texture,Rect,ZMinMax,OutData,Flags);
 	}
-	
-	UE_DEPRECATED(4.23, "CreateRenderQuery API is deprecated; use RHICreateRenderQueryPool and suballocate queries there")
-	FORCEINLINE FRenderQueryRHIRef CreateRenderQuery(ERenderQueryType QueryType)
-	{
-		FScopedRHIThreadStaller StallRHIThread(*this);
-		return GDynamicRHI->RHICreateRenderQuery(QueryType);
-	}
-
-	UE_DEPRECATED(4.23, "CreateRenderQuery API is deprecated; use RHICreateRenderQueryPool and suballocate queries there")
-	FORCEINLINE FRenderQueryRHIRef CreateRenderQuery_RenderThread(ERenderQueryType QueryType)
-	{
-		return GDynamicRHI->RHICreateRenderQuery_RenderThread(*this, QueryType);
-	}
 
 	UE_DEPRECATED(5.0, "AcquireTransientResource_RenderThread API is deprecated; use IRHITransientResourceAllocator instead.")
 	FORCEINLINE void AcquireTransientResource_RenderThread(FRHITexture* Texture) {}
@@ -4772,12 +4726,6 @@ PRAGMA_ENABLE_DEPRECATION_WARNINGS
 		GDynamicRHI->VirtualTextureSetFirstMipVisible_RenderThread(*this, Texture, FirstMip);
 	}
 
-	UE_DEPRECATED(4.23, "CopySubTextureRegion API is deprecated; please use CopyTexture instead.")
-	FORCEINLINE void CopySubTextureRegion(FRHITexture2D* SourceTexture, FRHITexture2D* DestinationTexture, FBox2D SourceBox, FBox2D DestinationBox)
-	{
-		GDynamicRHI->RHICopySubTextureRegion_RenderThread(*this, SourceTexture, DestinationTexture, SourceBox, DestinationBox);
-	}
-	
 	FORCEINLINE void ExecuteCommandList(FRHICommandList* CmdList)
 	{
 		FScopedRHIThreadStaller StallRHIThread(*this);
@@ -5592,14 +5540,6 @@ FORCEINLINE void RHIUnlockTextureCubeFace(FRHITextureCube* Texture, uint32 FaceI
 	 FRHICommandListExecutor::GetImmediateCommandList().UnlockTextureCubeFace(Texture, FaceIndex, ArrayIndex, MipIndex, bLockWithinMiptail);
 }
 
-UE_DEPRECATED(4.23, "CreateRenderQuery API is deprecated; use RHICreateRenderQueryPool and suballocate queries there")
-FORCEINLINE FRenderQueryRHIRef RHICreateRenderQuery(ERenderQueryType QueryType)
-{
-	PRAGMA_DISABLE_DEPRECATION_WARNINGS
-	return FRHICommandListExecutor::GetImmediateCommandList().CreateRenderQuery_RenderThread(QueryType);
-	PRAGMA_ENABLE_DEPRECATION_WARNINGS
-}
-
 UE_DEPRECATED(5.0, "RHIAcquireTransientResource API is deprecated; use IRHITransientResourceAllocator instead.")
 FORCEINLINE void RHIAcquireTransientResource(FRHITexture*) {}
 
diff --git a/Engine/Source/Runtime/RHI/Public/RHICommandListCommandExecutes.inl b/Engine/Source/Runtime/RHI/Public/RHICommandListCommandExecutes.inl
--- a/Engine/Source/Runtime/RHI/Public/RHICommandListCommandExecutes.inl
+++ b/Engine/Source/Runtime/RHI/Public/RHICommandListCommandExecutes.inl
@@ -537,11 +537,6 @@ void FRHICommandCopyBufferRegion::Execute(FRHICommandListBase& CmdList)
 
 #if RHI_RAYTRACING
 
-void FRHICommandCopyBufferRegions::Execute(FRHICommandListBase& CmdList)
-{
-	INTERNAL_DECORATOR(RHICopyBufferRegions)(Params);
-}
-
 void FRHICommandBindAccelerationStructureMemory::Execute(FRHICommandListBase& CmdList)
 {
 	INTERNAL_DECORATOR_COMPUTE(RHIBindAccelerationStructureMemory)(Scene, Buffer, BufferOffset);
diff --git a/Engine/Source/Runtime/RHI/Public/RHIContext.h b/Engine/Source/Runtime/RHI/Public/RHIContext.h
--- a/Engine/Source/Runtime/RHI/Public/RHIContext.h
+++ b/Engine/Source/Runtime/RHI/Public/RHIContext.h
@@ -723,13 +723,6 @@ public:
 		checkNoEntry();
 	}
 
-#if RHI_RAYTRACING
-	virtual void RHICopyBufferRegions(const TArrayView<const FCopyBufferRegionParams> Params)
-	{
-		checkNoEntry();
-	}
-#endif
-
 	virtual void RHIClearRayTracingBindings(FRHIRayTracingScene* Scene)
 	{
 		checkNoEntry();
diff --git a/Engine/Source/Runtime/RHI/Public/RHIValidation.h b/Engine/Source/Runtime/RHI/Public/RHIValidation.h
--- a/Engine/Source/Runtime/RHI/Public/RHIValidation.h
+++ b/Engine/Source/Runtime/RHI/Public/RHIValidation.h
@@ -1655,31 +1655,6 @@ public:
 		RHI->VirtualTextureSetFirstMipVisible_RenderThread(RHICmdList, Texture, FirstMip);
 	}
 
-	/* Copy the source box pixels in the destination box texture, return true if implemented for the current platform*/
-	virtual void RHICopySubTextureRegion_RenderThread(class FRHICommandListImmediate& RHICmdList, FRHITexture2D* SourceTexture, FRHITexture2D* DestinationTexture, FBox2D SourceBox, FBox2D DestinationBox) override final
-	{
-		int32 Width = static_cast<int32>(SourceBox.GetSize().X);
-		int32 Height = static_cast<int32>(SourceBox.GetSize().Y);
-		int32 SrcX = static_cast<int32>(SourceBox.Min.X);
-		int32 SrcY = static_cast<int32>(SourceBox.Min.Y);
-		int32 DestX = static_cast<int32>(DestinationBox.Min.X);
-		int32 DestY = static_cast<int32>(DestinationBox.Min.Y);
-		FValidationRHIUtils::ValidateCopyTexture(SourceTexture, DestinationTexture, FIntVector(Width, Height, 1), FIntVector(SrcX, SrcY, 0), FIntVector(DestX, DestY, 0));
-		return RHI->RHICopySubTextureRegion_RenderThread(RHICmdList, SourceTexture, DestinationTexture, SourceBox, DestinationBox);
-	}
-
-	virtual void RHICopySubTextureRegion(FRHITexture2D* SourceTexture, FRHITexture2D* DestinationTexture, FBox2D SourceBox, FBox2D DestinationBox) override final
-	{
-		int32 Width = static_cast<int32>(SourceBox.GetSize().X);
-		int32 Height = static_cast<int32>(SourceBox.GetSize().Y);
-		int32 SrcX = static_cast<int32>(SourceBox.Min.X);
-		int32 SrcY = static_cast<int32>(SourceBox.Min.Y);
-		int32 DestX = static_cast<int32>(DestinationBox.Min.X);
-		int32 DestY = static_cast<int32>(DestinationBox.Min.Y);
-		FValidationRHIUtils::ValidateCopyTexture(SourceTexture, DestinationTexture, FIntVector(Width, Height, 1), FIntVector(SrcX, SrcY, 0), FIntVector(DestX, DestY, 0));
-		RHI->RHICopySubTextureRegion(SourceTexture, DestinationTexture, SourceBox, DestinationBox);
-	}
-
 	virtual FRHIFlipDetails RHIWaitForFlip(double TimeoutInSeconds) override final
 	{
 		return RHI->RHIWaitForFlip(TimeoutInSeconds);
diff --git a/Engine/Source/Runtime/RHI/Public/RHIValidationContext.h b/Engine/Source/Runtime/RHI/Public/RHIValidationContext.h
--- a/Engine/Source/Runtime/RHI/Public/RHIValidationContext.h
+++ b/Engine/Source/Runtime/RHI/Public/RHIValidationContext.h
@@ -960,16 +960,6 @@ public:
 		Tracker->Assert(DestBuffer->GetWholeResourceIdentity(), ERHIAccess::CopyDest);
 		RHIContext->RHICopyBufferRegion(DestBuffer, DstOffset, SourceBuffer, SrcOffset, NumBytes);
 	}
-
-	virtual void RHICopyBufferRegions(const TArrayView<const FCopyBufferRegionParams> Params)
-	{
-		for (const FCopyBufferRegionParams& Args : Params)
-		{
-			Tracker->Assert(Args.SourceBuffer->GetWholeResourceIdentity(), ERHIAccess::CopySrc);
-			Tracker->Assert(Args.DestBuffer->GetWholeResourceIdentity(), ERHIAccess::CopyDest);
-		}
-		RHIContext->RHICopyBufferRegions(Params);
-	}
 #endif
 
 	void RHIClearRayTracingBindings(FRHIRayTracingScene* Scene)
diff --git a/Engine/Source/Runtime/VulkanRHI/Private/VulkanRHI.cpp b/Engine/Source/Runtime/VulkanRHI/Private/VulkanRHI.cpp
--- a/Engine/Source/Runtime/VulkanRHI/Private/VulkanRHI.cpp
+++ b/Engine/Source/Runtime/VulkanRHI/Private/VulkanRHI.cpp
@@ -1226,17 +1226,6 @@ FTextureCubeRHIRef FVulkanDynamicRHI::RHICreateTextureCubeFromResource(EPixelFor
 	return new FVulkanTextureCube(*Device, Format, Size, bArray, ArraySize, NumMips, Resource, Flags, ResourceCreateInfo);
 }
 
-void FVulkanDynamicRHI::RHIAliasTextureResources(FRHITexture* DestTextureRHI, FRHITexture* SrcTextureRHI)
-{
-	check(false);
-}
-
-FTextureRHIRef FVulkanDynamicRHI::RHICreateAliasedTexture(FRHITexture* SourceTexture)
-{
-	check(false);
-	return nullptr;
-}
-
 void FVulkanDynamicRHI::RHIAliasTextureResources(FTextureRHIRef& DestTextureRHI, FTextureRHIRef& SrcTextureRHI)
 {
 	if (DestTextureRHI && SrcTextureRHI)
@@ -1275,21 +1264,6 @@ FTextureRHIRef FVulkanDynamicRHI::RHICreateAliasedTexture(FTextureRHIRef& Source
 	return AliasedTexture;
 }
 
-void FVulkanDynamicRHI::RHICopySubTextureRegion(FRHITexture2D* SourceTexture, FRHITexture2D* DestinationTexture, FBox2D SourceBox, FBox2D DestinationBox)
-{
-	FRHICopyTextureInfo CopyInfo;
-
-	CopyInfo.Size.X = (int32)(SourceBox.Max.X - SourceBox.Min.X);
-	CopyInfo.Size.Y = (int32)(SourceBox.Max.Y - SourceBox.Min.Y);
-
-	CopyInfo.SourcePosition.X = (int32)(SourceBox.Min.X);
-	CopyInfo.SourcePosition.Y = (int32)(SourceBox.Min.Y);
-	CopyInfo.DestPosition.X = (int32)(DestinationBox.Min.X);
-	CopyInfo.DestPosition.Y = (int32)(DestinationBox.Min.Y);
-
-	RHIGetDefaultContext()->RHICopyTexture(SourceTexture, DestinationTexture, CopyInfo);
-}
-
 FVulkanDescriptorSetsLayout::FVulkanDescriptorSetsLayout(FVulkanDevice* InDevice) :
 	Device(InDevice)
 {
diff --git a/Engine/Source/Runtime/VulkanRHI/Public/VulkanDynamicRHI.h b/Engine/Source/Runtime/VulkanRHI/Public/VulkanDynamicRHI.h
--- a/Engine/Source/Runtime/VulkanRHI/Public/VulkanDynamicRHI.h
+++ b/Engine/Source/Runtime/VulkanRHI/Public/VulkanDynamicRHI.h
@@ -145,10 +145,6 @@ public:
 	virtual FRenderQueryRHIRef RHICreateRenderQuery(ERenderQueryType QueryType) final override;
 	virtual bool RHIGetRenderQueryResult(FRHIRenderQuery* RenderQuery, uint64& OutResult, bool bWait, uint32 GPUIndex = INDEX_NONE) final override;
 	virtual FTexture2DRHIRef RHIGetViewportBackBuffer(FRHIViewport* Viewport) final override;
-	UE_DEPRECATED(4.25, "RHIAliasTextureResources now takes references to FTextureRHIRef objects as parameters")
-	virtual void RHIAliasTextureResources(FRHITexture* DestTexture, FRHITexture* SrcTexture) final override;
-	UE_DEPRECATED(4.25, "RHICreateAliasedTexture now takes a reference to an FTextureRHIRef object")
-	virtual FTextureRHIRef RHICreateAliasedTexture(FRHITexture* SourceTexture) final override;
 	virtual void RHIAliasTextureResources(FTextureRHIRef& DestTexture, FTextureRHIRef& SrcTexture) final override;
 	virtual FTextureRHIRef RHICreateAliasedTexture(FTextureRHIRef& SourceTexture) final override;
 	virtual void RHIAdvanceFrameForGetViewportBackBuffer(FRHIViewport* Viewport) final override;
@@ -320,8 +316,6 @@ public:
 	virtual FTexture2DArrayRHIRef RHICreateTexture2DArrayFromResource(EPixelFormat Format, uint32 SizeX, uint32 SizeY, uint32 ArraySize, uint32 NumMips, uint32 NumSamples, VkImage Resource, ETextureCreateFlags Flags);
 	virtual FTextureCubeRHIRef RHICreateTextureCubeFromResource(EPixelFormat Format, uint32 Size, bool bArray, uint32 ArraySize, uint32 NumMips, VkImage Resource, ETextureCreateFlags Flags);
 
-	virtual void RHICopySubTextureRegion(FRHITexture2D* SourceTexture, FRHITexture2D* DestinationTexture, FBox2D SourceBox, FBox2D DestinationBox) final override;
-
 	void RHICalibrateTimers() override;
 
 #if VULKAN_RHI_RAYTRACING
diff --git a/Engine/Source/Runtime/Windows/D3D11RHI/Private/D3D11RHIPrivate.h b/Engine/Source/Runtime/Windows/D3D11RHI/Private/D3D11RHIPrivate.h
--- a/Engine/Source/Runtime/Windows/D3D11RHI/Private/D3D11RHIPrivate.h
+++ b/Engine/Source/Runtime/Windows/D3D11RHI/Private/D3D11RHIPrivate.h
@@ -529,10 +529,6 @@ public:
 	virtual FRenderQueryRHIRef RHICreateRenderQuery(ERenderQueryType QueryType) final override;
 	virtual bool RHIGetRenderQueryResult(FRHIRenderQuery* RenderQuery, uint64& OutResult, bool bWait, uint32 GPUIndex = INDEX_NONE) final override;
 	virtual FTexture2DRHIRef RHIGetViewportBackBuffer(FRHIViewport* Viewport) final override;
-	UE_DEPRECATED(4.25, "RHIAliasTextureResources now takes references to FTextureRHIRef objects as parameters")
-	virtual void RHIAliasTextureResources(FRHITexture* DestTexture, FRHITexture* SrcTexture) final override;
-	UE_DEPRECATED(4.25, "RHICreateAliasedTexture now takes a reference to an FTextureRHIRef object")
-	virtual FTextureRHIRef RHICreateAliasedTexture(FRHITexture* SourceTexture) final override;
 	virtual void RHIAliasTextureResources(FTextureRHIRef& DestTexture, FTextureRHIRef& SrcTexture) final override;
 	virtual FTextureRHIRef RHICreateAliasedTexture(FTextureRHIRef& SourceTexture) final override;
 	virtual void RHIAdvanceFrameForGetViewportBackBuffer(FRHIViewport* Viewport) final override;
@@ -633,8 +629,6 @@ public:
 	virtual void RHIPushEvent(const TCHAR* Name, FColor Color) final override;
 	virtual void RHIPopEvent() final override;
 
-	virtual void RHICopySubTextureRegion(FRHITexture2D* SourceTexture, FRHITexture2D* DestinationTexture, FBox2D SourceBox, FBox2D DestinationBox) final override;
-
 	virtual FTexture2DRHIRef RHICreateTexture2DFromResource(EPixelFormat Format, ETextureCreateFlags TexCreateFlags, const FClearValueBinding& ClearValueBinding, ID3D11Texture2D* Resource);
 	virtual FTexture2DArrayRHIRef RHICreateTexture2DArrayFromResource(EPixelFormat Format, ETextureCreateFlags TexCreateFlags, const FClearValueBinding&, ID3D11Texture2D* Resource);
 	virtual FTextureCubeRHIRef RHICreateTextureCubeFromResource(EPixelFormat Format, ETextureCreateFlags TexCreateFlags, const FClearValueBinding& ClearValueBinding, ID3D11Texture2D* Resource);
@@ -675,7 +669,6 @@ public:
 	virtual FRenderQueryRHIRef RHICreateRenderQuery_RenderThread(class FRHICommandListImmediate& RHICmdList, ERenderQueryType QueryType) final override;
 	virtual void* RHILockTextureCubeFace_RenderThread(class FRHICommandListImmediate& RHICmdList, FRHITextureCube* Texture, uint32 FaceIndex, uint32 ArrayIndex, uint32 MipIndex, EResourceLockMode LockMode, uint32& DestStride, bool bLockWithinMiptail) final override;
 	virtual void RHIUnlockTextureCubeFace_RenderThread(class FRHICommandListImmediate& RHICmdList, FRHITextureCube* Texture, uint32 FaceIndex, uint32 ArrayIndex, uint32 MipIndex, bool bLockWithinMiptail) final override;
-	virtual void RHICopySubTextureRegion_RenderThread(class FRHICommandListImmediate& RHICmdList, FRHITexture2D* SourceTexture, FRHITexture2D* DestinationTexture, FBox2D SourceBox, FBox2D DestinationBox) final override;
 
 	virtual void RHIBeginRenderPass(const FRHIRenderPassInfo& InInfo, const TCHAR* InName) final override;
 	virtual void RHIEndRenderPass() final override;
@@ -966,7 +959,7 @@ protected:
 	TD3D11Texture2D<BaseResourceType>* CreateTextureFromResource(bool bTextureArray, bool bCubeTexture, EPixelFormat Format, ETextureCreateFlags TexCreateFlags, const FClearValueBinding& ClearValueBinding, ID3D11Texture2D* TextureResource);
 
 	template<typename BaseResourceType>
-	TD3D11Texture2D<BaseResourceType>* CreateAliasedD3D11Texture2D(TD3D11Texture2D<BaseResourceType>* SourceTexture);
+	FTextureRHIRef CreateAliasedD3D11Texture2D(TD3D11Texture2D<BaseResourceType>* SourceTexture);
 
 	/** Initializes the constant buffers.  Called once at RHI initialization time. */
 	void InitConstantBuffers();
diff --git a/Engine/Source/Runtime/Windows/D3D11RHI/Private/D3D11Texture.cpp b/Engine/Source/Runtime/Windows/D3D11RHI/Private/D3D11Texture.cpp
--- a/Engine/Source/Runtime/Windows/D3D11RHI/Private/D3D11Texture.cpp
+++ b/Engine/Source/Runtime/Windows/D3D11RHI/Private/D3D11Texture.cpp
@@ -2179,107 +2179,6 @@ void FD3D11DynamicRHI::RHIVirtualTextureSetFirstMipVisible(FRHITexture2D* Textur
 {
 }
 
-void FD3D11DynamicRHI::RHICopySubTextureRegion(FRHITexture2D* SourceTextureRHI, FRHITexture2D* DestinationTextureRHI, FBox2D SourceBox, FBox2D DestinationBox)
-{
-	FD3D11Texture2D* SourceTexture = ResourceCast(SourceTextureRHI);
-	FD3D11Texture2D* DestinationTexture = ResourceCast(DestinationTextureRHI);
-
-	//Make sure the source box is fitting on right and top side of the source texture, no need to offset the destination
-	if (SourceBox.Max.X >= (float)SourceTexture->GetSizeX())
-	{
-		float Delta = (SourceBox.Max.X - (float)SourceTexture->GetSizeX());
-		SourceBox.Max.X -= Delta;
-	}
-	if (SourceBox.Max.Y >= (float)SourceTexture->GetSizeY())
-	{
-		float Delta = (SourceBox.Max.Y - (float)SourceTexture->GetSizeY());
-		SourceBox.Max.Y -= Delta;
-	}
-
-	int32 DestinationOffsetX = 0;
-	int32 DestinationOffsetY = 0;
-	int32 SourceStartX = SourceBox.Min.X;
-	int32 SourceEndX = SourceBox.Max.X;
-	int32 SourceStartY = SourceBox.Min.Y;
-	int32 SourceEndY = SourceBox.Max.Y;
-	//If the source box is not fitting on the left bottom side, offset the result so the destination pixel match the expectation
-	if (SourceStartX < 0)
-	{
-		DestinationOffsetX -= SourceStartX;
-		SourceStartX = 0;
-	}
-	if (SourceStartY < 0)
-	{
-		DestinationOffsetY -= SourceStartY;
-		SourceStartY = 0;
-	}
-
-	D3D11_BOX SourceBoxAdjust =
-	{
-		static_cast<UINT>(SourceStartX),
-		static_cast<UINT>(SourceStartY),
-		0,
-		static_cast<UINT>(SourceEndX),
-		static_cast<UINT>(SourceEndY),
-		1
-	};
-
-	bool bValidDest = DestinationBox.Min.X + DestinationOffsetX + (SourceEndX - SourceStartX) <= DestinationTexture->GetSizeX();
-	bValidDest &= DestinationBox.Min.Y + DestinationOffsetY + (SourceEndY - SourceStartY) <= DestinationTexture->GetSizeY();
-	bValidDest &= DestinationBox.Min.X <= DestinationBox.Max.X && DestinationBox.Min.Y <= DestinationBox.Max.Y;
-
-	bool bValidSrc = SourceStartX >= 0 && SourceEndX <= (int32)SourceTexture->GetSizeX();
-	bValidSrc &= SourceStartY >= 0 && SourceEndY <= (int32)SourceTexture->GetSizeY();
-	bValidSrc &= SourceStartX <= SourceEndX && SourceStartY <= SourceEndY;
-
-	if (!ensureMsgf(bValidSrc && bValidDest, TEXT("Invalid copy detected for RHICopySubTextureRegion. Skipping copy.  SrcBox: left:%i, right:%i, top:%i, bottom:%i, DstBox:left:%i, right:%i, top:%i, bottom:%i,  SrcTexSize: %i x %i, DestTexSize: %i x %i "),
-		SourceBox.Min.X,
-		SourceBox.Max.X,
-		SourceBox.Min.Y,
-		SourceBox.Max.Y,
-		DestinationBox.Min.X,
-		DestinationBox.Max.X,
-		DestinationBox.Min.Y,
-		DestinationBox.Max.Y,
-		SourceTexture->GetSizeX(),
-		SourceTexture->GetSizeY(),
-		DestinationTexture->GetSizeX(),
-		DestinationTexture->GetSizeY()))
-	{
-		return;
-	}
-
-	check(SourceBoxAdjust.left % GPixelFormats[SourceTexture->GetFormat()].BlockSizeX == 0);
-	check(SourceBoxAdjust.top % GPixelFormats[SourceTexture->GetFormat()].BlockSizeY == 0);
-	check((SourceBoxAdjust.right - SourceBoxAdjust.left) % GPixelFormats[SourceTexture->GetFormat()].BlockSizeX == 0);
-	check((SourceBoxAdjust.bottom - SourceBoxAdjust.top) % GPixelFormats[SourceTexture->GetFormat()].BlockSizeY == 0);
-	check(uint32(DestinationBox.Min.X + DestinationOffsetX) % GPixelFormats[DestinationTexture->GetFormat()].BlockSizeX == 0);
-	check(uint32(DestinationBox.Min.Y + DestinationOffsetY) % GPixelFormats[DestinationTexture->GetFormat()].BlockSizeY == 0);
-
-	ID3D11Texture2D* DestinationRessource = DestinationTexture->GetResource();
-	Direct3DDeviceIMContext->CopySubresourceRegion(DestinationRessource, 0, DestinationBox.Min.X + DestinationOffsetX, DestinationBox.Min.Y + DestinationOffsetY, 0, SourceTexture->GetResource(), 0, &SourceBoxAdjust);
-}
-
-void FD3D11DynamicRHI::RHICopySubTextureRegion_RenderThread(
-	class FRHICommandListImmediate& RHICmdList,
-	FRHITexture2D* SourceTexture,
-	FRHITexture2D* DestinationTexture,
-	FBox2D SourceBox,
-	FBox2D DestinationBox)
-{
-	if (RHICmdList.Bypass())
-	{
-		RHICopySubTextureRegion(SourceTexture, DestinationTexture, SourceBox, DestinationBox);
-	}
-	else
-	{
-		RunOnRHIThread([this, SourceTexture, DestinationTexture, SourceBox, DestinationBox]()
-		{
-			RHICopySubTextureRegion(SourceTexture, DestinationTexture, SourceBox, DestinationBox);
-		});
-	}
-}
-
 template<typename BaseResourceType>
 TD3D11Texture2D<BaseResourceType>* FD3D11DynamicRHI::CreateTextureFromResource(bool bTextureArray, bool bCubeTexture, EPixelFormat Format, ETextureCreateFlags TexCreateFlags, const FClearValueBinding& ClearValueBinding, ID3D11Texture2D* TextureResource)
 {
@@ -2529,7 +2428,7 @@ TD3D11Texture2D<BaseResourceType>* FD3D11DynamicRHI::CreateTextureFromResource(b
 }
 
 template <typename BaseResourceType>
-TD3D11Texture2D<BaseResourceType>* FD3D11DynamicRHI::CreateAliasedD3D11Texture2D(TD3D11Texture2D<BaseResourceType>* SourceTexture)
+FTextureRHIRef FD3D11DynamicRHI::CreateAliasedD3D11Texture2D(TD3D11Texture2D<BaseResourceType>* SourceTexture)
 {
 	D3D11_TEXTURE2D_DESC TextureDesc;
 	SourceTexture->GetResource()->GetDesc(&TextureDesc);
@@ -2605,11 +2504,7 @@ TD3D11Texture2D<BaseResourceType>* FD3D11DynamicRHI::CreateAliasedD3D11Texture2D
 	// We'll be the same size, since we're the same thing. Avoid the check in D3D11Resources.h (AliasResources).
 	Texture2D->SetMemorySize(SourceTexture->GetMemorySize());
 
-	// Disable deprecation warning; when the DynamicRHI raw-pointer method is fully deprecated, the D3D11 class will still provide a raw pointer version
-	// since this is required in this path.
-	PRAGMA_DISABLE_DEPRECATION_WARNINGS
-	RHIAliasTextureResources(Texture2D, SourceTexture);
-	PRAGMA_ENABLE_DEPRECATION_WARNINGS
+	Texture2D->AliasResources(SourceTexture);
 
 	return Texture2D;
 }
@@ -2630,7 +2525,7 @@ FTextureCubeRHIRef FD3D11DynamicRHI::RHICreateTextureCubeFromResource(EPixelForm
 	return CreateTextureFromResource<FD3D11BaseTextureCube>(false, true, Format, TexCreateFlags, ClearValueBinding, TextureResource);
 }
 
-void FD3D11DynamicRHI::RHIAliasTextureResources(FRHITexture* DestTextureRHI, FRHITexture* SrcTextureRHI)
+void FD3D11DynamicRHI::RHIAliasTextureResources(FTextureRHIRef& DestTextureRHI, FTextureRHIRef& SrcTextureRHI)
 {
 	FD3D11TextureBase* DestTexture = GetD3D11TextureFromRHITexture(DestTextureRHI);
 	FD3D11TextureBase* SrcTexture = GetD3D11TextureFromRHITexture(SrcTextureRHI);
@@ -2641,7 +2536,7 @@ void FD3D11DynamicRHI::RHIAliasTextureResources(FRHITexture* DestTextureRHI, FRH
 	}
 }
 
-FTextureRHIRef FD3D11DynamicRHI::RHICreateAliasedTexture(FRHITexture* SourceTexture)
+FTextureRHIRef FD3D11DynamicRHI::RHICreateAliasedTexture(FTextureRHIRef& SourceTexture)
 {
 	if (SourceTexture->GetTexture2D() != nullptr)
 	{
@@ -2660,22 +2555,6 @@ FTextureRHIRef FD3D11DynamicRHI::RHICreateAliasedTexture(FRHITexture* SourceText
 	return nullptr;
 }
 
-PRAGMA_DISABLE_DEPRECATION_WARNINGS
-
-void FD3D11DynamicRHI::RHIAliasTextureResources(FTextureRHIRef& DestTextureRHI, FTextureRHIRef& SrcTextureRHI)
-{
-	// @todo: Move the raw-pointer implementation down here when it's deprecation is completed.
-	RHIAliasTextureResources((FRHITexture*)DestTextureRHI, (FRHITexture*)SrcTextureRHI);
-}
-
-FTextureRHIRef FD3D11DynamicRHI::RHICreateAliasedTexture(FTextureRHIRef& SourceTexture)
-{
-	// @todo: Move the raw-pointer implementation down here when it's deprecation is completed.
-	return RHICreateAliasedTexture((FRHITexture*)SourceTexture);
-}
-
-PRAGMA_ENABLE_DEPRECATION_WARNINGS
-
 void FD3D11DynamicRHI::RHICopyTexture(FRHITexture* SourceTextureRHI, FRHITexture* DestTextureRHI, const FRHICopyTextureInfo& CopyInfo)
 {
 	if (!SourceTextureRHI || !DestTextureRHI || SourceTextureRHI == DestTextureRHI)
