Exclude Oculus from blanket disable of Niagara GPUSorting for GL

diff --git a/Engine/Source/Runtime/OpenGLDrv/Private/OpenGLDevice.cpp b/Engine/Source/Runtime/OpenGLDrv/Private/OpenGLDevice.cpp
--- a/Engine/Source/Runtime/OpenGLDrv/Private/OpenGLDevice.cpp
+++ b/Engine/Source/Runtime/OpenGLDrv/Private/OpenGLDevice.cpp
@@ -26,6 +26,7 @@
 
 #if PLATFORM_ANDROID
 #include <jni.h>
+extern bool AndroidThunkCpp_IsOculusMobileApplication();
 #endif
 
 #ifndef GL_STEREO
@@ -1082,7 +1083,12 @@ FOpenGLDynamicRHI::FOpenGLDynamicRHI()
 		auto* CVar = IConsoleManager::Get().FindConsoleVariable(TEXT("FX.AllowGPUSorting"));
 		if (CVar)
 		{
-			CVar->Set(false);
+#if PLATFORM_ANDROID
+			if(!AndroidThunkCpp_IsOculusMobileApplication())
+#endif
+			{
+				CVar->Set(false);
+			}
 		}
 	}
 
