Moving Vulkan and OpenGL implementations of ApplyStaticUniformBuffers out of headers to remove header depenency on RHICore.

diff --git a/Engine/Source/Runtime/OpenGLDrv/Private/OpenGLCommands.cpp b/Engine/Source/Runtime/OpenGLDrv/Private/OpenGLCommands.cpp
--- a/Engine/Source/Runtime/OpenGLDrv/Private/OpenGLCommands.cpp
+++ b/Engine/Source/Runtime/OpenGLDrv/Private/OpenGLCommands.cpp
@@ -15,6 +15,7 @@
 #include "OpenGLDrv.h"
 #include "OpenGLDrvPrivate.h"
 #include "RenderUtils.h"
+#include "RHICoreShader.h"
 
 /*
 #define DECLARE_ISBOUNDSHADER(ShaderType) template <typename TShaderType> inline void ValidateBoundShader(TRefCountPtr<FOpenGLBoundShaderState> InBoundShaderState, FRHIGraphicsShader* GfxShader, TShaderType* ShaderTypeRHI) \
@@ -2810,6 +2811,51 @@ uint32 FOpenGLDynamicRHI::RHIGetGPUFrameCycles(uint32 GPUIndex)
 	return GGPUFrameTime;
 }
 
+template <typename TRHIShader, typename TRHIProxyShader>
+void FOpenGLDynamicRHI::ApplyStaticUniformBuffers(TRHIShader* Shader, TRHIProxyShader* ProxyShader)
+{
+	if (ProxyShader)
+	{
+		check(Shader);
+		UE::RHICore::ApplyStaticUniformBuffers(this, Shader, ProxyShader->StaticSlots, ProxyShader->Bindings.ShaderResourceTable.ResourceTableLayoutHashes, GlobalUniformBuffers);
+	}
+}
+
+void FOpenGLDynamicRHI::RHISetGraphicsPipelineState(FRHIGraphicsPipelineState* GraphicsState, uint32 StencilRef, bool bApplyAdditionalState)
+{
+	FRHIGraphicsPipelineStateFallBack* FallbackGraphicsState = static_cast<FRHIGraphicsPipelineStateFallBack*>(GraphicsState);
+
+	auto& PsoInit = FallbackGraphicsState->Initializer;
+
+	RHISetBoundShaderState(
+		RHICreateBoundShaderState_internal(
+			PsoInit.BoundShaderState.VertexDeclarationRHI,
+			PsoInit.BoundShaderState.VertexShaderRHI,
+			PsoInit.BoundShaderState.PixelShaderRHI,
+			PsoInit.BoundShaderState.GetGeometryShader(),
+			PsoInit.bFromPSOFileCache
+		).GetReference()
+	);
+
+	RHISetDepthStencilState(FallbackGraphicsState->Initializer.DepthStencilState, StencilRef);
+	RHISetRasterizerState(FallbackGraphicsState->Initializer.RasterizerState);
+	RHISetBlendState(FallbackGraphicsState->Initializer.BlendState, FLinearColor(1.0f, 1.0f, 1.0f));
+	if (GSupportsDepthBoundsTest)
+	{
+		RHIEnableDepthBoundsTest(FallbackGraphicsState->Initializer.bDepthBounds);
+	}
+
+	if (bApplyAdditionalState)
+	{
+		ApplyStaticUniformBuffers(PsoInit.BoundShaderState.VertexShaderRHI, ResourceCast(PsoInit.BoundShaderState.VertexShaderRHI));
+		ApplyStaticUniformBuffers(PsoInit.BoundShaderState.GetGeometryShader(), ResourceCast(PsoInit.BoundShaderState.GetGeometryShader()));
+		ApplyStaticUniformBuffers(PsoInit.BoundShaderState.PixelShaderRHI, ResourceCast(PsoInit.BoundShaderState.PixelShaderRHI));
+	}
+
+	// Store the PSO's primitive (after since IRHICommandContext::RHISetGraphicsPipelineState sets the BSS)
+	PrimitiveType = PsoInit.PrimitiveType;
+}
+
 void FOpenGLDynamicRHI::RHISetComputeShader(FRHIComputeShader* ComputeShaderRHI)
 {
 	if (OpenGLConsoleVariables::bSkipCompute)
diff --git a/Engine/Source/Runtime/OpenGLDrv/Public/OpenGLDrv.h b/Engine/Source/Runtime/OpenGLDrv/Public/OpenGLDrv.h
--- a/Engine/Source/Runtime/OpenGLDrv/Public/OpenGLDrv.h
+++ b/Engine/Source/Runtime/OpenGLDrv/Public/OpenGLDrv.h
@@ -15,7 +15,6 @@
 #include "GPUProfiler.h"
 #include "RenderResource.h"
 #include "Templates/EnableIf.h"
-#include "RHICoreShader.h"
 
 // @todo platplug: Replace all of these includes with a call to COMPILED_PLATFORM_HEADER(OpenGLDrvPrivate.h)
 //TODO: Move these to OpenGLDrvPrivate.h
@@ -925,40 +924,7 @@ public:
 		RHITHREAD_GLCOMMAND_EPILOGUE();
 	}
 
-	virtual void RHISetGraphicsPipelineState(FRHIGraphicsPipelineState* GraphicsState, uint32 StencilRef, bool bApplyAdditionalState) final override
-	{
-		FRHIGraphicsPipelineStateFallBack* FallbackGraphicsState = static_cast<FRHIGraphicsPipelineStateFallBack*>(GraphicsState);
-
-		auto& PsoInit = FallbackGraphicsState->Initializer;
-
-		RHISetBoundShaderState(
-			RHICreateBoundShaderState_internal(
-				PsoInit.BoundShaderState.VertexDeclarationRHI,
-				PsoInit.BoundShaderState.VertexShaderRHI,
-				PsoInit.BoundShaderState.PixelShaderRHI,
-				PsoInit.BoundShaderState.GetGeometryShader(),
-				PsoInit.bFromPSOFileCache
-			).GetReference()
-		);
-
-		RHISetDepthStencilState(FallbackGraphicsState->Initializer.DepthStencilState, StencilRef);
-		RHISetRasterizerState(FallbackGraphicsState->Initializer.RasterizerState);
-		RHISetBlendState(FallbackGraphicsState->Initializer.BlendState, FLinearColor(1.0f, 1.0f, 1.0f));
-		if (GSupportsDepthBoundsTest)
-		{
-			RHIEnableDepthBoundsTest(FallbackGraphicsState->Initializer.bDepthBounds);
-		}
-
-		if (bApplyAdditionalState)
-		{
-			ApplyStaticUniformBuffers(PsoInit.BoundShaderState.VertexShaderRHI, ResourceCast(PsoInit.BoundShaderState.VertexShaderRHI));
-			ApplyStaticUniformBuffers(PsoInit.BoundShaderState.GetGeometryShader(), ResourceCast(PsoInit.BoundShaderState.GetGeometryShader()));
-			ApplyStaticUniformBuffers(PsoInit.BoundShaderState.PixelShaderRHI, ResourceCast(PsoInit.BoundShaderState.PixelShaderRHI));
-		}
-
-		// Store the PSO's primitive (after since IRHICommandContext::RHISetGraphicsPipelineState sets the BSS)
-		PrimitiveType = PsoInit.PrimitiveType;
-	}
+	virtual void RHISetGraphicsPipelineState(FRHIGraphicsPipelineState* GraphicsState, uint32 StencilRef, bool bApplyAdditionalState) final override;
 
 	FBoundShaderStateRHIRef RHICreateBoundShaderState_internal(
 		FRHIVertexDeclaration* VertexDeclarationRHI,
@@ -1039,14 +1005,7 @@ private:
 	int32 BeginSceneContextType;
 
 	template <typename TRHIShader, typename TRHIProxyShader>
-	void ApplyStaticUniformBuffers(TRHIShader* Shader, TRHIProxyShader* ProxyShader)
-	{
-		if (ProxyShader)
-		{
-			check(Shader);
-			UE::RHICore::ApplyStaticUniformBuffers(this, Shader, ProxyShader->StaticSlots, ProxyShader->Bindings.ShaderResourceTable.ResourceTableLayoutHashes, GlobalUniformBuffers);
-		}
-	}
+	void ApplyStaticUniformBuffers(TRHIShader* Shader, TRHIProxyShader* ProxyShader);
 
 	TArray<FRHIUniformBuffer*> GlobalUniformBuffers;
 
diff --git a/Engine/Source/Runtime/VulkanRHI/Private/VulkanCommands.cpp b/Engine/Source/Runtime/VulkanRHI/Private/VulkanCommands.cpp
--- a/Engine/Source/Runtime/VulkanRHI/Private/VulkanCommands.cpp
+++ b/Engine/Source/Runtime/VulkanRHI/Private/VulkanCommands.cpp
@@ -115,31 +115,6 @@ void FVulkanCommandListContext::RHISetStreamSource(uint32 StreamIndex, FRHIBuffe
 	}
 }
 
-void FVulkanCommandListContext::RHISetComputePipelineState(FRHIComputePipelineState* ComputePipelineState)
-{
-	FVulkanCmdBuffer* CmdBuffer = CommandBufferManager->GetActiveCmdBuffer();
-	if (CmdBuffer->IsInsideRenderPass())
-	{
-		if (GVulkanSubmitAfterEveryEndRenderPass)
-		{
-			CommandBufferManager->SubmitActiveCmdBuffer();
-			CommandBufferManager->PrepareForNewActiveCommandBuffer();
-			CmdBuffer = CommandBufferManager->GetActiveCmdBuffer();
-		}
-	}
-
-	if (CmdBuffer->CurrentDescriptorPoolSetContainer == nullptr)
-	{
-		CmdBuffer->CurrentDescriptorPoolSetContainer = &Device->GetDescriptorPoolsManager().AcquirePoolSetContainer();
-	}
-
-	//#todo-rco: Set PendingGfx to null
-	FVulkanComputePipeline* ComputePipeline = ResourceCast(ComputePipelineState);
-	PendingComputeState->SetComputePipeline(ComputePipeline);
-
-	ApplyStaticUniformBuffers(const_cast<FVulkanComputeShader*>(ComputePipeline->GetShader()));
-}
-
 void FVulkanCommandListContext::RHIDispatchComputeShader(uint32 ThreadGroupCountX, uint32 ThreadGroupCountY, uint32 ThreadGroupCountZ)
 {
 #if VULKAN_ENABLE_AGGRESSIVE_STATS
diff --git a/Engine/Source/Runtime/VulkanRHI/Private/VulkanContext.h b/Engine/Source/Runtime/VulkanRHI/Private/VulkanContext.h
--- a/Engine/Source/Runtime/VulkanRHI/Private/VulkanContext.h
+++ b/Engine/Source/Runtime/VulkanRHI/Private/VulkanContext.h
@@ -9,7 +9,6 @@
 #include "VulkanResources.h"
 #include "VulkanGPUProfiler.h"
 #include "VulkanBarriers.h"
-#include "RHICoreShader.h"
 
 class FVulkanDevice;
 class FVulkanCommandBufferManager;
@@ -254,30 +253,7 @@ private:
 	FVulkanGPUTiming* FrameTiming;
 
 	template <typename TRHIShader>
-	void ApplyStaticUniformBuffers(TRHIShader* Shader)
-	{
-		if (Shader)
-		{
-			const auto& StaticSlots = Shader->StaticSlots;
-			const auto& UBInfos = Shader->GetCodeHeader().UniformBuffers;
-
-			for (int32 BufferIndex = 0; BufferIndex < StaticSlots.Num(); ++BufferIndex)
-			{
-				const FUniformBufferStaticSlot Slot = StaticSlots[BufferIndex];
-
-				if (IsUniformBufferStaticSlotValid(Slot))
-				{
-					FRHIUniformBuffer* Buffer = GlobalUniformBuffers[Slot];
-					UE::RHICore::ValidateStaticUniformBuffer(Buffer, Slot, UBInfos[BufferIndex].LayoutHash);
-
-					if (Buffer)
-					{
-						RHISetShaderUniformBuffer(Shader, BufferIndex, Buffer);
-					}
-				}
-			}
-		}
-	}
+	void ApplyStaticUniformBuffers(TRHIShader* Shader);
 
 	TArray<FRHIUniformBuffer*> GlobalUniformBuffers;
 
diff --git a/Engine/Source/Runtime/VulkanRHI/Private/VulkanPipelineState.cpp b/Engine/Source/Runtime/VulkanRHI/Private/VulkanPipelineState.cpp
--- a/Engine/Source/Runtime/VulkanRHI/Private/VulkanPipelineState.cpp
+++ b/Engine/Source/Runtime/VulkanRHI/Private/VulkanPipelineState.cpp
@@ -12,6 +12,7 @@
 #include "VulkanPendingState.h"
 #include "VulkanPipeline.h"
 #include "VulkanLLM.h"
+#include "RHICoreShader.h"
 
 enum
 {
@@ -347,6 +348,31 @@ bool FVulkanGraphicsPipelineDescriptorState::InternalUpdateDescriptorSets(FVulka
 	return true;
 }
 
+template <typename TRHIShader>
+void FVulkanCommandListContext::ApplyStaticUniformBuffers(TRHIShader* Shader)
+{
+	if (Shader)
+	{
+		const auto& StaticSlots = Shader->StaticSlots;
+		const auto& UBInfos = Shader->GetCodeHeader().UniformBuffers;
+
+		for (int32 BufferIndex = 0; BufferIndex < StaticSlots.Num(); ++BufferIndex)
+		{
+			const FUniformBufferStaticSlot Slot = StaticSlots[BufferIndex];
+
+			if (IsUniformBufferStaticSlotValid(Slot))
+			{
+				FRHIUniformBuffer* Buffer = GlobalUniformBuffers[Slot];
+				UE::RHICore::ValidateStaticUniformBuffer(Buffer, Slot, UBInfos[BufferIndex].LayoutHash);
+
+				if (Buffer)
+				{
+					RHISetShaderUniformBuffer(Shader, BufferIndex, Buffer);
+				}
+			}
+		}
+	}
+}
 
 void FVulkanCommandListContext::RHISetGraphicsPipelineState(FRHIGraphicsPipelineState* GraphicsState, uint32 StencilRef, bool bApplyAdditionalState)
 {
@@ -380,6 +406,31 @@ void FVulkanCommandListContext::RHISetGraphicsPipelineState(FRHIGraphicsPipeline
 	}
 }
 
+void FVulkanCommandListContext::RHISetComputePipelineState(FRHIComputePipelineState* ComputePipelineState)
+{
+	FVulkanCmdBuffer* CmdBuffer = CommandBufferManager->GetActiveCmdBuffer();
+	if (CmdBuffer->IsInsideRenderPass())
+	{
+		if (GVulkanSubmitAfterEveryEndRenderPass)
+		{
+			CommandBufferManager->SubmitActiveCmdBuffer();
+			CommandBufferManager->PrepareForNewActiveCommandBuffer();
+			CmdBuffer = CommandBufferManager->GetActiveCmdBuffer();
+		}
+	}
+
+	if (CmdBuffer->CurrentDescriptorPoolSetContainer == nullptr)
+	{
+		CmdBuffer->CurrentDescriptorPoolSetContainer = &Device->GetDescriptorPoolsManager().AcquirePoolSetContainer();
+	}
+
+	//#todo-rco: Set PendingGfx to null
+	FVulkanComputePipeline* ComputePipeline = ResourceCast(ComputePipelineState);
+	PendingComputeState->SetComputePipeline(ComputePipeline);
+
+	ApplyStaticUniformBuffers(const_cast<FVulkanComputeShader*>(ComputePipeline->GetShader()));
+}
+
 
 template bool FVulkanGraphicsPipelineDescriptorState::InternalUpdateDescriptorSets<true>(FVulkanCommandListContext* CmdListContext, FVulkanCmdBuffer* CmdBuffer);
 template bool FVulkanGraphicsPipelineDescriptorState::InternalUpdateDescriptorSets<false>(FVulkanCommandListContext* CmdListContext, FVulkanCmdBuffer* CmdBuffer);
