Adding minimal interface to OpenGLDrv
- Adding IOpenGLDynamicRHI interface for plugins that want to touch OpenGL resources/handles directly.
- Modifying plugins that were using OpenGLDrv.h to use the interface instead.
- Removing plugin references to UE Modules and ThirdParty libraries that were only needed because they included the private RHI headers.
- Removing OpenGLDrv/Private from plugin include paths.

diff --git a/Engine/Plugins/Runtime/AR/Google/GoogleARCore/Source/GoogleARCoreBase/GoogleARCoreBase.Build.cs b/Engine/Plugins/Runtime/AR/Google/GoogleARCore/Source/GoogleARCoreBase/GoogleARCoreBase.Build.cs
--- a/Engine/Plugins/Runtime/AR/Google/GoogleARCore/Source/GoogleARCoreBase/GoogleARCoreBase.Build.cs
+++ b/Engine/Plugins/Runtime/AR/Google/GoogleARCore/Source/GoogleARCoreBase/GoogleARCoreBase.Build.cs
@@ -64,10 +64,6 @@ namespace UnrealBuildTool.Rules
 				// Register Plugin Language
 				string PluginPath = Utils.MakePathRelativeTo(ModuleDirectory, Target.RelativeEnginePath);
 				AdditionalPropertiesForReceipt.Add("AndroidPlugin", Path.Combine(PluginPath, "GoogleARCoreBase_APL.xml"));
-				
-				// Needed for including "AndroidEGL.h"
-				string EnginePath = Path.GetFullPath(Target.RelativeEnginePath);
-				PrivateIncludePaths.Add(Path.Combine(EnginePath, "Source/Runtime/OpenGLDrv/Private"));
 			}
 
 			if (Target.bBuildEditor)
diff --git a/Engine/Plugins/Runtime/AR/Google/GoogleARCore/Source/GoogleARCoreBase/Private/GoogleARCoreOpenGLContext.cpp b/Engine/Plugins/Runtime/AR/Google/GoogleARCore/Source/GoogleARCoreBase/Private/GoogleARCoreOpenGLContext.cpp
--- a/Engine/Plugins/Runtime/AR/Google/GoogleARCore/Source/GoogleARCoreBase/Private/GoogleARCoreOpenGLContext.cpp
+++ b/Engine/Plugins/Runtime/AR/Google/GoogleARCore/Source/GoogleARCoreBase/Private/GoogleARCoreOpenGLContext.cpp
@@ -30,12 +30,12 @@ void FGoogleARCoreOpenGLContext::InitContext()
 #if PLATFORM_ANDROID
 	UE_LOG(LogARCoreOpenGL, Log, TEXT("FGoogleARCoreOpenGLContext::InitContext"));
 	
-	auto EGL = AndroidEGL::GetInstance();
-	EGL->Init(AndroidEGL::AV_OpenGLES, 2, 0, false);
-	EGL->InitSurface(false, false);
+	IOpenGLDynamicRHI* RHI = GetIOpenGLDynamicRHI();
+
+	RHI->RHIInitEGLInstanceGLES2();
 	SaveContext();
 	MakeCurrent();
-	EGL->InitBackBuffer(); //can be done only after context is made current.
+	RHI->RHIInitEGLBackBuffer(); //can be done only after context is made current.
 	glGenTextures(1, &PassthroughCameraTextureId);
 	RestoreContext();
 	
@@ -48,14 +48,14 @@ void FGoogleARCoreOpenGLContext::ReleaseContext()
 #if PLATFORM_ANDROID
 	glDeleteTextures(1, &PassthroughCameraTextureId);
 	PassthroughCameraTextureId = 0;
-	AndroidEGL::GetInstance()->Terminate();
+	GetIOpenGLDynamicRHI()->RHIEGLTerminateContext();
 #endif
 }
 
 void FGoogleARCoreOpenGLContext::MakeCurrent()
 {
 #if PLATFORM_ANDROID
-	AndroidEGL::GetInstance()->SetCurrentRenderingContext();
+	GetIOpenGLDynamicRHI()->RHIEGLSetCurrentRenderingContext();
 #endif
 }
 
diff --git a/Engine/Plugins/Runtime/AR/Google/GoogleARCore/Source/GoogleARCoreBase/Private/GoogleARCoreOpenGLContext.h b/Engine/Plugins/Runtime/AR/Google/GoogleARCore/Source/GoogleARCoreBase/Private/GoogleARCoreOpenGLContext.h
--- a/Engine/Plugins/Runtime/AR/Google/GoogleARCore/Source/GoogleARCoreBase/Private/GoogleARCoreOpenGLContext.h
+++ b/Engine/Plugins/Runtime/AR/Google/GoogleARCore/Source/GoogleARCoreBase/Private/GoogleARCoreOpenGLContext.h
@@ -7,7 +7,7 @@
 #include "Templates/SharedPointer.h"
 
 #if PLATFORM_ANDROID
-#include "Android/AndroidEGL.h"
+#include "IOpenGLDynamicRHI.h"
 #endif
 
 class FGoogleARCoreOpenGLContext
diff --git a/Engine/Plugins/Runtime/Oculus/OculusVR/Source/OculusHMD/OculusHMD.Build.cs b/Engine/Plugins/Runtime/Oculus/OculusVR/Source/OculusHMD/OculusHMD.Build.cs
--- a/Engine/Plugins/Runtime/Oculus/OculusVR/Source/OculusHMD/OculusHMD.Build.cs
+++ b/Engine/Plugins/Runtime/Oculus/OculusVR/Source/OculusHMD/OculusHMD.Build.cs
@@ -13,7 +13,6 @@ namespace UnrealBuildTool.Rules
 				new string[] {
 					// Relative to Engine\Plugins\Runtime\Oculus\OculusVR\Source
 					"../../../../../Source/Runtime/Renderer/Private",
-					"../../../../../Source/Runtime/OpenGLDrv/Private",
 					"../../../../../Source/Runtime/VulkanRHI/Private",
 					"../../../../../Source/Runtime/Engine/Classes/Components",
 					"../../../../../Source/Runtime/Engine/Classes/Kismet",
diff --git a/Engine/Plugins/Runtime/Oculus/OculusVR/Source/OculusHMD/Private/OculusHMD.cpp b/Engine/Plugins/Runtime/Oculus/OculusVR/Source/OculusHMD/Private/OculusHMD.cpp
--- a/Engine/Plugins/Runtime/Oculus/OculusVR/Source/OculusHMD/Private/OculusHMD.cpp
+++ b/Engine/Plugins/Runtime/Oculus/OculusVR/Source/OculusHMD/Private/OculusHMD.cpp
@@ -33,8 +33,6 @@
 
 #if PLATFORM_ANDROID
 #include "Android/AndroidJNI.h"
-#include "Android/AndroidEGL.h"
-#include "Android/AndroidOpenGL.h"
 #include "Android/AndroidApplication.h"
 #include "HAL/IConsoleManager.h"
 #endif
@@ -2214,7 +2212,8 @@ namespace OculusHMD
 			GIsEditor ? ovrpBool_True : ovrpBool_False);
 
 #if PLATFORM_ANDROID
-		FOculusHMDModule::GetPluginWrapper().SetupDisplayObjects2(AndroidEGL::GetInstance()->GetRenderingContext()->eglContext, AndroidEGL::GetInstance()->GetDisplay(), AndroidEGL::GetInstance()->GetNativeWindow());
+		IOpenGLDynamicRHI* RHI = GetIOpenGLDynamicRHI();
+		FOculusHMDModule::GetPluginWrapper().SetupDisplayObjects2(RHI->RHIGetEGLContext(), RHI->RHIGetEGLDisplay(), RHI->RHIGetEGLNativeWindow());
 		GSupportsMobileMultiView = true;
 		if (GRHISupportsLateVariableRateShadingUpdate)
 		{
@@ -2223,7 +2222,7 @@ namespace OculusHMD
 #endif
 		int flag = ovrpDistortionWindowFlag_None;
 #if PLATFORM_ANDROID && USE_ANDROID_EGL_NO_ERROR_CONTEXT
-		if (AndroidEGL::GetInstance()->GetSupportsNoErrorContext())
+		if (RHI->RHIEGLSupportsNoErrorContext())
 		{
 			flag = ovrpDistortionWindowFlag_NoErrorContext;
 		}
diff --git a/Engine/Plugins/Runtime/Oculus/OculusVR/Source/OculusHMD/Private/OculusHMDPrivateRHI.h b/Engine/Plugins/Runtime/Oculus/OculusVR/Source/OculusHMD/Private/OculusHMDPrivateRHI.h
--- a/Engine/Plugins/Runtime/Oculus/OculusVR/Source/OculusHMD/Private/OculusHMDPrivateRHI.h
+++ b/Engine/Plugins/Runtime/Oculus/OculusVR/Source/OculusHMD/Private/OculusHMDPrivateRHI.h
@@ -27,8 +27,7 @@
 //-------------------------------------------------------------------------------------------------
 
 #if OCULUS_HMD_SUPPORTED_PLATFORMS_OPENGL
-#include "OpenGLDrvPrivate.h"
-#include "OpenGLResources.h"
+#include "IOpenGLDynamicRHI.h"
 #endif // OCULUS_HMD_SUPPORTED_PLATFORMS_OPENGL
 
 
diff --git a/Engine/Plugins/Runtime/Oculus/OculusVR/Source/OculusHMD/Private/OculusHMD_CustomPresent.cpp b/Engine/Plugins/Runtime/Oculus/OculusVR/Source/OculusHMD/Private/OculusHMD_CustomPresent.cpp
--- a/Engine/Plugins/Runtime/Oculus/OculusVR/Source/OculusHMD/Private/OculusHMD_CustomPresent.cpp
+++ b/Engine/Plugins/Runtime/Oculus/OculusVR/Source/OculusHMD/Private/OculusHMD_CustomPresent.cpp
@@ -12,7 +12,6 @@
 
 #if PLATFORM_ANDROID
 #include "Android/AndroidJNI.h"
-#include "Android/AndroidEGL.h"
 #include "Android/AndroidApplication.h"
 #include "Android/AndroidPlatformMisc.h"
 #endif
diff --git a/Engine/Plugins/Runtime/Oculus/OculusVR/Source/OculusHMD/Private/OculusHMD_CustomPresent_OpenGL.cpp b/Engine/Plugins/Runtime/Oculus/OculusVR/Source/OculusHMD/Private/OculusHMD_CustomPresent_OpenGL.cpp
--- a/Engine/Plugins/Runtime/Oculus/OculusVR/Source/OculusHMD/Private/OculusHMD_CustomPresent_OpenGL.cpp
+++ b/Engine/Plugins/Runtime/Oculus/OculusVR/Source/OculusHMD/Private/OculusHMD_CustomPresent_OpenGL.cpp
@@ -12,10 +12,6 @@
 #endif
 #endif
 
-#if PLATFORM_ANDROID
-#include "Android/AndroidOpenGL.h"
-#endif
-
 namespace OculusHMD
 {
 
@@ -55,7 +51,7 @@ FTextureRHIRef FOpenGLCustomPresent::CreateTexture_RenderThread(uint32 InSizeX,
 {
 	CheckInRenderThread();
 
-	FOpenGLDynamicRHI* DynamicRHI = GetDynamicRHI<FOpenGLDynamicRHI>();
+	IOpenGLDynamicRHI* DynamicRHI = GetIOpenGLDynamicRHI();
 
 	switch (InResourceType)
 	{
@@ -75,8 +71,7 @@ FTextureRHIRef FOpenGLCustomPresent::CreateTexture_RenderThread(uint32 InSizeX,
 
 void FOpenGLCustomPresent::SubmitGPUFrameTime(float GPUFrameTime)
 {
-	FOpenGLDynamicRHI* DynamicRHI = GetDynamicRHI<FOpenGLDynamicRHI>();
-	DynamicRHI->GetGPUProfilingData().ExternalGPUTime = GPUFrameTime * 1000;
+	GetIOpenGLDynamicRHI()->RHISetExternalGPUTime(GPUFrameTime * 1000);
 }
 
 
@@ -86,12 +81,7 @@ void FOpenGLCustomPresent::SubmitGPUFrameTime(float GPUFrameTime)
 
 FCustomPresent* CreateCustomPresent_OpenGL(FOculusHMD* InOculusHMD)
 {
-#if PLATFORM_ANDROID
-	const bool sRGBSupport = FAndroidOpenGL::SupportsFramebufferSRGBEnable();
-#else
-	const bool sRGBSupport = true;
-#endif
-	return new FOpenGLCustomPresent(InOculusHMD, sRGBSupport);
+	return new FOpenGLCustomPresent(InOculusHMD, GetIOpenGLDynamicRHI()->RHISupportsFramebufferSRGBEnable());
 }
 
 
diff --git a/Engine/Plugins/Runtime/Oculus/OculusVR/Source/OculusHMD/Private/OculusHMD_Splash.cpp b/Engine/Plugins/Runtime/Oculus/OculusVR/Source/OculusHMD/Private/OculusHMD_Splash.cpp
--- a/Engine/Plugins/Runtime/Oculus/OculusVR/Source/OculusHMD/Private/OculusHMD_Splash.cpp
+++ b/Engine/Plugins/Runtime/Oculus/OculusVR/Source/OculusHMD/Private/OculusHMD_Splash.cpp
@@ -10,7 +10,6 @@
 #include "StereoLayerFunctionLibrary.h"
 #if PLATFORM_ANDROID
 #include "Android/AndroidJNI.h"
-#include "Android/AndroidEGL.h"
 #include "Android/AndroidApplication.h"
 #include "OculusHMDTypes.h"
 #endif
diff --git a/Engine/Plugins/Runtime/Oculus/OculusVR/Source/OculusOpenXRHMD/OculusOpenXRHMD.Build.cs b/Engine/Plugins/Runtime/Oculus/OculusVR/Source/OculusOpenXRHMD/OculusOpenXRHMD.Build.cs
--- a/Engine/Plugins/Runtime/Oculus/OculusVR/Source/OculusOpenXRHMD/OculusOpenXRHMD.Build.cs
+++ b/Engine/Plugins/Runtime/Oculus/OculusVR/Source/OculusOpenXRHMD/OculusOpenXRHMD.Build.cs
@@ -14,7 +14,6 @@ namespace UnrealBuildTool.Rules
                     // Relative to Engine\Plugins\Runtime\Oculus\OculusOpenXR\Source
                     "../../../OpenXR/Source/OpenXRHMD/Private",
                     "../../../../../Source/Runtime/Renderer/Private",
-                    "../../../../../Source/Runtime/OpenGLDrv/Private",
                     "../../../../../Source/Runtime/VulkanRHI/Private",
                     "../../../../../Source/Runtime/Engine/Classes/Components",
                     "../../../../../Source/Runtime/Engine/Classes/Kismet",
diff --git a/Engine/Plugins/Runtime/OpenXR/Source/OpenXRHMD/OpenXRHMD.Build.cs b/Engine/Plugins/Runtime/OpenXR/Source/OpenXRHMD/OpenXRHMD.Build.cs
--- a/Engine/Plugins/Runtime/OpenXR/Source/OpenXRHMD/OpenXRHMD.Build.cs
+++ b/Engine/Plugins/Runtime/OpenXR/Source/OpenXRHMD/OpenXRHMD.Build.cs
@@ -15,7 +15,6 @@ namespace UnrealBuildTool.Rules
 					"OpenXRHMD/Private",
                     EngineDir + "/Source/ThirdParty/OpenXR/include",
                     EngineDir + "/Source/Runtime/Renderer/Private",
-                    EngineDir + "/Source/Runtime/OpenGLDrv/Private",
                     EngineDir + "/Source/Runtime/VulkanRHI/Private",
 					// ... add other private include paths required here ...
 				}
diff --git a/Engine/Plugins/Runtime/OpenXR/Source/OpenXRHMD/Private/OpenXRHMD_RenderBridge.cpp b/Engine/Plugins/Runtime/OpenXR/Source/OpenXRHMD/Private/OpenXRHMD_RenderBridge.cpp
--- a/Engine/Plugins/Runtime/OpenXR/Source/OpenXRHMD/Private/OpenXRHMD_RenderBridge.cpp
+++ b/Engine/Plugins/Runtime/OpenXR/Source/OpenXRHMD/Private/OpenXRHMD_RenderBridge.cpp
@@ -116,7 +116,8 @@ public:
 		Requirements.maxApiVersionSupported = 0;
 		XR_ENSURE(GetOpenGLGraphicsRequirementsKHR(InInstance, InSystem, &Requirements));
 
-		XrVersion RHIVersion = XR_MAKE_VERSION(FOpenGL::GetMajorVersion(), FOpenGL::GetMinorVersion(), 0);
+		IOpenGLDynamicRHI* RHI = GetIOpenGLDynamicRHI();
+		XrVersion RHIVersion = XR_MAKE_VERSION(RHI->RHIGetGLMajorVersion(), RHI->RHIGetGLMinorVersion(), 0);
 		if (RHIVersion < Requirements.minApiVersionSupported) //-V547
 		{
 			UE_LOG(LogHMD, Fatal, TEXT("The OpenGL API version does not meet the minimum version required by the OpenXR runtime"));
@@ -172,7 +173,8 @@ public:
 		XR_ENSURE(GetOpenGLESGraphicsRequirementsKHR(InInstance, InSystem, &Requirements));
 
 #if PLATFORM_ANDROID
-		XrVersion RHIVersion = XR_MAKE_VERSION(FAndroidOpenGL::GLMajorVerion, FAndroidOpenGL::GLMinorVersion, 0);
+		IOpenGLDynamicRHI* RHI = GetIOpenGLDynamicRHI();
+		XrVersion RHIVersion = XR_MAKE_VERSION(RHI->RHIGetGLMajorVersion(), RHI->RHIGetGLMinorVersion(), 0);
 		if (RHIVersion < Requirements.minApiVersionSupported) //-V547
 		{
 			UE_LOG(LogHMD, Fatal, TEXT("The OpenGLES API version does not meet the minimum version required by the OpenXR runtime"));
@@ -188,11 +190,12 @@ public:
 	virtual void* GetGraphicsBinding() override
 	{
 #if PLATFORM_ANDROID
+		IOpenGLDynamicRHI* RHI = GetIOpenGLDynamicRHI();
 		Binding.type = XR_TYPE_GRAPHICS_BINDING_OPENGL_ES_ANDROID_KHR;
 		Binding.next = nullptr;
-		Binding.display = AndroidEGL::GetInstance()->GetDisplay();
-		Binding.config = AndroidEGL::GetInstance()->GetConfig();
-		Binding.context = AndroidEGL::GetInstance()->GetRenderingContext()->eglContext;
+		Binding.display = RHI->RHIGetEGLDisplay();
+		Binding.config = RHI->RHIGetEGLConfig();
+		Binding.context = RHI->RHIGetEGLContext();
 		return &Binding;
 #endif
 		return nullptr;
diff --git a/Engine/Plugins/Runtime/OpenXR/Source/OpenXRHMD/Private/OpenXRHMD_Swapchain.cpp b/Engine/Plugins/Runtime/OpenXR/Source/OpenXRHMD/Private/OpenXRHMD_Swapchain.cpp
--- a/Engine/Plugins/Runtime/OpenXR/Source/OpenXRHMD/Private/OpenXRHMD_Swapchain.cpp
+++ b/Engine/Plugins/Runtime/OpenXR/Source/OpenXRHMD/Private/OpenXRHMD_Swapchain.cpp
@@ -339,13 +339,14 @@ FXRSwapChainPtr CreateSwapchain_OpenGL(XrSession InSession, uint8 Format, uint32
 	}
 
 	TArray<FTextureRHIRef> TextureChain;
-	FOpenGLDynamicRHI* DynamicRHI = GetDynamicRHI<FOpenGLDynamicRHI>();
+	IOpenGLDynamicRHI* DynamicRHI = GetIOpenGLDynamicRHI();
 	TArray<XrSwapchainImageOpenGLKHR> Images = EnumerateImages<XrSwapchainImageOpenGLKHR>(Swapchain, XR_TYPE_SWAPCHAIN_IMAGE_OPENGL_KHR);
 	for (const auto& Image : Images)
 	{
-		TextureChain.Add(static_cast<FTextureRHIRef>(DynamicRHI->RHICreateTexture2DArrayFromResource(GPixelFormats[Format].UnrealFormat, SizeX, SizeY, ArraySize, NumMips, NumSamples, 1, ClearValueBinding, Image.image, CreateFlags)));
+		FTexture2DRHIRef NewTexture = DynamicRHI->RHICreateTexture2DArrayFromResource(GPixelFormats[Format].UnrealFormat, SizeX, SizeY, ArraySize, NumMips, NumSamples, 1, ClearValueBinding, Image.image, CreateFlags);
+		TextureChain.Add(NewTexture.GetReference());
 	}
-	FTextureRHIRef ChainTarget = static_cast<FTextureRHIRef>(GDynamicRHI->RHICreateAliasedTexture((FTextureRHIRef&)TextureChain[0]));
+	FTextureRHIRef ChainTarget = static_cast<FTextureRHIRef>(DynamicRHI->RHICreateAliasedTexture((FTextureRHIRef&)TextureChain[0]));
 
 	return CreateXRSwapChain<FOpenXRSwapchain>(MoveTemp(TextureChain), ChainTarget, Swapchain);
 }
@@ -367,13 +368,14 @@ FXRSwapChainPtr CreateSwapchain_OpenGLES(XrSession InSession, uint8 Format, uint
 	}
 
 	TArray<FTextureRHIRef> TextureChain;
-	FOpenGLDynamicRHI* DynamicRHI = GetDynamicRHI<FOpenGLDynamicRHI>();
+	IOpenGLDynamicRHI* DynamicRHI = GetIOpenGLDynamicRHI();
 	TArray<XrSwapchainImageOpenGLESKHR> Images = EnumerateImages<XrSwapchainImageOpenGLESKHR>(Swapchain, XR_TYPE_SWAPCHAIN_IMAGE_OPENGL_ES_KHR);
 	for (const auto& Image : Images)
 	{
-		TextureChain.Add(static_cast<FTextureRHIRef>(DynamicRHI->RHICreateTexture2DArrayFromResource(GPixelFormats[Format].UnrealFormat, SizeX, SizeY, ArraySize, NumMips, NumSamples, 1, ClearValueBinding, Image.image, CreateFlags)));
+		FTexture2DRHIRef NewTexture = DynamicRHI->RHICreateTexture2DArrayFromResource(GPixelFormats[Format].UnrealFormat, SizeX, SizeY, ArraySize, NumMips, NumSamples, 1, ClearValueBinding, Image.image, CreateFlags);
+		TextureChain.Add(NewTexture.GetReference());
 	}
-	FTextureRHIRef ChainTarget = static_cast<FTextureRHIRef>(GDynamicRHI->RHICreateAliasedTexture((FTextureRHIRef&)TextureChain[0]));
+	FTextureRHIRef ChainTarget = static_cast<FTextureRHIRef>(DynamicRHI->RHICreateAliasedTexture((FTextureRHIRef&)TextureChain[0]));
 
 	return CreateXRSwapChain<FOpenXRSwapchain>(MoveTemp(TextureChain), ChainTarget, Swapchain);
 }
diff --git a/Engine/Plugins/Runtime/OpenXR/Source/OpenXRHMD/Private/OpenXRPlatformRHI.h b/Engine/Plugins/Runtime/OpenXR/Source/OpenXRHMD/Private/OpenXRPlatformRHI.h
--- a/Engine/Plugins/Runtime/OpenXR/Source/OpenXRHMD/Private/OpenXRPlatformRHI.h
+++ b/Engine/Plugins/Runtime/OpenXR/Source/OpenXRHMD/Private/OpenXRPlatformRHI.h
@@ -44,14 +44,7 @@
 //-------------------------------------------------------------------------------------------------
 
 #if defined(XR_USE_GRAPHICS_API_OPENGL) || defined(XR_USE_GRAPHICS_API_OPENGL_ES)
-#include "OpenGLDrvPrivate.h"
-#include "OpenGLResources.h"
-
-#ifdef XR_USE_PLATFORM_ANDROID
-#include "Android/AndroidEGL.h"
-#include "Android/AndroidOpenGL.h"
-#endif
-
+#include "IOpenGLDynamicRHI.h"
 #endif // XR_USE_GRAPHICS_API_OPENGL
 
 //-------------------------------------------------------------------------------------------------
diff --git a/Engine/Plugins/Runtime/Steam/SteamVR/Source/SteamVR/Private/SteamVRHMD.h b/Engine/Plugins/Runtime/Steam/SteamVR/Source/SteamVR/Private/SteamVRHMD.h
--- a/Engine/Plugins/Runtime/Steam/SteamVR/Source/SteamVR/Private/SteamVRHMD.h
+++ b/Engine/Plugins/Runtime/Steam/SteamVR/Source/SteamVR/Private/SteamVRHMD.h
@@ -24,7 +24,7 @@
 #endif
 
 #if !PLATFORM_MAC // No OpenGL on Mac anymore
-#include "OpenGLDrv.h"
+#include "IOpenGLDynamicRHI.h"
 #endif
 
 #include "SceneViewExtension.h"
diff --git a/Engine/Plugins/Runtime/Steam/SteamVR/Source/SteamVR/Private/SteamVRRender.cpp b/Engine/Plugins/Runtime/Steam/SteamVR/Source/SteamVR/Private/SteamVRRender.cpp
--- a/Engine/Plugins/Runtime/Steam/SteamVR/Source/SteamVR/Private/SteamVRRender.cpp
+++ b/Engine/Plugins/Runtime/Steam/SteamVR/Source/SteamVR/Private/SteamVRRender.cpp
@@ -493,13 +493,15 @@ FSteamVRHMD::OpenGLBridge::OpenGLBridge(FSteamVRHMD* plugin):
 
 void FSteamVRHMD::OpenGLBridge::FinishRendering()
 {
-	GLuint RenderTargetTexture = *reinterpret_cast<GLuint*>(SwapChain->GetTexture2D()->GetNativeResource());
-	GLuint DepthTargetTexture = *reinterpret_cast<GLuint*>(DepthSwapChain->GetTexture2D()->GetNativeResource());
+	IOpenGLDynamicRHI* RHI = GetIOpenGLDynamicRHI();
+
+	const GLuint RenderTargetTexture = RHI->RHIGetResource(SwapChain->GetTexture2D());
+	const GLuint DepthTargetTexture = RHI->RHIGetResource(DepthSwapChain->GetTexture2D());
 
 	// Yaakuro:
 	// TODO This is a workaround. After exiting VR Editor the texture gets invalid at some point.
 	// Need to find it. This at least prevents to use this method when texture name is not valid anymore.
-	if (!glIsTexture(RenderTargetTexture) || !glIsTexture(DepthTargetTexture))
+	if (!RHI->RHIIsValidTexture(RenderTargetTexture) || !RHI->RHIIsValidTexture(DepthTargetTexture))
 	{
 		return;
 	}
diff --git a/Engine/Source/Runtime/OpenGLDrv/OpenGLDrv.Build.cs b/Engine/Source/Runtime/OpenGLDrv/OpenGLDrv.Build.cs
--- a/Engine/Source/Runtime/OpenGLDrv/OpenGLDrv.Build.cs
+++ b/Engine/Source/Runtime/OpenGLDrv/OpenGLDrv.Build.cs
@@ -9,30 +9,26 @@ public class OpenGLDrv : ModuleRules
 	{
 		PrivateIncludePaths.Add("Runtime/OpenGLDrv/Private");
 
-		PrivateDependencyModuleNames.AddRange(
-			new string[] {
-				"Core",
-				"CoreUObject",
-				"ApplicationCore",
-				"Engine",
-				"RHI",
-				"RHICore",
-				"RenderCore",
-				"PreLoadScreen"
-			}
-			);
+		PublicDependencyModuleNames.AddRange(new string[] {
+			"Core",
+			"RHI",
+		});
+
+		PrivateDependencyModuleNames.AddRange(new string[] {
+			"CoreUObject",
+			"ApplicationCore",
+			"Engine",
+			"RHICore",
+			"RenderCore",
+			"PreLoadScreen"
+		});
 
 		PrivateIncludePathModuleNames.Add("ImageWrapper");
 		DynamicallyLoadedModuleNames.Add("ImageWrapper");
 
+		PublicIncludePathModuleNames.Add("OpenGL");
 		AddEngineThirdPartyPrivateStaticDependencies(Target, "OpenGL");
 
-		if (Target.IsInPlatformGroup(UnrealPlatformGroup.Linux))
-		{
-			string GLPath = Target.UEThirdPartySourceDirectory + "OpenGL/";
-			PublicIncludePaths.Add(GLPath);
-		}
-
 		if (Target.IsInPlatformGroup(UnrealPlatformGroup.Linux))
 		{
 			AddEngineThirdPartyPrivateStaticDependencies(Target, "SDL2");
@@ -40,12 +36,7 @@ public class OpenGLDrv : ModuleRules
 
 		if (Target.Configuration != UnrealTargetConfiguration.Shipping)
 		{
-			PrivateIncludePathModuleNames.AddRange(
-				new string[]
-				{
-					"TaskGraph"
-				}
-			);
+			PrivateIncludePathModuleNames.Add("TaskGraph");
 		}
 
 		if (Target.Platform == UnrealTargetPlatform.Android)
@@ -58,19 +49,8 @@ public class OpenGLDrv : ModuleRules
             // for Swappy
             PublicDefinitions.Add("USE_ANDROID_OPENGL_SWAPPY=1");
 
-            PrivateDependencyModuleNames.AddRange(
-                new string[]
-                {
-					"GoogleGameSDK"
-                }
-            );
-
-			PrivateIncludePathModuleNames.AddRange(
-				new string[]
-				{
-					"Launch"
-				}
-			);
+            PrivateDependencyModuleNames.Add("GoogleGameSDK");
+			PrivateIncludePathModuleNames.Add("Launch");
 		}
 
         if (Target.Platform != UnrealTargetPlatform.Win64
diff --git a/Engine/Source/Runtime/OpenGLDrv/Private/Android/AndroidOpenGL.h b/Engine/Source/Runtime/OpenGLDrv/Private/Android/AndroidOpenGL.h
--- a/Engine/Source/Runtime/OpenGLDrv/Private/Android/AndroidOpenGL.h
+++ b/Engine/Source/Runtime/OpenGLDrv/Private/Android/AndroidOpenGL.h
@@ -181,6 +181,16 @@ struct FAndroidOpenGL : public FOpenGLES
 	static int32 GLMajorVerion;
 	static int32 GLMinorVersion;
 
+	static FORCEINLINE GLuint GetMajorVersion()
+	{
+		return GLMajorVerion;
+	}
+
+	static FORCEINLINE GLuint GetMinorVersion()
+	{
+		return GLMinorVersion;
+	}
+
 	/** Whether device supports image external */
 	static bool bSupportsImageExternal;
 
diff --git a/Engine/Source/Runtime/OpenGLDrv/Private/OpenGLDrv.cpp b/Engine/Source/Runtime/OpenGLDrv/Private/OpenGLDrv.cpp
--- a/Engine/Source/Runtime/OpenGLDrv/Private/OpenGLDrv.cpp
+++ b/Engine/Source/Runtime/OpenGLDrv/Private/OpenGLDrv.cpp
@@ -29,6 +29,92 @@ DEFINE_LOG_CATEGORY(LogOpenGL);
 ERHIFeatureLevel::Type GRequestedFeatureLevel = ERHIFeatureLevel::Num;
 
 
+int32 FOpenGLDynamicRHI::RHIGetGLMajorVersion() const
+{
+	return FOpenGL::GetMajorVersion();
+}
+
+int32 FOpenGLDynamicRHI::RHIGetGLMinorVersion() const
+{
+	return FOpenGL::GetMinorVersion();
+}
+
+bool FOpenGLDynamicRHI::RHISupportsFramebufferSRGBEnable() const
+{
+	return FOpenGL::SupportsFramebufferSRGBEnable();
+}
+
+GLuint FOpenGLDynamicRHI::RHIGetResource(FRHITexture* InTexture) const
+{
+	FOpenGLTextureBase* GLTexture = GetOpenGLTextureFromRHITexture(InTexture);
+	return GLTexture->GetResource();
+}
+
+bool FOpenGLDynamicRHI::RHIIsValidTexture(GLuint InTexture) const
+{
+	return glIsTexture(InTexture) == GL_TRUE;
+}
+
+void FOpenGLDynamicRHI::RHISetExternalGPUTime(uint32 InExternalGPUTime)
+{
+	GetGPUProfilingData().ExternalGPUTime = InExternalGPUTime;
+}
+
+#if PLATFORM_ANDROID
+
+EGLDisplay FOpenGLDynamicRHI::RHIGetEGLDisplay() const
+{
+	return AndroidEGL::GetInstance()->GetDisplay();
+}
+
+EGLSurface FOpenGLDynamicRHI::RHIGetEGLSurface() const
+{
+	return AndroidEGL::GetInstance()->GetSurface();
+}
+
+EGLConfig FOpenGLDynamicRHI::RHIGetEGLConfig() const
+{
+	return AndroidEGL::GetInstance()->GetConfig();
+}
+
+EGLContext FOpenGLDynamicRHI::RHIGetEGLContext() const
+{
+	return AndroidEGL::GetInstance()->GetRenderingContext()->eglContext;
+}
+
+ANativeWindow* FOpenGLDynamicRHI::RHIGetEGLNativeWindow() const
+{
+	return AndroidEGL::GetInstance()->GetNativeWindow();
+}
+
+bool FOpenGLDynamicRHI::RHIEGLSupportsNoErrorContext() const
+{
+	return AndroidEGL::GetInstance()->GetSupportsNoErrorContext();
+}
+
+void FOpenGLDynamicRHI::RHIInitEGLInstanceGLES2()
+{
+	AndroidEGL::GetInstance()->Init(AndroidEGL::AV_OpenGLES, 2, 0, false);
+	AndroidEGL::GetInstance()->InitSurface(false, false);
+}
+
+void FOpenGLDynamicRHI::RHIInitEGLBackBuffer()
+{
+	AndroidEGL::GetInstance()->InitBackBuffer();
+}
+
+void FOpenGLDynamicRHI::RHIEGLSetCurrentRenderingContext()
+{
+	AndroidEGL::GetInstance()->SetCurrentRenderingContext();
+}
+
+void FOpenGLDynamicRHI::RHIEGLTerminateContext()
+{
+	AndroidEGL::GetInstance()->Terminate();
+}
+#endif
+
+
 void FOpenGLDynamicRHI::RHIPushEvent(const TCHAR* Name, FColor Color)
 {
 #if ENABLE_OPENGL_DEBUG_GROUPS
diff --git a/Engine/Source/Runtime/OpenGLDrv/Public/IOpenGLDynamicRHI.h b/Engine/Source/Runtime/OpenGLDrv/Public/IOpenGLDynamicRHI.h
new file mode 100644
--- /dev/null
+++ b/Engine/Source/Runtime/OpenGLDrv/Public/IOpenGLDynamicRHI.h
@@ -0,0 +1,72 @@
+// Copyright Epic Games, Inc. All Rights Reserved.
+#pragma once
+
+#include "RHI.h"
+
+#if PLATFORM_WINDOWS
+	#include "Windows/AllowWindowsPlatformTypes.h"
+	THIRD_PARTY_INCLUDES_START
+		#include <GL/glcorearb.h>
+		#include <GL/glext.h>
+		#include <GL/wglext.h>
+	THIRD_PARTY_INCLUDES_END
+	#include "Windows/HideWindowsPlatformTypes.h"
+#elif PLATFORM_LINUX
+	THIRD_PARTY_INCLUDES_START
+		#include <GL/glcorearb.h>
+		#include <GL/glext.h>
+	THIRD_PARTY_INCLUDES_END
+#elif PLATFORM_ANDROID
+	#include "Android/AndroidPlatform.h"
+	THIRD_PARTY_INCLUDES_START
+		#include <EGL/egl.h>
+		#include <EGL/eglext.h>
+		#include <GLES3/gl31.h>
+	THIRD_PARTY_INCLUDES_END
+
+	#ifndef USE_ANDROID_EGL_NO_ERROR_CONTEXT
+		#if UE_BUILD_SHIPPING
+			#define USE_ANDROID_EGL_NO_ERROR_CONTEXT 1
+		#else
+			#define USE_ANDROID_EGL_NO_ERROR_CONTEXT 0
+		#endif
+	#endif
+#endif
+
+struct IOpenGLDynamicRHI : public FDynamicRHIPSOFallback
+{
+	virtual ERHIInterfaceType GetInterfaceType() const override final { return ERHIInterfaceType::OpenGL; }
+
+	virtual int32 RHIGetGLMajorVersion() const = 0;
+	virtual int32 RHIGetGLMinorVersion() const = 0;
+
+	virtual bool RHISupportsFramebufferSRGBEnable() const = 0;
+
+	virtual FTexture2DRHIRef RHICreateTexture2DFromResource(EPixelFormat Format, uint32 SizeX, uint32 SizeY, uint32 NumMips, uint32 NumSamples, uint32 NumSamplesTileMem, const FClearValueBinding& ClearValueBinding, GLuint Resource, ETextureCreateFlags Flags) = 0;
+	virtual FTexture2DRHIRef RHICreateTexture2DArrayFromResource(EPixelFormat Format, uint32 SizeX, uint32 SizeY, uint32 ArraySize, uint32 NumMips, uint32 NumSamples, uint32 NumSamplesTileMem, const FClearValueBinding& ClearValueBinding, GLuint Resource, ETextureCreateFlags Flags) = 0;
+	virtual FTextureCubeRHIRef RHICreateTextureCubeFromResource(EPixelFormat Format, uint32 Size, bool bArray, uint32 ArraySize, uint32 NumMips, uint32 NumSamples, uint32 NumSamplesTileMem, const FClearValueBinding& ClearValueBinding, GLuint Resource, ETextureCreateFlags Flags) = 0;
+
+	virtual GLuint RHIGetResource(FRHITexture* InTexture) const = 0;
+	virtual bool RHIIsValidTexture(GLuint InTexture) const = 0;
+	virtual void RHISetExternalGPUTime(uint32 InExternalGPUTime) = 0;
+
+#if PLATFORM_ANDROID
+	virtual EGLDisplay RHIGetEGLDisplay() const = 0;
+	virtual EGLSurface RHIGetEGLSurface() const = 0;
+	virtual EGLConfig  RHIGetEGLConfig() const = 0;
+	virtual EGLContext RHIGetEGLContext() const = 0;
+	virtual ANativeWindow* RHIGetEGLNativeWindow() const = 0;
+	virtual bool RHIEGLSupportsNoErrorContext() const = 0;
+
+	virtual void RHIInitEGLInstanceGLES2() = 0;
+	virtual void RHIInitEGLBackBuffer() = 0;
+	virtual void RHIEGLSetCurrentRenderingContext() = 0;
+	virtual void RHIEGLTerminateContext() = 0;
+#endif
+};
+
+inline IOpenGLDynamicRHI* GetIOpenGLDynamicRHI()
+{
+	check(GDynamicRHI->GetInterfaceType() == ERHIInterfaceType::OpenGL);
+	return GetDynamicRHI<IOpenGLDynamicRHI>();
+}
diff --git a/Engine/Source/Runtime/OpenGLDrv/Public/OpenGLDrv.h b/Engine/Source/Runtime/OpenGLDrv/Public/OpenGLDrv.h
--- a/Engine/Source/Runtime/OpenGLDrv/Public/OpenGLDrv.h
+++ b/Engine/Source/Runtime/OpenGLDrv/Public/OpenGLDrv.h
@@ -10,6 +10,7 @@
 
 #include "CoreMinimal.h"
 #include "Containers/IndirectArray.h"
+#include "IOpenGLDynamicRHI.h"
 #include "RHIDefinitions.h"
 #include "RHI.h"
 #include "GPUProfiler.h"
@@ -19,9 +20,9 @@
 // @todo platplug: Replace all of these includes with a call to COMPILED_PLATFORM_HEADER(OpenGLDrvPrivate.h)
 //TODO: Move these to OpenGLDrvPrivate.h
 #if PLATFORM_WINDOWS
-	#include "Runtime/OpenGLDrv/Private/Windows/OpenGLWindows.h"
+	#include "Windows/OpenGLWindows.h"
 #elif PLATFORM_LINUX
-	#include "Runtime/OpenGLDrv/Private/Linux/OpenGLLinux.h"
+	#include "Linux/OpenGLLinux.h"
 #elif PLATFORM_ANDROID
 	#include "Android/AndroidOpenGL.h"
 #else
@@ -287,7 +288,7 @@ struct FOpenGLGPUProfiler : public FGPUProfiler
 PRAGMA_DISABLE_DEPRECATION_WARNINGS
 
 /** The interface which is implemented by the dynamically bound RHI. */
-class OPENGLDRV_API FOpenGLDynamicRHI  final : public FDynamicRHIPSOFallback, public IRHICommandContextPSOFallback
+class OPENGLDRV_API FOpenGLDynamicRHI  final : public IOpenGLDynamicRHI, public IRHICommandContextPSOFallback
 {
 public:
 
@@ -299,13 +300,36 @@ public:
 	/** Destructor */
 	~FOpenGLDynamicRHI() {}
 
+	// IOpenGLDynamicRHI interface.
+	virtual int32 RHIGetGLMajorVersion() const final override;
+	virtual int32 RHIGetGLMinorVersion() const final override;
+	virtual bool RHISupportsFramebufferSRGBEnable() const final override;
+	virtual FTexture2DRHIRef RHICreateTexture2DFromResource(EPixelFormat Format, uint32 SizeX, uint32 SizeY, uint32 NumMips, uint32 NumSamples, uint32 NumSamplesTileMem, const FClearValueBinding& ClearValueBinding, GLuint Resource, ETextureCreateFlags Flags) final override;
+	virtual FTexture2DRHIRef RHICreateTexture2DArrayFromResource(EPixelFormat Format, uint32 SizeX, uint32 SizeY, uint32 ArraySize, uint32 NumMips, uint32 NumSamples, uint32 NumSamplesTileMem, const FClearValueBinding& ClearValueBinding, GLuint Resource, ETextureCreateFlags Flags) final override;
+	virtual FTextureCubeRHIRef RHICreateTextureCubeFromResource(EPixelFormat Format, uint32 Size, bool bArray, uint32 ArraySize, uint32 NumMips, uint32 NumSamples, uint32 NumSamplesTileMem, const FClearValueBinding& ClearValueBinding, GLuint Resource, ETextureCreateFlags Flags) final override;
+	virtual GLuint RHIGetResource(FRHITexture* InTexture) const final override;
+	virtual bool RHIIsValidTexture(GLuint InTexture) const final override;
+	virtual void RHISetExternalGPUTime(uint32 InExternalGPUTime) final override;
+
+#if PLATFORM_ANDROID
+	virtual EGLDisplay RHIGetEGLDisplay() const final override;
+	virtual EGLSurface RHIGetEGLSurface() const final override;
+	virtual EGLConfig  RHIGetEGLConfig() const final override;
+	virtual EGLContext RHIGetEGLContext() const final override;
+	virtual ANativeWindow* RHIGetEGLNativeWindow() const final override;
+	virtual bool RHIEGLSupportsNoErrorContext() const final override;
+	virtual void RHIInitEGLInstanceGLES2() final override;
+	virtual void RHIInitEGLBackBuffer() final override;
+	virtual void RHIEGLSetCurrentRenderingContext() final override;
+	virtual void RHIEGLTerminateContext() final override;
+#endif
+
 	// FDynamicRHI interface.
 	virtual void Init();
 	virtual void PostInit();
 
 	virtual void Shutdown();
 	virtual const TCHAR* GetName() override { return TEXT("OpenGL"); }
-	virtual ERHIInterfaceType GetInterfaceType() const override final { return ERHIInterfaceType::OpenGL; }
 
 	// If using a Proxy object return the contained GL object rather than the proxy itself.
 	template<typename TRHIType>
@@ -492,11 +516,6 @@ public:
 
 	virtual void RHICopyBufferRegion(FRHIBuffer* DestBuffer, uint64 DstOffset, FRHIBuffer* SourceBuffer, uint64 SrcOffset, uint64 NumBytes) final override;
 
-	// FOpenGLDynamicRHI interface.
-	FTexture2DRHIRef RHICreateTexture2DFromResource(EPixelFormat Format, uint32 SizeX, uint32 SizeY, uint32 NumMips, uint32 NumSamples, uint32 NumSamplesTileMem, const FClearValueBinding& ClearValueBinding, GLuint Resource, ETextureCreateFlags Flags);
-	FTexture2DRHIRef RHICreateTexture2DArrayFromResource(EPixelFormat Format, uint32 SizeX, uint32 SizeY, uint32 ArraySize, uint32 NumMips, uint32 NumSamples, uint32 NumSamplesTileMem, const FClearValueBinding& ClearValueBinding, GLuint Resource, ETextureCreateFlags Flags);
-	FTextureCubeRHIRef RHICreateTextureCubeFromResource(EPixelFormat Format, uint32 Size, bool bArray, uint32 ArraySize, uint32 NumMips, uint32 NumSamples, uint32 NumSamplesTileMem, const FClearValueBinding& ClearValueBinding, GLuint Resource, ETextureCreateFlags Flags);
-
 	// Inline copy
 	virtual void RHICopyToStagingBuffer(FRHIBuffer* SourceBufferRHI, FRHIStagingBuffer* DestinationStagingBufferRHI, uint32 InOffset, uint32 InNumBytes) final override;
 	virtual void RHIWriteGPUFence(FRHIGPUFence* FenceRHI) final override;
