Fixed: Correctly count texture and render target memory usage in OpenGL RHI

diff --git a/Engine/Source/Runtime/OpenGLDrv/Private/OpenGLTexture.cpp b/Engine/Source/Runtime/OpenGLDrv/Private/OpenGLTexture.cpp
--- a/Engine/Source/Runtime/OpenGLDrv/Private/OpenGLTexture.cpp
+++ b/Engine/Source/Runtime/OpenGLDrv/Private/OpenGLTexture.cpp
@@ -84,7 +84,7 @@ int64 GOpenGLDedicatedVideoMemory = 0;
 // In bytes. Never changed after RHI init. Our estimate of the amount of memory that we can use for graphics resources in total.
 int64 GOpenGLTotalGraphicsMemory = 0;
 
-static bool ShouldCountAsTextureMemory(ETextureCreateFlags Flags)
+static bool IsRenderTarget(ETextureCreateFlags Flags)
 {
 	return EnumHasAnyFlags(Flags, TexCreate_RenderTargetable | TexCreate_ResolveTargetable | TexCreate_DepthStencilTargetable);
 }
@@ -96,7 +96,7 @@ void OpenGLTextureAllocated(FRHITexture* Texture, ETextureCreateFlags Flags)
 	FOpenGLTexture2D* Texture2D = 0;
 	FOpenGLTexture2DArray* Texture2DArray = 0;
 	FOpenGLTexture3D* Texture3D = 0;
-	bool bRenderTarget = !ShouldCountAsTextureMemory(Flags);
+	bool bRenderTarget = IsRenderTarget(Flags);
 
 	if (( TextureCube = (FOpenGLTextureCube*)Texture->GetTextureCube()) != NULL)
 	{
@@ -193,7 +193,7 @@ void OpenGLTextureAllocated(FRHITexture* Texture, ETextureCreateFlags Flags)
 
 void OpenGLTextureDeleted( FRHITexture* Texture )
 {
-	bool bRenderTarget = !ShouldCountAsTextureMemory(Texture->GetFlags());
+	bool bRenderTarget = IsRenderTarget(Texture->GetFlags());
 	int32 TextureSize = 0;
 	if (Texture->GetTextureCube())
 	{
