Copying //UE5/Dev-ParallelRendering to //UE5/Main

diff --git a/Engine/Source/Runtime/OpenGLDrv/Private/OpenGLRenderTarget.cpp b/Engine/Source/Runtime/OpenGLDrv/Private/OpenGLRenderTarget.cpp
--- a/Engine/Source/Runtime/OpenGLDrv/Private/OpenGLRenderTarget.cpp
+++ b/Engine/Source/Runtime/OpenGLDrv/Private/OpenGLRenderTarget.cpp
@@ -28,20 +28,20 @@ class FOpenGLFramebufferKey
 {
 	struct RenderTargetInfo
 	{
-		FOpenGLTextureBase* Texture;
-		GLuint				Resource;
-		uint32				MipmapLevel;
-		uint32				ArrayIndex;
+		FOpenGLTexture* Texture;
+		GLuint			Resource;
+		uint32			MipmapLevel;
+		uint32			ArrayIndex;
 	};
 
 public:
 
 	FOpenGLFramebufferKey(
 		uint32 InNumRenderTargets,
-		FOpenGLTextureBase** InRenderTargets,
+		FOpenGLTexture** InRenderTargets,
 		const uint32* InRenderTargetArrayIndices,
 		const uint32* InRenderTargetMipmapLevels,
-		FOpenGLTextureBase* InDepthStencilTarget,
+		FOpenGLTexture* InDepthStencilTarget,
 		EOpenGLCurrentContext InContext
 		)
 		:	DepthStencilTarget(InDepthStencilTarget)
@@ -83,13 +83,13 @@ public:
 		return FCrc::MemCrc_DEPRECATED(Key.RenderTargets, sizeof(Key.RenderTargets)) ^ GetTypeHash(Key.DepthStencilTarget) ^ GetTypeHash(Key.Context);
 	}
 
-	const FOpenGLTextureBase* GetRenderTarget( int32 Index ) const { return RenderTargets[Index].Texture; }
-	const FOpenGLTextureBase* GetDepthStencilTarget( void ) const { return DepthStencilTarget; }
+	const FOpenGLTexture* GetRenderTarget( int32 Index ) const { return RenderTargets[Index].Texture; }
+	const FOpenGLTexture* GetDepthStencilTarget( void ) const { return DepthStencilTarget; }
 
 private:
 
 	RenderTargetInfo RenderTargets[MaxSimultaneousRenderTargets];
-	FOpenGLTextureBase* DepthStencilTarget;
+	FOpenGLTexture* DepthStencilTarget;
 	EOpenGLCurrentContext Context;
 };
 
@@ -102,7 +102,7 @@ static FOpenGLFramebufferCache& GetOpenGLFramebufferCache()
 	return OpenGLFramebufferCache;
 }
 
-GLuint FOpenGLDynamicRHI::GetOpenGLFramebuffer(uint32 NumSimultaneousRenderTargets, FOpenGLTextureBase** RenderTargets, const uint32* ArrayIndices, const uint32* MipmapLevels, FOpenGLTextureBase* DepthStencilTarget)
+GLuint FOpenGLDynamicRHI::GetOpenGLFramebuffer(uint32 NumSimultaneousRenderTargets, FOpenGLTexture** RenderTargets, const uint32* ArrayIndices, const uint32* MipmapLevels, FOpenGLTexture* DepthStencilTarget)
 {
 	VERIFY_GL_SCOPE();
 
@@ -142,20 +142,19 @@ GLuint FOpenGLDynamicRHI::GetOpenGLFramebuffer(uint32 NumSimultaneousRenderTarge
 
 	if (bUsingArrayTextures && FOpenGL::SupportsMobileMultiView() && bMultiViewCVar)
 	{
-		FOpenGLTextureBase* const RenderTarget = RenderTargets[0];
+		FOpenGLTexture* const RenderTarget = RenderTargets[0];
 		glBindFramebuffer(GL_FRAMEBUFFER, 0);
 		glBindFramebuffer(GL_DRAW_FRAMEBUFFER, Framebuffer);
 
-		FOpenGLTexture2D* RenderTarget2D = (FOpenGLTexture2D*)RenderTarget;
-		const uint32 NumSamplesTileMem = RenderTarget2D->GetNumSamplesTileMem();
-		if (NumSamplesTileMem > 1)
+		const uint32 NumSamplesRendered = RenderTarget->GetNumSamplesRendered();
+		if (NumSamplesRendered > 1)
 		{
-			glFramebufferTextureMultisampleMultiviewOVR(GL_DRAW_FRAMEBUFFER, GL_COLOR_ATTACHMENT0, RenderTarget->GetResource(), 0, NumSamplesTileMem, 0, 2);
+			glFramebufferTextureMultisampleMultiviewOVR(GL_DRAW_FRAMEBUFFER, GL_COLOR_ATTACHMENT0, RenderTarget->GetResource(), 0, NumSamplesRendered, 0, 2);
 			VERIFY_GL(glFramebufferTextureMultisampleMultiviewOVR);
 
 			if (DepthStencilTarget)
 			{
-				glFramebufferTextureMultisampleMultiviewOVR(GL_DRAW_FRAMEBUFFER, GL_DEPTH_ATTACHMENT, DepthStencilTarget->GetResource(), 0, NumSamplesTileMem, 0, 2);
+				glFramebufferTextureMultisampleMultiviewOVR(GL_DRAW_FRAMEBUFFER, GL_DEPTH_ATTACHMENT, DepthStencilTarget->GetResource(), 0, NumSamplesRendered, 0, 2);
 				VERIFY_GL(glFramebufferTextureMultisampleMultiviewOVR);
 			}
 		}
@@ -183,9 +182,9 @@ GLuint FOpenGLDynamicRHI::GetOpenGLFramebuffer(uint32 NumSimultaneousRenderTarge
 #endif
 
 	int32 FirstNonzeroRenderTarget = -1;
-	for( int32 RenderTargetIndex = NumSimultaneousRenderTargets - 1; RenderTargetIndex >= 0 ; --RenderTargetIndex )
+	for (int32 RenderTargetIndex = NumSimultaneousRenderTargets - 1; RenderTargetIndex >= 0; --RenderTargetIndex)
 	{
-		FOpenGLTextureBase* RenderTarget = RenderTargets[RenderTargetIndex];
+		FOpenGLTexture* RenderTarget = RenderTargets[RenderTargetIndex];
 		if (!RenderTarget)
 		{
 			continue;
@@ -197,18 +196,14 @@ GLuint FOpenGLDynamicRHI::GetOpenGLFramebuffer(uint32 NumSimultaneousRenderTarge
 			switch (RenderTarget->Target)
 			{
 			case GL_TEXTURE_2D:
-#if PLATFORM_ANDROID
 			case GL_TEXTURE_EXTERNAL_OES:
-#endif
 			case GL_TEXTURE_2D_MULTISAMPLE:
 			{
 #if PLATFORM_ANDROID
-				FOpenGLTexture2D* RenderTarget2D = (FOpenGLTexture2D*)RenderTarget;
-				const uint32 NumSamplesTileMem = RenderTarget2D->GetNumSamplesTileMem();
-				if (NumSamplesTileMem > 1 && glFramebufferTexture2DMultisampleEXT)
+				if (RenderTarget->IsTiledMSAA() && glFramebufferTexture2DMultisampleEXT)
 				{
 					// GL_EXT_multisampled_render_to_texture requires GL_COLOR_ATTACHMENT0
-					glFramebufferTexture2DMultisampleEXT(GL_FRAMEBUFFER, GL_COLOR_ATTACHMENT0 + RenderTargetIndex, RenderTarget->Target, RenderTarget->GetResource(), MipmapLevels[RenderTargetIndex], NumSamplesTileMem);
+					glFramebufferTexture2DMultisampleEXT(GL_FRAMEBUFFER, GL_COLOR_ATTACHMENT0 + RenderTargetIndex, RenderTarget->Target, RenderTarget->GetResource(), MipmapLevels[RenderTargetIndex], RenderTarget->GetNumSamplesRendered());
 					VERIFY_GL(glFramebufferTexture2DMultisampleEXT);
 				}
 				else
@@ -235,19 +230,15 @@ GLuint FOpenGLDynamicRHI::GetOpenGLFramebuffer(uint32 NumSimultaneousRenderTarge
 			switch( RenderTarget->Target )
 			{
 			case GL_TEXTURE_2D:
-#if PLATFORM_ANDROID
 			case GL_TEXTURE_EXTERNAL_OES:
-#endif
 			case GL_TEXTURE_2D_MULTISAMPLE:
 			{
 				check(ArrayIndices[RenderTargetIndex] == 0);
 #if PLATFORM_ANDROID
-				FOpenGLTexture2D* RenderTarget2D = (FOpenGLTexture2D*)RenderTarget;
-				const uint32 NumSamplesTileMem = RenderTarget2D->GetNumSamplesTileMem();
-				if (NumSamplesTileMem > 1 && glFramebufferTexture2DMultisampleEXT)
+				if (RenderTarget->IsTiledMSAA() && glFramebufferTexture2DMultisampleEXT)
 				{
 					// GL_EXT_multisampled_render_to_texture requires GL_COLOR_ATTACHMENT0
-					glFramebufferTexture2DMultisampleEXT(GL_FRAMEBUFFER, GL_COLOR_ATTACHMENT0 + RenderTargetIndex, RenderTarget->Target, RenderTarget->GetResource(), MipmapLevels[RenderTargetIndex], NumSamplesTileMem);
+					glFramebufferTexture2DMultisampleEXT(GL_FRAMEBUFFER, GL_COLOR_ATTACHMENT0 + RenderTargetIndex, RenderTarget->Target, RenderTarget->GetResource(), MipmapLevels[RenderTargetIndex], RenderTarget->GetNumSamplesRendered());
 					VERIFY_GL(glFramebufferTexture2DMultisampleEXT);
 				}
 				else
@@ -282,24 +273,17 @@ GLuint FOpenGLDynamicRHI::GetOpenGLFramebuffer(uint32 NumSimultaneousRenderTarge
 		switch (DepthStencilTarget->Target)
 		{
 		case GL_TEXTURE_2D:
-#if PLATFORM_ANDROID
 		case GL_TEXTURE_EXTERNAL_OES:
-#endif
 		case GL_TEXTURE_2D_MULTISAMPLE:
 		{
-				FOpenGL::FramebufferTexture2D(GL_FRAMEBUFFER, DepthStencilTarget->Attachment, DepthStencilTarget->Target, DepthStencilTarget->GetResource(), 0);
+			FOpenGL::FramebufferTexture2D(GL_FRAMEBUFFER, DepthStencilTarget->Attachment, DepthStencilTarget->Target, DepthStencilTarget->GetResource(), 0);
 			break;
 		}
 		case GL_RENDERBUFFER:
 		{
-			FOpenGLTexture2D* DepthStencilTarget2D = (FOpenGLTexture2D*)DepthStencilTarget;
-			const uint32 NumSamplesTileMem = DepthStencilTarget2D->GetNumSamplesTileMem();
-			if (NumSamplesTileMem > 1)
-			{	
-				glFramebufferRenderbuffer(GL_FRAMEBUFFER, GL_DEPTH_ATTACHMENT, GL_RENDERBUFFER, DepthStencilTarget->GetResource());
-					glFramebufferRenderbuffer(GL_FRAMEBUFFER, GL_STENCIL_ATTACHMENT, GL_RENDERBUFFER, DepthStencilTarget->GetResource());
-				VERIFY_GL(glFramebufferRenderbuffer);
-			}
+			glFramebufferRenderbuffer(GL_FRAMEBUFFER, GL_DEPTH_ATTACHMENT, GL_RENDERBUFFER, DepthStencilTarget->GetResource());
+			glFramebufferRenderbuffer(GL_FRAMEBUFFER, GL_STENCIL_ATTACHMENT, GL_RENDERBUFFER, DepthStencilTarget->GetResource());
+			VERIFY_GL(glFramebufferRenderbuffer);
 			break;
 		}
 		case GL_TEXTURE_3D:
@@ -343,9 +327,11 @@ GLuint FOpenGLDynamicRHI::GetOpenGLFramebuffer(uint32 NumSimultaneousRenderTarge
 	return Framebuffer;
 }
 
-void ReleaseOpenGLFramebuffers(FOpenGLDynamicRHI* Device, FRHITexture* TextureRHI)
+void ReleaseOpenGLFramebuffers(FRHITexture* TextureRHI)
 {
-	const FOpenGLTextureBase* Texture = GetOpenGLTextureFromRHITexture(TextureRHI);
+	VERIFY_GL_SCOPE();
+
+	const FOpenGLTexture* Texture = GetOpenGLTextureFromRHITexture(TextureRHI);
 
 	if (Texture)
 	{
@@ -354,7 +340,7 @@ void ReleaseOpenGLFramebuffers(FOpenGLDynamicRHI* Device, FRHITexture* TextureRH
 			bool bPurgeFramebuffer = false;
 			FOpenGLFramebufferKey Key = It.Key();
 
-			const FOpenGLTextureBase* DepthStencilTarget = Key.GetDepthStencilTarget();
+			const FOpenGLTexture* DepthStencilTarget = Key.GetDepthStencilTarget();
 			if( DepthStencilTarget && DepthStencilTarget->Target == Texture->Target && DepthStencilTarget->GetRawResourceName() == Texture->GetRawResourceName() )
 			{
 				bPurgeFramebuffer = true;
@@ -363,7 +349,7 @@ void ReleaseOpenGLFramebuffers(FOpenGLDynamicRHI* Device, FRHITexture* TextureRH
 			{
 				for( uint32 RenderTargetIndex = 0; RenderTargetIndex < MaxSimultaneousRenderTargets; ++RenderTargetIndex )
 				{
-					const FOpenGLTextureBase* RenderTarget = Key.GetRenderTarget(RenderTargetIndex);
+					const FOpenGLTexture* RenderTarget = Key.GetRenderTarget(RenderTargetIndex);
 					if( RenderTarget && RenderTarget->Target == Texture->Target && RenderTarget->GetRawResourceName() == Texture->GetRawResourceName() )
 					{
 						bPurgeFramebuffer = true;
@@ -377,12 +363,8 @@ void ReleaseOpenGLFramebuffers(FOpenGLDynamicRHI* Device, FRHITexture* TextureRH
 				GLuint FramebufferToDelete = It.Value()-1;
 				check(FramebufferToDelete > 0);
 
-				RunOnGLRenderContextThread( [=]() 
-					{
-						VERIFY_GL_SCOPE();
-						Device->PurgeFramebufferFromCaches( FramebufferToDelete );
-						glDeleteFramebuffers( 1, &FramebufferToDelete );
-					});
+				FOpenGLDynamicRHI::Get().PurgeFramebufferFromCaches( FramebufferToDelete );
+				glDeleteFramebuffers( 1, &FramebufferToDelete );
 
 				It.RemoveCurrent();
 			}
@@ -423,11 +405,14 @@ void FOpenGLDynamicRHI::RHICopyToResolveTarget(FRHITexture* SourceTextureRHI, FR
 		return;
 	}
 
-	FOpenGLTextureBase* SourceTexture = GetOpenGLTextureFromRHITexture(SourceTextureRHI);
-	FOpenGLTextureBase* DestTexture = GetOpenGLTextureFromRHITexture(DestTextureRHI);
+	FOpenGLTexture* SourceTexture = GetOpenGLTextureFromRHITexture(SourceTextureRHI);
+	FOpenGLTexture* DestTexture = GetOpenGLTextureFromRHITexture(DestTextureRHI);
 
 	if (SourceTexture != DestTexture)
 	{
+		const FRHITextureDesc& SrcDesc = SourceTexture->GetDesc();
+		const FRHITextureDesc& DstDesc = DestTexture->GetDesc();
+
 		VERIFY_GL_SCOPE();
 
 		check(GMaxRHIFeatureLevel >= ERHIFeatureLevel::SM5 || ResolveParams.SourceArrayIndex == 0);
@@ -446,8 +431,8 @@ void FOpenGLDynamicRHI::RHICopyToResolveTarget(FRHITexture* SourceTextureRHI, FR
 			// returns false if the rectangle is "inside out" (e.g. X1 > X2), and we want to perform flipping when that's the case.
 			SrcRect.Min.X = 0;
 			SrcRect.Min.Y = 0;
-			SrcRect.Max.X = GetOpenGLTextureSizeXFromRHITexture(SourceTextureRHI);
-			SrcRect.Max.Y = GetOpenGLTextureSizeYFromRHITexture(SourceTextureRHI);
+			SrcRect.Max.X = SrcDesc.Extent.X;
+			SrcRect.Max.Y = SrcDesc.Extent.Y;
 			SrcRect.Max.X = FMath::Max<int32>(1, SrcRect.Max.X >> ResolveParams.MipIndex);
 			SrcRect.Max.Y = FMath::Max<int32>(1, SrcRect.Max.Y >> ResolveParams.MipIndex);
 		}
@@ -457,8 +442,8 @@ void FOpenGLDynamicRHI::RHICopyToResolveTarget(FRHITexture* SourceTextureRHI, FR
 		{
 			DestRect.Min.X = 0;
 			DestRect.Min.Y = 0;
-			DestRect.Max.X = GetOpenGLTextureSizeXFromRHITexture(DestTextureRHI);
-			DestRect.Max.Y = GetOpenGLTextureSizeYFromRHITexture(DestTextureRHI);
+			DestRect.Max.X = DstDesc.Extent.X;
+			DestRect.Max.Y = DstDesc.Extent.Y;
 			DestRect.Max.X = FMath::Max<int32>(1, DestRect.Max.X >> ResolveParams.MipIndex);
 			DestRect.Max.Y = FMath::Max<int32>(1, DestRect.Max.Y >> ResolveParams.MipIndex);
 		}
@@ -546,8 +531,7 @@ void FOpenGLDynamicRHI::RHICopyToResolveTarget(FRHITexture* SourceTextureRHI, FR
 		bool bLockableTarget = DestTextureRHI->GetTexture2D() && EnumHasAnyFlags(DestTextureRHI->GetFlags(), TexCreate_CPUReadback) && !EnumHasAnyFlags(DestTextureRHI->GetFlags(), TexCreate_RenderTargetable|TexCreate_DepthStencilTargetable) && !DestTextureRHI->IsMultisampled();
 		if(bLockableTarget && !ResolveParams.Rect.IsValid())
 		{
-			FOpenGLTexture2D* DestTex = (FOpenGLTexture2D*)DestTexture;
-			DestTex->Resolve(MipmapLevel, DestIndex);
+			DestTexture->Resolve(MipmapLevel, DestIndex);
 		}
 
 		GetContextStateForCurrentContext().Framebuffer = (GLuint)-1;
@@ -564,14 +548,12 @@ void FOpenGLDynamicRHI::ReadSurfaceDataRaw(FOpenGLContextState& ContextState, FR
 {
 	VERIFY_GL_SCOPE();
 
-	FOpenGLTexture2D* Texture2D = (FOpenGLTexture2D*)TextureRHI->GetTexture2D();
-	if( !Texture2D )
+	FOpenGLTexture* Texture = GetOpenGLTextureFromRHITexture(TextureRHI);
+	if( !Texture)
 	{
 		return;	// just like in D3D11
 	}
 
-	FOpenGLTextureBase* Texture = (FOpenGLTextureBase*)Texture2D;
-
 	GLuint FramebufferToDelete = 0;
 	GLuint RenderbufferToDelete = 0;
 	const FOpenGLTextureFormat& GLFormat = GOpenGLTextureFormats[TextureRHI->GetFormat()];
@@ -645,15 +627,15 @@ void FOpenGLDynamicRHI::ReadSurfaceDataRaw(FOpenGLContextState& ContextState, FR
 		GLuint Renderbuffer = 0;
 		glGenRenderbuffers(1, &RenderbufferToDelete);
 		glBindRenderbuffer(GL_RENDERBUFFER, RenderbufferToDelete);
-		glRenderbufferStorage(GL_RENDERBUFFER, GLFormat.InternalFormat[false], Texture2D->GetSizeX(), Texture2D->GetSizeY());
+		glRenderbufferStorage(GL_RENDERBUFFER, GLFormat.InternalFormat[false], Texture->GetSizeX(), Texture->GetSizeY());
 		glBindRenderbuffer(GL_RENDERBUFFER, 0);
 
 		glFramebufferRenderbuffer(GL_FRAMEBUFFER, Attachment, GL_RENDERBUFFER, RenderbufferToDelete);
 		FOpenGL::CheckFrameBuffer();
 		glBindFramebuffer(UGL_READ_FRAMEBUFFER, SourceFramebuffer);
 		FOpenGL::BlitFramebuffer(
-			0, 0, Texture2D->GetSizeX(), Texture2D->GetSizeY(),
-			0, 0, Texture2D->GetSizeX(), Texture2D->GetSizeY(),
+			0, 0, Texture->GetSizeX(), Texture->GetSizeY(),
+			0, 0, Texture->GetSizeX(), Texture->GetSizeY(),
 			(bDepthFormat ? (bDepthStencilFormat ? (GL_DEPTH_BUFFER_BIT | GL_STENCIL_BUFFER_BIT) : GL_DEPTH_BUFFER_BIT) : GL_COLOR_BUFFER_BIT),
 			GL_NEAREST
 			);
@@ -902,7 +884,7 @@ void FOpenGLDynamicRHI::RHIReadSurfaceData(FRHITexture* TextureRHI, FIntRect Rec
 	RHITHREAD_GLCOMMAND_PROLOGUE();
 	VERIFY_GL_SCOPE();
 
-	FOpenGLTextureBase* Texture = GetOpenGLTextureFromRHITexture(TextureRHI);
+	FOpenGLTexture* Texture = GetOpenGLTextureFromRHITexture(TextureRHI);
 	if (!ensure(Texture))
 	{
 		return;
@@ -943,15 +925,15 @@ void FOpenGLDynamicRHI::RHIMapStagingSurface(FRHITexture* TextureRHI, FRHIGPUFen
 
 	VERIFY_GL_SCOPE();
 	
-	FOpenGLTexture2D* Texture2D = (FOpenGLTexture2D*)TextureRHI->GetTexture2D();
-	check(Texture2D);
-	check(Texture2D->IsStaging());
+	FOpenGLTexture* Texture = GetOpenGLTextureFromRHITexture(TextureRHI->GetTexture2D());
+	check(Texture);
+	check(EnumHasAnyFlags(Texture->GetDesc().Flags, TexCreate_CPUReadback));
 
-	OutWidth = Texture2D->GetSizeX();
-	OutHeight = Texture2D->GetSizeY();
+	OutWidth = Texture->GetSizeX();
+	OutHeight = Texture->GetSizeY();
 	
 	uint32 Stride = 0;
-	OutData = Texture2D->Lock( 0, 0, RLM_ReadOnly, Stride );
+	OutData = Texture->Lock( 0, 0, RLM_ReadOnly, Stride );
 	RHITHREAD_GLCOMMAND_EPILOGUE();
 }
 
@@ -963,10 +945,10 @@ void FOpenGLDynamicRHI::RHIUnmapStagingSurface(FRHITexture* TextureRHI, uint32 G
 
 	VERIFY_GL_SCOPE();
 	
-	FOpenGLTexture2D* Texture2D = (FOpenGLTexture2D*)TextureRHI->GetTexture2D();
-	check(Texture2D);
+	FOpenGLTexture* Texture = GetOpenGLTextureFromRHITexture(TextureRHI->GetTexture2D());
+	check(Texture);
 
-	Texture2D->Unlock( 0, 0 );
+	Texture->Unlock( 0, 0 );
 	RHITHREAD_GLCOMMAND_EPILOGUE();
 }
 
@@ -980,7 +962,7 @@ void FOpenGLDynamicRHI::RHIReadSurfaceFloatData(FRHITexture* TextureRHI,FIntRect
 
 	//reading from arrays only supported on SM5 and up.
 	check(FOpenGL::SupportsFloatReadSurface() && (ArrayIndex == 0 || GMaxRHIFeatureLevel >= ERHIFeatureLevel::SM5));	
-	FOpenGLTextureBase* Texture = GetOpenGLTextureFromRHITexture(TextureRHI);
+	FOpenGLTexture* Texture = GetOpenGLTextureFromRHITexture(TextureRHI);
 	check(TextureRHI->GetFormat() == PF_FloatRGBA);
 
 	const uint32 MipmapLevel = MipIndex;
@@ -1058,8 +1040,7 @@ void FOpenGLDynamicRHI::RHIRead3DSurfaceFloatData(FRHITexture* TextureRHI,FIntRe
 	check( FOpenGL::SupportsTexture3D() );
 	check( TextureRHI->GetFormat() == PF_FloatRGBA );
 
-	FRHITexture3D* Texture3DRHI = TextureRHI->GetTexture3D();
-	FOpenGLTextureBase* Texture = GetOpenGLTextureFromRHITexture(TextureRHI);
+	FOpenGLTexture* Texture = GetOpenGLTextureFromRHITexture(TextureRHI);
 
 	uint32 SizeX = Rect.Width();
 	uint32 SizeY = Rect.Height();
