Copying //UE5/Dev-ParallelRendering to //UE5/Main

diff --git a/Engine/Source/Runtime/OpenGLDrv/Private/OpenGLTexture.cpp b/Engine/Source/Runtime/OpenGLDrv/Private/OpenGLTexture.cpp
--- a/Engine/Source/Runtime/OpenGLDrv/Private/OpenGLTexture.cpp
+++ b/Engine/Source/Runtime/OpenGLDrv/Private/OpenGLTexture.cpp
@@ -68,8 +68,6 @@ static FAutoConsoleVariableRef CVarTextureEvictionLogging(
 	ECVF_RenderThreadSafe
 );
 
-class FOpenGLDynamicRHI* FOpenGLTextureBase::OpenGLRHI = nullptr;
-
 /*-----------------------------------------------------------------------------
 	Texture allocator support.
 -----------------------------------------------------------------------------*/
@@ -84,206 +82,59 @@ int64 GOpenGLDedicatedVideoMemory = 0;
 // In bytes. Never changed after RHI init. Our estimate of the amount of memory that we can use for graphics resources in total.
 int64 GOpenGLTotalGraphicsMemory = 0;
 
-static bool IsRenderTarget(ETextureCreateFlags Flags)
+void FOpenGLTexture::UpdateTextureStats(FOpenGLTexture* Texture, bool bAllocating)
 {
-	return EnumHasAnyFlags(Flags, TexCreate_RenderTargetable | TexCreate_ResolveTargetable | TexCreate_DepthStencilTargetable);
-}
+	const FRHITextureDesc& Desc = Texture->GetDesc();
 
-void OpenGLTextureAllocated(FRHITexture* Texture, ETextureCreateFlags Flags)
-{
-	int32 TextureSize = 0;
-	FOpenGLTextureCube* TextureCube = 0;
-	FOpenGLTexture2D* Texture2D = 0;
-	FOpenGLTexture2DArray* Texture2DArray = 0;
-	FOpenGLTexture3D* Texture3D = 0;
-	bool bRenderTarget = IsRenderTarget(Flags);
+	const bool bRenderTarget = EnumHasAnyFlags(Desc.Flags, TexCreate_RenderTargetable | TexCreate_ResolveTargetable | TexCreate_DepthStencilTargetable);
+	const int64 TextureSize = bAllocating
+		?  int64(Texture->MemorySize)
+		: -int64(Texture->MemorySize);
 
-	if (( TextureCube = (FOpenGLTextureCube*)Texture->GetTextureCube()) != NULL)
-	{
-		if (TextureCube->IsMemorySizeSet())
-		{
-			return; // already set this up on RT
-		}
+	const int64 TextureSizeInKBs = bAllocating
+		?  FMath::DivideAndRoundUp(int64(Texture->MemorySize), 1024ll)
+		: -FMath::DivideAndRoundUp(int64(Texture->MemorySize), 1024ll);
 
-		TextureSize = CalcTextureSize( TextureCube->GetSize(), TextureCube->GetSize(), TextureCube->GetFormat(), TextureCube->GetNumMips() );
-		TextureSize *= TextureCube->GetArraySize() * (TextureCube->GetArraySize() == 1 ? 6 : 1);
-		TextureCube->SetMemorySize( TextureSize );
-		TextureCube->SetIsPowerOfTwo(FMath::IsPowerOfTwo(TextureCube->GetSizeX()) && FMath::IsPowerOfTwo(TextureCube->GetSizeY()));
-		if (bRenderTarget)
-		{
-			INC_MEMORY_STAT_BY(STAT_RenderTargetMemoryCube,TextureSize);
-		}
-		else
-		{
-			INC_MEMORY_STAT_BY(STAT_TextureMemoryCube,TextureSize);
-		}
-	}
-	else if ((Texture2D = (FOpenGLTexture2D*)Texture->GetTexture2D()) != NULL)
-	{
-		if (Texture2D->IsMemorySizeSet())
-		{
-			return; // already set this up on RT
-		}
-		TextureSize = CalcTextureSize( Texture2D->GetSizeX(), Texture2D->GetSizeY(), Texture2D->GetFormat(), Texture2D->GetNumMips() )*Texture2D->GetNumSamples();
-		Texture2D->SetMemorySize( TextureSize );
-		Texture2D->SetIsPowerOfTwo(FMath::IsPowerOfTwo(Texture2D->GetSizeX()) && FMath::IsPowerOfTwo(Texture2D->GetSizeY()));
-		if (bRenderTarget)
-		{
-			INC_MEMORY_STAT_BY(STAT_RenderTargetMemory2D,TextureSize);
-		}
-		else
-		{
-			INC_MEMORY_STAT_BY(STAT_TextureMemory2D,TextureSize);
-		}
-	}
-	else if ((Texture3D = (FOpenGLTexture3D*)Texture->GetTexture3D()) != NULL)
-	{
-		if (Texture3D->IsMemorySizeSet())
-		{
-			return; // already set this up on RT
-		}
-		TextureSize = CalcTextureSize3D( Texture3D->GetSizeX(), Texture3D->GetSizeY(), Texture3D->GetSizeZ(), Texture3D->GetFormat(), Texture3D->GetNumMips() );
-		Texture3D->SetMemorySize( TextureSize );
-		Texture3D->SetIsPowerOfTwo(FMath::IsPowerOfTwo(Texture3D->GetSizeX()) && FMath::IsPowerOfTwo(Texture3D->GetSizeY()) && FMath::IsPowerOfTwo(Texture3D->GetSizeZ()));
-		if (bRenderTarget)
-		{
-			INC_MEMORY_STAT_BY(STAT_RenderTargetMemory3D,TextureSize);
-		}
-		else
-		{
-			INC_MEMORY_STAT_BY(STAT_TextureMemory3D,TextureSize);
-		}
-	}
-	else if ((Texture2DArray = (FOpenGLTexture2DArray*)Texture->GetTexture2DArray()) != NULL)
+	if (bRenderTarget)
 	{
-		if (Texture2DArray->IsMemorySizeSet())
-		{
-			return; // already set this up on RT
-		}
-		TextureSize = Texture2DArray->GetSizeZ() * CalcTextureSize( Texture2DArray->GetSizeX(), Texture2DArray->GetSizeY(), Texture2DArray->GetFormat(), Texture2DArray->GetNumMips() );
-		Texture2DArray->SetMemorySize( TextureSize );
-		Texture2DArray->SetIsPowerOfTwo(FMath::IsPowerOfTwo(Texture2DArray->GetSizeX()) && FMath::IsPowerOfTwo(Texture2DArray->GetSizeY()));
-		if (bRenderTarget)
+		switch (Desc.Dimension)
 		{
-			INC_MEMORY_STAT_BY(STAT_RenderTargetMemory2D,TextureSize);
+		default: checkNoEntry();
+		case ETextureDimension::Texture2D:		  INC_MEMORY_STAT_BY(STAT_RenderTargetMemory2D  , TextureSize); break;
+		case ETextureDimension::Texture2DArray:	  INC_MEMORY_STAT_BY(STAT_RenderTargetMemory2D  , TextureSize); break;
+		case ETextureDimension::Texture3D:		  INC_MEMORY_STAT_BY(STAT_RenderTargetMemory3D  , TextureSize); break;
+		case ETextureDimension::TextureCube:	  INC_MEMORY_STAT_BY(STAT_RenderTargetMemoryCube, TextureSize); break;
+		case ETextureDimension::TextureCubeArray: INC_MEMORY_STAT_BY(STAT_RenderTargetMemoryCube, TextureSize); break;
 		}
-		else
-		{
-			INC_MEMORY_STAT_BY(STAT_TextureMemory2D,TextureSize);
-		}
-	}
-	else
-	{
-		check(0);	// Add handling of other texture types
-	}
 
-	if( bRenderTarget )
-	{
-		GCurrentRendertargetMemorySize += Align(TextureSize, 1024) / 1024;
+		GCurrentRendertargetMemorySize += TextureSizeInKBs;
 		LLM_SCOPED_PAUSE_TRACKING_WITH_ENUM_AND_AMOUNT(ELLMTag::GraphicsPlatform, TextureSize, ELLMTracker::Platform, ELLMAllocType::None);
-		LLM_SCOPED_PAUSE_TRACKING_WITH_ENUM_AND_AMOUNT(ELLMTag::RenderTargets, TextureSize, ELLMTracker::Default, ELLMAllocType::None);
+		LLM_SCOPED_PAUSE_TRACKING_WITH_ENUM_AND_AMOUNT(ELLMTag::RenderTargets   , TextureSize, ELLMTracker::Default , ELLMAllocType::None);
 	}
 	else
 	{
-		GCurrentTextureMemorySize += Align(TextureSize, 1024) / 1024;
-		LLM_SCOPED_PAUSE_TRACKING_WITH_ENUM_AND_AMOUNT(ELLMTag::GraphicsPlatform, TextureSize, ELLMTracker::Platform, ELLMAllocType::None);
-		LLM_SCOPED_PAUSE_TRACKING_WITH_ENUM_AND_AMOUNT(ELLMTag::Textures, TextureSize, ELLMTracker::Default, ELLMAllocType::None);
-	}
-}
-
-void OpenGLTextureDeleted( FRHITexture* Texture )
-{
-	bool bRenderTarget = IsRenderTarget(Texture->GetFlags());
-	int32 TextureSize = 0;
-	if (Texture->GetTextureCube())
-	{
-		TextureSize = ((FOpenGLTextureCube*)Texture->GetTextureCube())->GetMemorySize();
-		if (bRenderTarget)
-		{
-			DEC_MEMORY_STAT_BY(STAT_RenderTargetMemoryCube,TextureSize);
-		}
-		else
-		{
-			DEC_MEMORY_STAT_BY(STAT_TextureMemoryCube,TextureSize);
-		}
-	}
-	else if (Texture->GetTexture2D())
-	{
-		TextureSize = ((FOpenGLTexture2D*)Texture->GetTexture2D())->GetMemorySize();
-		if (bRenderTarget)
-		{
-			DEC_MEMORY_STAT_BY(STAT_RenderTargetMemory2D,TextureSize);
-		}
-		else
-		{
-			DEC_MEMORY_STAT_BY(STAT_TextureMemory2D,TextureSize);
-		}
-	}
-	else if (Texture->GetTexture3D())
-	{
-		TextureSize = ((FOpenGLTexture3D*)Texture->GetTexture3D())->GetMemorySize();
-		if (bRenderTarget)
-		{
-			DEC_MEMORY_STAT_BY(STAT_RenderTargetMemory3D,TextureSize);
-		}
-		else
-		{
-			DEC_MEMORY_STAT_BY(STAT_TextureMemory3D,TextureSize);
-		}
-	}
-	else if (Texture->GetTexture2DArray())
-	{
-		TextureSize = ((FOpenGLTexture2DArray*)Texture->GetTexture2DArray())->GetMemorySize();
-		if (bRenderTarget)
-		{
-			DEC_MEMORY_STAT_BY(STAT_RenderTargetMemory2D,TextureSize);
-		}
-		else
+		switch (Desc.Dimension)
 		{
-			DEC_MEMORY_STAT_BY(STAT_TextureMemory2D,TextureSize);
+		default: checkNoEntry();
+		case ETextureDimension::Texture2D:		  INC_MEMORY_STAT_BY(STAT_TextureMemory2D  , TextureSize); break;
+		case ETextureDimension::Texture2DArray:	  INC_MEMORY_STAT_BY(STAT_TextureMemory2D  , TextureSize); break;
+		case ETextureDimension::Texture3D:		  INC_MEMORY_STAT_BY(STAT_TextureMemory3D  , TextureSize); break;
+		case ETextureDimension::TextureCube:	  INC_MEMORY_STAT_BY(STAT_TextureMemoryCube, TextureSize); break;
+		case ETextureDimension::TextureCubeArray: INC_MEMORY_STAT_BY(STAT_TextureMemoryCube, TextureSize); break;
 		}
-	}
-	else
-	{
-		check(0);	// Add handling of other texture types
-	}
 
-	if( bRenderTarget )
-	{
-		GCurrentRendertargetMemorySize -= Align(TextureSize, 1024) / 1024;
-		LLM_SCOPED_PAUSE_TRACKING_WITH_ENUM_AND_AMOUNT(ELLMTag::GraphicsPlatform, -TextureSize, ELLMTracker::Platform, ELLMAllocType::None);
-		LLM_SCOPED_PAUSE_TRACKING_WITH_ENUM_AND_AMOUNT(ELLMTag::RenderTargets, -TextureSize, ELLMTracker::Default, ELLMAllocType::None);
-	}
-	else
-	{
-		GCurrentTextureMemorySize -= Align(TextureSize, 1024) / 1024;
-		LLM_SCOPED_PAUSE_TRACKING_WITH_ENUM_AND_AMOUNT(ELLMTag::GraphicsPlatform, -TextureSize, ELLMTracker::Platform, ELLMAllocType::None);
-		LLM_SCOPED_PAUSE_TRACKING_WITH_ENUM_AND_AMOUNT(ELLMTag::Textures, -TextureSize, ELLMTracker::Default, ELLMAllocType::None);
+		GCurrentTextureMemorySize += TextureSizeInKBs;
+		LLM_SCOPED_PAUSE_TRACKING_WITH_ENUM_AND_AMOUNT(ELLMTag::GraphicsPlatform, TextureSize, ELLMTracker::Platform, ELLMAllocType::None);
+		LLM_SCOPED_PAUSE_TRACKING_WITH_ENUM_AND_AMOUNT(ELLMTag::Textures        , TextureSize, ELLMTracker::Default , ELLMAllocType::None);
 	}
 }
 
-uint64 FOpenGLDynamicRHI::RHICalcTexture2DPlatformSize(uint32 SizeX, uint32 SizeY, uint8 Format, uint32 NumMips, uint32 NumSamples, ETextureCreateFlags Flags, const FRHIResourceCreateInfo& CreateInfo, uint32& OutAlign)
-{
-	OutAlign = 0;
-	return CalcTextureSize(SizeX, SizeY, (EPixelFormat)Format, NumMips);
-}
-
-uint64 FOpenGLDynamicRHI::RHICalcTexture2DArrayPlatformSize(uint32 SizeX, uint32 SizeY, uint32 ArraySize, uint8 Format, uint32 NumMips, uint32 NumSamples, ETextureCreateFlags Flags, const FRHIResourceCreateInfo& CreateInfo, uint32& OutAlign)
-{
-	OutAlign = 0;
-	return CalcTextureSize(SizeX, SizeY, (EPixelFormat)Format, NumMips) * ArraySize;
-}
-
-uint64 FOpenGLDynamicRHI::RHICalcTexture3DPlatformSize(uint32 SizeX, uint32 SizeY, uint32 SizeZ, uint8 Format, uint32 NumMips, ETextureCreateFlags Flags, const FRHIResourceCreateInfo& CreateInfo, uint32& OutAlign)
-{
-	OutAlign = 0;
-	return CalcTextureSize3D(SizeX, SizeY, SizeZ, (EPixelFormat)Format, NumMips);
-}
-
-uint64 FOpenGLDynamicRHI::RHICalcTextureCubePlatformSize(uint32 Size, uint8 Format, uint32 NumMips, ETextureCreateFlags Flags, const FRHIResourceCreateInfo& CreateInfo, uint32& OutAlign)
+FDynamicRHI::FRHICalcTextureSizeResult FOpenGLDynamicRHI::RHICalcTexturePlatformSize(FRHITextureDesc const& Desc, uint32 FirstMipIndex)
 {
-	OutAlign = 0;
-	return CalcTextureSize(Size, Size, (EPixelFormat)Format, NumMips) * 6;
+	FDynamicRHI::FRHICalcTextureSizeResult Result;
+	Result.Size = Desc.CalcMemorySizeEstimate(FirstMipIndex);
+	Result.Align = 1;
+	return Result;
 }
 
 /**
@@ -321,247 +172,460 @@ bool FOpenGLDynamicRHI::RHIGetTextureMemoryVisualizeData( FColor* /*TextureData*
 	return false;
 }
 
-
-FRHITexture* FOpenGLDynamicRHI::CreateOpenGLTexture(uint32 SizeX, uint32 SizeY, bool bCubeTexture, bool bArrayTexture, bool bIsExternal, uint8 Format, uint32 NumMips, uint32 NumSamples, uint32 ArraySize, ETextureCreateFlags Flags, const FClearValueBinding& InClearValue, FResourceBulkDataInterface* BulkData)
+FOpenGLTextureDesc::FOpenGLTextureDesc(FRHITextureDesc const& InDesc)
+	: NumSamplesRendered (InDesc.NumSamples)
+	, NumSamplesStored   (InDesc.NumSamples)
+	, bCubemap           (InDesc.IsTextureCube())
+	, bArrayTexture      (InDesc.IsTextureArray())
+	, bStreamable        (EnumHasAnyFlags(InDesc.Flags, TexCreate_Streamable))
+	, bDepthStencil      (EnumHasAnyFlags(InDesc.Flags, TexCreate_DepthStencilTargetable))
+	, bCanCreateAsEvicted(false)
+	, bIsPowerOfTwo      (false)
+	, bTileMemDepthBuffer(false)
 {
-	// Fill in the GL resources.
-	FRHITexture* Texture = CreateOpenGLRHITextureOnly(SizeX, SizeY, bCubeTexture, bArrayTexture, bIsExternal, Format, NumMips, NumSamples, ArraySize, Flags, InClearValue, BulkData);
-
-	InitializeGLTexture(Texture, SizeX, SizeY, bCubeTexture, bArrayTexture, bIsExternal, Format, NumMips, NumSamples, ArraySize, Flags, InClearValue, BulkData);
-	return Texture;
-}
+	checkf(!bCubemap || NumSamplesStored == 1, TEXT("Texture cubes cannot be multisampled."));
+	checkf(FOpenGL::SupportsTexture3D() || (!InDesc.IsTexture3D() && !InDesc.IsTextureArray()), TEXT("Texture3D / Texture2DArray support requires FOpenGL::SupportsTexture3D()."));
 
-// Allocate only the RHIresource and its initialize FRHITexture's state.
-// note this can change the value of some input parameters.
-FRHITexture* FOpenGLDynamicRHI::CreateOpenGLRHITextureOnly(const uint32 SizeX, const uint32 SizeY, const bool bCubeTexture, const bool bArrayTexture, const bool bIsExternal, uint8& Format, uint32& NumMips, uint32& NumSamples, const uint32 ArraySize, ETextureCreateFlags Flags, const FClearValueBinding& InClearValue, FResourceBulkDataInterface* BulkData)
-{
-	SCOPE_CYCLE_COUNTER(STAT_OpenGLCreateTextureTime);
+	// Use on-chip tile memory for MSAA if available
+	if (NumSamplesStored <= FOpenGL::GetMaxMSAASamplesTileMem())
+	{
+		NumSamplesRendered = FMath::Min<uint32>(NumSamplesStored, FOpenGL::GetMaxMSAASamplesTileMem());
+		NumSamplesStored = 1;
+	}
 
-	if (NumMips == 0)
+	// Select an appropriate texture target
+#if PLATFORM_ANDROID
+	if (bDepthStencil && NumSamplesRendered != NumSamplesStored)
+	{
+		// Special case for MSAA depth render target on tiled GPUs / mobile renderer.
+		Target = GL_RENDERBUFFER;
+		bTileMemDepthBuffer = true;
+	}
+	else 
+#endif
 	{
-		if (NumSamples <= 1)
+		if (EnumHasAnyFlags(InDesc.Flags, TexCreate_External))
+		{
+			check(InDesc.IsTexture2D());
+			check(!InDesc.IsTextureArray());
+
+			Target = FOpenGL::SupportsImageExternal()
+				? GL_TEXTURE_EXTERNAL_OES
+				// Fall back to a regular 2d texture if we don't have support.
+				// Texture samplers in the shader will also fall back to a regular sampler2D.
+				: GL_TEXTURE_2D;
+		}
+		else if (EnumHasAnyFlags(InDesc.Flags, TexCreate_Presentable))
 		{
-			NumMips = FindMaxMipmapLevel(SizeX, SizeY);
+			check(InDesc.Dimension == ETextureDimension::Texture2D);
+			Target = GL_RENDERBUFFER;
 		}
 		else
 		{
-			NumMips = 1;
+			switch (InDesc.Dimension)
+			{
+			default: checkNoEntry();
+			case ETextureDimension::Texture2D:		  Target = (NumSamplesStored > 1) ? GL_TEXTURE_2D_MULTISAMPLE       : GL_TEXTURE_2D;       break;
+			case ETextureDimension::Texture2DArray:	  Target = (NumSamplesStored > 1) ? GL_TEXTURE_2D_MULTISAMPLE_ARRAY : GL_TEXTURE_2D_ARRAY; break;
+			case ETextureDimension::TextureCubeArray: Target = GL_TEXTURE_CUBE_MAP_ARRAY; break;
+			case ETextureDimension::TextureCube:	  Target = GL_TEXTURE_CUBE_MAP;       break;
+			case ETextureDimension::Texture3D:		  Target = GL_TEXTURE_3D;             break;
+			}
 		}
 	}
+	check(Target != GL_NONE);
 
-#if UE_BUILD_DEBUG
-	check(!(NumSamples > 1 && bCubeTexture));
-	check(bArrayTexture != (ArraySize == 1));
-#endif
+	// can run on RT.
+	bCanCreateAsEvicted =
+		CanDeferTextureCreation()
+		&& FOpenGL::SupportsCopyImage()
+		&& InDesc.Flags != TexCreate_None // ignore TexCreate_None
+		&& !EnumHasAnyFlags((ETextureCreateFlags)CVarDeferTextureCreationExcludeMask.GetValueOnAnyThread(), InDesc.Flags)  // Anything outside of these flags cannot be evicted.
+		&& Target == GL_TEXTURE_2D
+		&& InDesc.IsTexture2D(); // 2d only.
 
-	// Move NumSamples to on-chip MSAA if supported
-	uint32 NumSamplesTileMem = 1;
-	GLint MaxSamplesTileMem = FOpenGL::GetMaxMSAASamplesTileMem(); /* RHIs which do not support tiled GPU MSAA return 1 */
-	if (MaxSamplesTileMem > 1)
+	if (GOGLTextureEvictLogging)
 	{
-		NumSamplesTileMem = FMath::Min<uint32>(NumSamples, MaxSamplesTileMem);
-		NumSamples = 1;
+		UE_CLOG(!bCanCreateAsEvicted, LogRHI, Warning, TEXT("CanDeferTextureCreation:%d, SupportsCopyImage:%d, Flags:%llx Mask:%x, Target:%x"),
+			bCanCreateAsEvicted, FOpenGL::SupportsCopyImage(), InDesc.Flags, CVarDeferTextureCreationExcludeMask.GetValueOnAnyThread(), Target);
 	}
 
-	GLenum Target = GL_NONE;
-	if (bCubeTexture)
+	bIsPowerOfTwo = 
+		   FMath::IsPowerOfTwo(InDesc.Extent.X)
+		&& FMath::IsPowerOfTwo(InDesc.Extent.Y)
+		&& FMath::IsPowerOfTwo(InDesc.Depth);
+
+	MemorySize = InDesc.CalcMemorySizeEstimate();
+
+	// Determine the attachment point for the texture.	
+	if (EnumHasAnyFlags(InDesc.Flags, TexCreate_RenderTargetable | TexCreate_CPUReadback))
 	{
-		if (FOpenGL::SupportsTexture3D())
-		{
-			Target = bArrayTexture ? GL_TEXTURE_CUBE_MAP_ARRAY : GL_TEXTURE_CUBE_MAP;
-		}
-		else
-		{
-			check(!bArrayTexture);
-			Target = GL_TEXTURE_CUBE_MAP;
-		}
-		check(SizeX == SizeY);
+		Attachment = GL_COLOR_ATTACHMENT0;
 	}
-#if PLATFORM_ANDROID
-	else if (bIsExternal)
+	else if (EnumHasAnyFlags(InDesc.Flags, TexCreate_DepthStencilTargetable))
+	{
+		Attachment = (InDesc.Format == PF_DepthStencil) ? GL_DEPTH_STENCIL_ATTACHMENT : GL_DEPTH_ATTACHMENT;
+	}
+	else if (EnumHasAnyFlags(InDesc.Flags, TexCreate_ResolveTargetable))
+	{
+		Attachment = (InDesc.Format == PF_DepthStencil)
+			? GL_DEPTH_STENCIL_ATTACHMENT
+			: ((InDesc.Format == PF_ShadowDepth || InDesc.Format == PF_D24)
+				? GL_DEPTH_ATTACHMENT
+				: GL_COLOR_ATTACHMENT0);
+	}
+	else
+	{
+		Attachment = GL_NONE;
+	}
+
+	switch (Attachment)
+	{
+	case GL_COLOR_ATTACHMENT0:
+		check(GMaxOpenGLColorSamples >= (GLint)NumSamplesRendered);
+		break;
+	case GL_DEPTH_ATTACHMENT:
+	case GL_DEPTH_STENCIL_ATTACHMENT:
+		check(GMaxOpenGLDepthSamples >= (GLint)NumSamplesRendered);
+		break;
+	default:
+		break;
+	}
+}
+
+// Constructor for RHICreateAliasedTexture
+FOpenGLTexture::FOpenGLTexture(FOpenGLTexture& Other, EAliasConstructorParam)
+	: FRHITexture(FRHITextureCreateDesc(Other.GetDesc(), ERHIAccess::SRVMask, *Other.GetName().ToString()))
+	, Target             (Other.Target)
+	, Attachment         (Other.Attachment)
+	, MemorySize         (0)
+	, NumSamplesInternal (Other.NumSamplesInternal)
+	, bIsPowerOfTwo      (Other.bIsPowerOfTwo)
+	, bCanCreateAsEvicted(false)
+	, bStreamable        (Other.bStreamable)
+	, bCubemap           (Other.bCubemap)
+	, bArrayTexture      (Other.bArrayTexture)
+	, bDepthStencil      (Other.bDepthStencil)
+	, bTileMemDepthBuffer(Other.bTileMemDepthBuffer)
+	, bAlias             (true)
+{
+	RunOnGLRenderContextThread([&]()
+	{
+		AliasResources(Other);
+	});
+}
+
+void FOpenGLTexture::AliasResources(FOpenGLTexture& Texture)
+{
+	VERIFY_GL_SCOPE();
+	check(bAlias && !Texture.bAlias);
+
+	// restore the source texture, do not allow the texture to become evicted, the aliasing texture cannot re-create the resource.
+	if (Texture.IsEvicted())
+	{
+		Texture.RestoreEvictedGLResource(false);
+	}
+
+	Resource = Texture.Resource;
+}
+
+// Constructor for external resources (RHICreateTexture2DFromResource etc).
+FOpenGLTexture::FOpenGLTexture(FOpenGLTextureCreateDesc const& CreateDesc, GLuint InResource)
+	: FRHITexture        (CreateDesc)
+	, Resource           (InResource)
+	, Target             (CreateDesc.Target)
+	, Attachment         (CreateDesc.Attachment)
+	, MemorySize         (CreateDesc.MemorySize)
+	, NumSamplesInternal (CreateDesc.NumSamplesRendered)
+	, bIsPowerOfTwo      (CreateDesc.bIsPowerOfTwo)
+	, bCanCreateAsEvicted(false)
+	, bStreamable        (CreateDesc.bStreamable)
+	, bCubemap           (CreateDesc.bCubemap)
+	, bArrayTexture      (CreateDesc.bArrayTexture)
+	, bDepthStencil      (CreateDesc.bDepthStencil)
+	, bTileMemDepthBuffer(CreateDesc.bTileMemDepthBuffer)
+	, bAlias             (true)
+{}
+
+// Standard constructor.
+FOpenGLTexture::FOpenGLTexture(FOpenGLTextureCreateDesc const& CreateDesc)
+	: FRHITexture        (CreateDesc)
+	, Target             (CreateDesc.Target)
+	, Attachment         (CreateDesc.Attachment)
+	, MemorySize         (CreateDesc.MemorySize)
+	, NumSamplesInternal (CreateDesc.NumSamplesRendered)
+	, bIsPowerOfTwo      (CreateDesc.bIsPowerOfTwo)
+	, bCanCreateAsEvicted(CreateDesc.bCanCreateAsEvicted)
+	, bStreamable        (CreateDesc.bStreamable)
+	, bCubemap           (CreateDesc.bCubemap)
+	, bArrayTexture      (CreateDesc.bArrayTexture)
+	, bDepthStencil      (CreateDesc.bDepthStencil)
+	, bTileMemDepthBuffer(CreateDesc.bTileMemDepthBuffer)
+	, bAlias             (false)
+{
+	check(IsInRenderingThread());
+	FRHICommandListImmediate& RHICmdList = FRHICommandListExecutor::GetImmediateCommandList();
+
+	SCOPE_CYCLE_COUNTER(STAT_OpenGLCreateTextureTime);
+
+	if (bCanCreateAsEvicted)
+	{
+		EvictionParamsPtr = MakeUnique<FTextureEvictionParams>(CreateDesc.NumMips);
+	}
+
+	void* BulkDataPtr = nullptr;
+	uint64 BulkDataSize = 0;
+	bool bFreeBulkData = false;
+
+	if (CreateDesc.BulkData)
 	{
-		if (FOpenGL::SupportsImageExternal())
+		if (!ShouldRunGLRenderContextOpOnThisThread(RHICmdList))
 		{
-			Target = GL_TEXTURE_EXTERNAL_OES;
+			// If bulk data is provided, and texture initialization is done by the RHI thread, it needs to be copied out of the FResourceBulkDataInterface.
+			// It is not safe to pass this pointer to the RHI thread, as the interface may have been stack allocated in the renderer.
+			// @todo: remove this memcpy when the new resource creation API is available.
+			BulkDataSize = CreateDesc.BulkData->GetResourceBulkDataSize();
+			BulkDataPtr = FMemory::Malloc(BulkDataSize);
+
+			FMemory::Memcpy(BulkDataPtr, CreateDesc.BulkData->GetResourceBulkData(), BulkDataSize);
+			bFreeBulkData = true;
 		}
 		else
 		{
-			// Fall back to a regular 2d texture if we don't have support. Texture samplers in the shader will also fall back to a regular sampler2D.
-			Target = GL_TEXTURE_2D;
+			// Otherwise, initialization will be done on this thread.
+			// Just use the raw pointer / size as-is.
+			BulkDataSize = CreateDesc.BulkData->GetResourceBulkDataSize();
+			BulkDataPtr = const_cast<void*>(CreateDesc.BulkData->GetResourceBulkData());
 		}
 	}
-#endif
-	else
+
+	RunOnGLRenderContextThread([this, BulkDataPtr, BulkDataSize, bFreeBulkData]()
 	{
-		Target =  (NumSamples > 1) ? GL_TEXTURE_2D_MULTISAMPLE : GL_TEXTURE_2D;
+		FOpenGLDynamicRHI::Get().InitializeGLTexture(this, BulkDataPtr, BulkDataSize);
+		if (bFreeBulkData)
+		{
+			FMemory::Free(BulkDataPtr);
+		}
+	});
 
-		// @todo: refactor 2d texture array support here?
-		check(!bArrayTexture);
-	}
-	check(Target != GL_NONE);
+	UpdateTextureStats(this, true);
 
+	PixelBuffers.AddZeroed(CreateDesc.NumMips * (bCubemap ? 6 : 1) * GetEffectiveSizeZ());
 
-	FRHITexture* Result;
-	// Allocate RHIResource with empty GL values.
-	if (bCubeTexture)
+	if (CreateDesc.BulkData)
 	{
-		Result = new FOpenGLTextureCube(this, 0, Target, -1, SizeX, SizeY, 0, NumMips, 1, 1, ArraySize, (EPixelFormat)Format, true, false, Flags, InClearValue);
+		CreateDesc.BulkData->Discard();
 	}
-	else
+}
+
+FOpenGLTexture::~FOpenGLTexture()
+{
+	VERIFY_GL_SCOPE();
+
+	FTextureEvictionLRU::Get().Remove(this);
+
+	if (!bCanCreateAsEvicted)
 	{
-		Result = new FOpenGLTexture2D(this, 0, Target, -1, SizeX, SizeY, 0, NumMips, NumSamples, NumSamplesTileMem, 1, (EPixelFormat)Format, false, false, Flags, InClearValue);
+		ReleaseOpenGLFramebuffers(this);
 	}
-	OpenGLTextureAllocated(Result, Flags);
 
-	check(!GetOpenGLTextureFromRHITexture(Result)->EvictionParamsPtr.IsValid());
+	DeleteGLResource();
+	UpdateTextureStats(this, false);
+}
+
+void FOpenGLTexture::DeleteGLResource()
+{
+	VERIFY_GL_SCOPE();
+	QUICK_SCOPE_CYCLE_COUNTER(STAT_OpenGLDeleteGLTextureTime);
 
-	if (GetOpenGLTextureFromRHITexture(Result)->CanCreateAsEvicted())
+	if (Resource != 0)
 	{
-		GetOpenGLTextureFromRHITexture(Result)->EvictionParamsPtr = MakeUnique<FTextureEvictionParams>(NumMips);
+		switch (Target)
+		{
+		case GL_TEXTURE_2D:
+		case GL_TEXTURE_2D_MULTISAMPLE:
+		case GL_TEXTURE_3D:
+		case GL_TEXTURE_CUBE_MAP:
+		case GL_TEXTURE_2D_ARRAY:
+		case GL_TEXTURE_CUBE_MAP_ARRAY:
+		case GL_TEXTURE_EXTERNAL_OES:
+			FOpenGLDynamicRHI::Get().InvalidateTextureResourceInCache(Resource);
+			if (!bAlias)
+			{
+				FOpenGL::DeleteTextures(1, &Resource);
+			}
+			break;
+
+		case GL_RENDERBUFFER:
+			if (!bAlias)
+			{
+				glDeleteRenderbuffers(1, &Resource);
+			}
+			break;
+
+		default:
+			checkNoEntry();
+			break;
+		}
 	}
+}
 
-	return Result;
+static inline bool IsAstcLdrRGBAFormat(GLenum Format)
+{
+	return Format >= GL_COMPRESSED_RGBA_ASTC_4x4_KHR && Format <= GL_COMPRESSED_RGBA_ASTC_12x12_KHR;
 }
 
 uint32 GTotalTexStorageSkipped = 0;
 uint32 GTotalCompressedTexStorageSkipped = 0;
-// Initalize the FRHITexture's GL resources and fill in state.
-void FOpenGLDynamicRHI::InitializeGLTexture(FRHITexture* Texture, uint32 SizeX, const uint32 SizeY, const bool bCubeTexture, const bool bArrayTexture, const bool bIsExternal, const uint8 Format, const uint32 NumMips, const uint32 NumSamples, const uint32 ArraySize, const ETextureCreateFlags Flags, const FClearValueBinding& InClearValue, FResourceBulkDataInterface* BulkData)
+
+void FOpenGLDynamicRHI::InitializeGLTexture(FOpenGLTexture* Texture, const void* BulkDataPtr, uint64 BulkDataSize)
 {
 	VERIFY_GL_SCOPE();
 
-	const uint32 NumSamplesTileMem = bCubeTexture ? 1 : ((FOpenGLTexture2D*)Texture)->GetNumSamplesTileMem();
-	const bool TileMemDepth = NumSamplesTileMem > 1 && EnumHasAnyFlags(Flags, TexCreate_DepthStencilTargetable);
+	if (EnumHasAnyFlags(Texture->GetDesc().Flags, TexCreate_Presentable))
+		return;
 
-	GLuint TextureID = 0;
-	if (!TileMemDepth)
+	// Allocate the GL resource ID
+	GLuint TextureID;
+	if (Texture->bTileMemDepthBuffer)
+	{
+		check(Texture->Target == GL_RENDERBUFFER);
+		glGenRenderbuffers(1, &TextureID);
+	}
+	else
 	{
-		FOpenGL::GenTextures(1, &TextureID);
+		check(Texture->Target != GL_RENDERBUFFER);
+		glGenTextures(1, &TextureID);
 	}
-	if (!GetOpenGLTextureFromRHITexture(Texture)->IsEvicted())
+	Texture->SetResource(TextureID);
+
+	if (!Texture->IsEvicted())
 	{
-		InitializeGLTextureInternal(TextureID, Texture, SizeX, SizeY, bCubeTexture, bArrayTexture, bIsExternal, Format, NumMips, NumSamples, ArraySize, Flags, InClearValue, BulkData);
+		InitializeGLTextureInternal(Texture, BulkDataPtr, BulkDataSize);
 	}
 	else
 	{
 		// creating this as 'evicted'.
 		GTotalTexStorageSkipped++;
-		//check(!GetOpenGLTextureFromRHITexture(Texture)->IsEvicted());
+
+		EPixelFormat PixelFormat = Texture->GetFormat();
+		const FOpenGLTextureFormat& GLFormat = GOpenGLTextureFormats[PixelFormat];
+		bool bIsCompressed = GLFormat.bCompressed;
+		GTotalCompressedTexStorageSkipped += bIsCompressed ? 1 : 0;
+
+		if (BulkDataPtr)
 		{
-			EPixelFormat PixelFormat = Texture->GetFormat();
-			const FOpenGLTextureFormat& GLFormat = GOpenGLTextureFormats[PixelFormat];
-			bool bIsCompressed = GLFormat.bCompressed;
-			GTotalCompressedTexStorageSkipped += bIsCompressed ? 1 : 0;
+			check(!GLFormat.bCompressed);
+			const uint32 BlockSizeX = GPixelFormats[PixelFormat].BlockSizeX;
+			const uint32 BlockSizeY = GPixelFormats[PixelFormat].BlockSizeY;
 
-			if(BulkData)
-			{
+			uint8* Data = (uint8*)BulkDataPtr;
+			uint32 MipOffset = 0;
 
-				check(!GLFormat.bCompressed);
-				const uint32 BlockSizeX = GPixelFormats[Format].BlockSizeX;
-				const uint32 BlockSizeY = GPixelFormats[Format].BlockSizeY;
-				uint8* Data = (uint8*)BulkData->GetResourceBulkData();
-				uint32 MipOffset = 0;
-				// copy bulk data to evicted mip store:
-				for (uint32 MipIndex = 0; MipIndex < NumMips; MipIndex++)
-				{
-					uint32 NumBlocksX = AlignArbitrary(FMath::Max<uint32>(1, (SizeX >> MipIndex)), BlockSizeX) / BlockSizeX;
-					uint32 NumBlocksY = AlignArbitrary(FMath::Max<uint32>(1, (SizeY >> MipIndex)), BlockSizeY) / BlockSizeY;
-					uint32 NumLayers = FMath::Max<uint32>(1, ArraySize);
-					uint32 MipDataSize = NumBlocksX * NumBlocksY * NumLayers * GPixelFormats[Format].BlockBytes;
+			const FRHITextureDesc& Desc = Texture->GetDesc();
 
-					GetOpenGLTextureFromRHITexture(Texture)->EvictionParamsPtr->SetMipData(MipIndex, &Data[MipOffset],MipDataSize);
-					MipOffset += MipDataSize;
-				}
-				BulkData->Discard();
+			// copy bulk data to evicted mip store:
+			for (uint32 MipIndex = 0; MipIndex < Desc.NumMips; MipIndex++)
+			{
+				uint32 NumBlocksX = AlignArbitrary(FMath::Max<uint32>(1, (Desc.Extent.X >> MipIndex)), BlockSizeX) / BlockSizeX;
+				uint32 NumBlocksY = AlignArbitrary(FMath::Max<uint32>(1, (Desc.Extent.Y >> MipIndex)), BlockSizeY) / BlockSizeY;
+				uint32 NumLayers = FMath::Max<uint32>(1, Desc.ArraySize);
+				uint32 MipDataSize = NumBlocksX * NumBlocksY * NumLayers * GPixelFormats[PixelFormat].BlockBytes;
+
+				Texture->EvictionParamsPtr->SetMipData(MipIndex, &Data[MipOffset], MipDataSize);
+				MipOffset += MipDataSize;
 			}
 		}
-		GetOpenGLTextureFromRHITexture(Texture)->SetResource(TextureID);
 	}
 }
 
-static inline bool IsAstcLdrRGBAFormat(GLenum Format)
-{
-	return Format >= GL_COMPRESSED_RGBA_ASTC_4x4_KHR && Format <= GL_COMPRESSED_RGBA_ASTC_12x12_KHR;
-}
-
-void FOpenGLDynamicRHI::InitializeGLTextureInternal(GLuint TextureID, FRHITexture* Texture, uint32 SizeX, const uint32 SizeY, const bool bCubeTexture, const bool bArrayTexture, const bool bIsExternal, const uint8 Format, const uint32 NumMips, const uint32 NumSamples, const uint32 ArraySize, const ETextureCreateFlags Flags, const FClearValueBinding& InClearValue, FResourceBulkDataInterface* BulkData)
+void FOpenGLDynamicRHI::InitializeGLTextureInternal(FOpenGLTexture* Texture, void const* BulkDataPtr, uint64 BulkDataSize)
 {
 	VERIFY_GL_SCOPE();
 
-	bool bAllocatedStorage = false;
-
-	GLenum Target = bCubeTexture ? ((FOpenGLTextureCube*)Texture)->Target : ((FOpenGLTexture2D*)Texture)->Target;
-	const uint32 NumSamplesTileMem = bCubeTexture ? 1 : ((FOpenGLTexture2D*)Texture)->GetNumSamplesTileMem();
-	const bool TileMemDepth = NumSamplesTileMem > 1 && EnumHasAnyFlags(Flags, TexCreate_DepthStencilTargetable);
+	GLuint const TextureID = Texture->GetRawResourceName();
 
-	check(TextureID || TileMemDepth);
+	const FRHITextureDesc& Desc = Texture->GetDesc();
+	const GLenum Target = Texture->Target;
 
-	const bool bSRGB = EnumHasAnyFlags(Flags, TexCreate_SRGB);
-	const FOpenGLTextureFormat& GLFormat = GOpenGLTextureFormats[Format];
+	const bool bSRGB = EnumHasAnyFlags(Desc.Flags, TexCreate_SRGB);
+	const FOpenGLTextureFormat& GLFormat = GOpenGLTextureFormats[Desc.Format];
 	if (GLFormat.InternalFormat[bSRGB] == GL_NONE)
 	{
-		UE_LOG(LogRHI, Fatal,TEXT("Texture format '%s' not supported (sRGB=%d)."), GPixelFormats[Format].Name, bSRGB);
+		UE_LOG(LogRHI, Fatal,TEXT("Texture format '%s' not supported (sRGB=%d)."), GPixelFormats[Desc.Format].Name, bSRGB);
 	}
 
 	FOpenGLContextState& ContextState = GetContextStateForCurrentContext();
 
 	// Make sure PBO is disabled
-	CachedBindPixelUnpackBuffer(ContextState,0);
+	CachedBindPixelUnpackBuffer(ContextState, 0);
 
-	// Use a texture stage that's not likely to be used for draws, to avoid waiting
-	CachedSetupTextureStage(ContextState, FOpenGL::GetMaxCombinedTextureImageUnits() - 1, Target, TextureID, 0, NumMips);
-#if !PLATFORM_ANDROID
-	if(GLFormat.InternalFormat[0] == GL_RGB5_A1)
+	bool bAllocatedStorage = false;
+#if PLATFORM_ANDROID
+	if (Texture->bTileMemDepthBuffer)
 	{
-		glTexParameteri(Target, GL_TEXTURE_SWIZZLE_R, GL_BLUE);
-		glTexParameteri(Target, GL_TEXTURE_SWIZZLE_B, GL_RED);
+		glBindRenderbuffer(GL_RENDERBUFFER, TextureID);
+
+		glRenderbufferStorageMultisampleEXT(GL_RENDERBUFFER, Texture->GetNumSamplesRendered(), GL_DEPTH24_STENCIL8, Desc.Extent.X, Desc.Extent.Y);
+		VERIFY_GL(glRenderbufferStorageMultisampleEXT);
+
+		glBindRenderbuffer(GL_RENDERBUFFER, 0);
+
+		bAllocatedStorage = true;
 	}
+	else
 #endif
-	if (NumSamples == 1 && !TileMemDepth)
 	{
-		if (Target == GL_TEXTURE_EXTERNAL_OES || !FMath::IsPowerOfTwo(SizeX) || !FMath::IsPowerOfTwo(SizeY))
+		// Use a texture stage that's not likely to be used for draws, to avoid waiting
+		CachedSetupTextureStage(ContextState, FOpenGL::GetMaxCombinedTextureImageUnits() - 1, Target, TextureID, 0, Desc.NumMips);
+
+		if((GLFormat.bBGRA && !EnumHasAnyFlags(Desc.Flags, TexCreate_RenderTargetable)) ||
+			((!PLATFORM_ANDROID) && GLFormat.InternalFormat[0] == GL_RGB5_A1))
 		{
-			glTexParameteri(Target, GL_TEXTURE_WRAP_S, GL_CLAMP_TO_EDGE);
-			glTexParameteri(Target, GL_TEXTURE_WRAP_T, GL_CLAMP_TO_EDGE);
-			if ( FOpenGL::SupportsTexture3D() )
-			{
-				glTexParameteri(Target, GL_TEXTURE_WRAP_R, GL_CLAMP_TO_EDGE);
-			}
+			glTexParameteri(Target, GL_TEXTURE_SWIZZLE_R, GL_BLUE);
+			glTexParameteri(Target, GL_TEXTURE_SWIZZLE_B, GL_RED);
 		}
-		else
+
+		if (Texture->GetNumSamplesStored() == 1)
 		{
-			glTexParameteri(Target, GL_TEXTURE_WRAP_S, GL_REPEAT);
-			glTexParameteri(Target, GL_TEXTURE_WRAP_T, GL_REPEAT);
-			if ( FOpenGL::SupportsTexture3D() )
+			if (Target == GL_TEXTURE_EXTERNAL_OES || !FMath::IsPowerOfTwo(Desc.Extent.X) || !FMath::IsPowerOfTwo(Desc.Extent.Y))
 			{
-				glTexParameteri(Target, GL_TEXTURE_WRAP_R, GL_REPEAT);
+				glTexParameteri(Target, GL_TEXTURE_WRAP_S, GL_CLAMP_TO_EDGE);
+				glTexParameteri(Target, GL_TEXTURE_WRAP_T, GL_CLAMP_TO_EDGE);
+				if (FOpenGL::SupportsTexture3D())
+				{
+					glTexParameteri(Target, GL_TEXTURE_WRAP_R, GL_CLAMP_TO_EDGE);
+				}
+			}
+			else
+			{
+				glTexParameteri(Target, GL_TEXTURE_WRAP_S, GL_REPEAT);
+				glTexParameteri(Target, GL_TEXTURE_WRAP_T, GL_REPEAT);
+				if (FOpenGL::SupportsTexture3D())
+				{
+					glTexParameteri(Target, GL_TEXTURE_WRAP_R, GL_REPEAT);
+				}
+			}
+
+			glTexParameteri(Target, GL_TEXTURE_MAG_FILTER, GL_NEAREST);
+			glTexParameteri(Target, GL_TEXTURE_MIN_FILTER, Desc.NumMips > 1 ? GL_NEAREST_MIPMAP_NEAREST : GL_NEAREST);
+
+			if (FOpenGL::SupportsTextureFilterAnisotropic())
+			{
+				glTexParameteri(Target, GL_TEXTURE_MAX_ANISOTROPY_EXT, 1);
 			}
 		}
-		glTexParameteri(Target, GL_TEXTURE_MAG_FILTER, GL_NEAREST);
-		glTexParameteri(Target, GL_TEXTURE_MIN_FILTER, NumMips > 1 ? GL_NEAREST_MIPMAP_NEAREST : GL_NEAREST);
-		if( FOpenGL::SupportsTextureFilterAnisotropic() )
-		{
-			glTexParameteri(Target, GL_TEXTURE_MAX_ANISOTROPY_EXT, 1);
-		}
-		
+
 		glTexParameteri(Target, GL_TEXTURE_BASE_LEVEL, 0);
-	
-#if PLATFORM_ANDROID
-		// Do not use GL_TEXTURE_MAX_LEVEL if external texture on Android
+
+		// Do not use GL_TEXTURE_MAX_LEVEL if external texture
 		if (Target != GL_TEXTURE_EXTERNAL_OES)
-#endif
 		{
-			glTexParameteri(Target, GL_TEXTURE_MAX_LEVEL, NumMips - 1);
-		}
-		
-		TextureMipLimits.Add(TextureID, TPair<GLenum, GLenum>(0, NumMips - 1));
-		
-		if (GLFormat.bBGRA && !EnumHasAnyFlags(Flags, TexCreate_RenderTargetable))
-		{
-			glTexParameteri(Target, GL_TEXTURE_SWIZZLE_R, GL_BLUE);
-			glTexParameteri(Target, GL_TEXTURE_SWIZZLE_B, GL_RED);
+			glTexParameteri(Target, GL_TEXTURE_MAX_LEVEL, Desc.NumMips - 1);
 		}
 
+		TextureMipLimits.Add(TextureID, TPair<GLenum, GLenum>(0, Desc.NumMips - 1));
+
 		if (FOpenGL::SupportsASTCDecodeMode())
 		{
 			if (IsAstcLdrRGBAFormat(GLFormat.InternalFormat[bSRGB]))
@@ -570,216 +634,235 @@ void FOpenGLDynamicRHI::InitializeGLTextureInternal(GLuint TextureID, FRHITextur
 			}
 		}
 
-		if (bArrayTexture)
-		{
-			FOpenGL::TexStorage3D( Target, NumMips, GLFormat.InternalFormat[bSRGB], SizeX, SizeY, ArraySize, GLFormat.Format, GLFormat.Type);
-		}
-		else if (Target != GL_TEXTURE_EXTERNAL_OES)
+		if (Target != GL_TEXTURE_EXTERNAL_OES)
 		{
-			// Try to allocate using TexStorage2D
-			if (FOpenGL::TexStorage2D(Target, NumMips, GLFormat.SizedInternalFormat[bSRGB], SizeX, SizeY, GLFormat.Format, GLFormat.Type, Flags))
+			auto EnumerateSubresources = [&](void const* Data, TFunctionRef<bool(GLenum Target, uint32 SizeX, uint32 SizeY, uint32 SizeZ, uint32 MipIndex, uint32 ArraySlice, void const* MipSliceData, uint32 MipSliceSize)> Callback)
 			{
-				bAllocatedStorage = true;
-			}
-			else if (!GLFormat.bCompressed)
-			{
-				// Otherwise, allocate storage for each mip using TexImage2D
-				// We can't do so for compressed textures because we can't pass NULL in to CompressedTexImage2D!
-				bAllocatedStorage = true;
+				struct FScopedPackAlignment
+				{
+					 FScopedPackAlignment() { glPixelStorei(GL_UNPACK_ALIGNMENT, 1); }
+					~FScopedPackAlignment() { glPixelStorei(GL_UNPACK_ALIGNMENT, 4); }
+				} PackAlignment;
 
-				const bool bIsCubeTexture = Target == GL_TEXTURE_CUBE_MAP;
-				const GLenum FirstTarget = bIsCubeTexture ? GL_TEXTURE_CUBE_MAP_POSITIVE_X : Target;
-				const uint32 NumTargets = bIsCubeTexture ? 6 : 1;
+				uint64 DataOffset = 0;
 
-				for(uint32 MipIndex = 0; MipIndex < uint32(NumMips); MipIndex++)
+				for (uint32 MipIndex = 0; MipIndex < Desc.NumMips; ++MipIndex)
 				{
-					for(uint32 TargetIndex = 0; TargetIndex < NumTargets; TargetIndex++)
+					for (uint32 ArraySlice = 0; ArraySlice < Desc.ArraySize; ++ArraySlice)
 					{
-						glTexImage2D(
-							FirstTarget + TargetIndex,
-							MipIndex,
-							GLFormat.InternalFormat[bSRGB],
-							FMath::Max<uint32>(1,(SizeX >> MipIndex)),
-							FMath::Max<uint32>(1,(SizeY >> MipIndex)),
-							0,
-							GLFormat.Format,
-							GLFormat.Type,
-							NULL
-							);
+						for (uint32 FaceIndex = 0; FaceIndex < (Desc.IsTextureCube() ? 6u : 1u); ++FaceIndex)
+						{
+							const uint32 MipPixelSizeX = FMath::Max<uint32>(1u, Desc.Extent.X >> MipIndex);
+							const uint32 MipPixelSizeY = FMath::Max<uint32>(1u, Desc.Extent.Y >> MipIndex);
+							const uint32 MipPixelSizeZ = FMath::Max<uint32>(1u, Desc.Depth >> MipIndex);
+
+							const GLenum CurrentTarget = Target == GL_TEXTURE_CUBE_MAP
+								? GL_TEXTURE_CUBE_MAP_POSITIVE_X + FaceIndex
+								: Target;
+
+							if (Data)
+							{
+								const uint32 MipBlockSizeX = FMath::DivideAndRoundUp<uint32>(MipPixelSizeX, GPixelFormats[Desc.Format].BlockSizeX);
+								const uint32 MipBlockSizeY = FMath::DivideAndRoundUp<uint32>(MipPixelSizeY, GPixelFormats[Desc.Format].BlockSizeY);
+								const uint32 MipBlockSizeZ = FMath::DivideAndRoundUp<uint32>(MipPixelSizeZ, GPixelFormats[Desc.Format].BlockSizeZ);
+
+								const uint32 MipNumBlocks = MipBlockSizeX * MipBlockSizeY * MipBlockSizeZ;
+								const uint32 MipSize = MipNumBlocks * GPixelFormats[Desc.Format].BlockBytes * Desc.ArraySize;
+
+								if (!Callback(CurrentTarget, MipPixelSizeX, MipPixelSizeY, MipPixelSizeZ, MipIndex, ArraySlice, Data, MipSize))
+									return;
+
+								DataOffset += MipSize;
+								if (DataOffset >= BulkDataSize)
+								{
+									// Reach the end of bulk data. Only pass nullptr to the callback for any subsequent mips / slices
+									Data = nullptr;
+								}
+							}
+							else
+							{
+								if (!Callback(CurrentTarget, MipPixelSizeX, MipPixelSizeY, MipPixelSizeZ, MipIndex, ArraySlice, nullptr, 0))
+									return;
+							}
+						}
 					}
 				}
-			}
-		}
-
-		if (BulkData != NULL)
-		{
-			uint8* Data = (uint8*)BulkData->GetResourceBulkData();
-			uint32 MipOffset = 0;
+			};
 
-			const uint32 BlockSizeX = GPixelFormats[Format].BlockSizeX;
-			const uint32 BlockSizeY = GPixelFormats[Format].BlockSizeY;
-			for(uint32 MipIndex = 0; MipIndex < NumMips; MipIndex++)
+			// Create the texture resource
+			switch (Target)
 			{
-				uint32 NumBlocksX = AlignArbitrary(FMath::Max<uint32>(1,(SizeX >> MipIndex)), BlockSizeX) / BlockSizeX;
-				uint32 NumBlocksY = AlignArbitrary(FMath::Max<uint32>(1,(SizeY >> MipIndex)), BlockSizeY) / BlockSizeY;
-				uint32 NumLayers = FMath::Max<uint32>(1,ArraySize);
-				
-				glPixelStorei(GL_UNPACK_ALIGNMENT, 1);
-
-				if(bArrayTexture )
+			default: checkNoEntry();
+			case GL_RENDERBUFFER:
+			case GL_TEXTURE_2D:
+			case GL_TEXTURE_CUBE_MAP:
+				// Try to create the texture using immutable storage
+				if (FOpenGL::TexStorage2D(Target, Desc.NumMips, GLFormat.InternalFormat[bSRGB], Desc.Extent.X, Desc.Extent.Y, GLFormat.Format, GLFormat.Type, Desc.Flags))
 				{
-					if(bCubeTexture)
-					{
-						check(FOpenGL::SupportsTexture3D());
-						FOpenGL::TexSubImage3D(
-							/*Target=*/ Target,
-							/*Level=*/ MipIndex,
-							/* XOffset */ 0,
-							/* YOffset */ 0,
-							/* ZOffset */ 0,
-							/*SizeX=*/ FMath::Max<uint32>(1,(SizeX >> MipIndex)),
-							/*SizeY=*/ FMath::Max<uint32>(1,(SizeY >> MipIndex)),
-							/*SizeZ=*/ ArraySize,
-							/*Format=*/ GLFormat.Format,
-							/*Type=*/ GLFormat.Type,
-							/*Data=*/ &Data[MipOffset]
-							);
-					}
-					else
+					// Texture created with immutable storage. Now fill in the bulk data.
+					bAllocatedStorage = true;
+
+					if (BulkDataPtr)
 					{
-						// @todo: refactor 2d texture arrays here?
-						check(!bCubeTexture);
+						EnumerateSubresources(BulkDataPtr, [&](GLenum CurrentTarget, uint32 MipSizeX, uint32 MipSizeY, uint32 MipSizeZ, uint32 MipIndex, uint32 ArraySlice, void const* MipSliceData, uint32 MipSliceSize)
+						{
+							// Stop when there's no more bulk data
+							if (MipSliceData == nullptr)
+								return false;
+
+							if (GLFormat.bCompressed)
+							{
+								glCompressedTexSubImage2D(
+									CurrentTarget,
+									MipIndex,
+									0, 0, // X/Y offset
+									MipSizeX, MipSizeY,
+									GLFormat.Format,
+									MipSliceSize,
+									MipSliceData);
+							}
+							else
+							{
+								glTexSubImage2D(
+									CurrentTarget,
+									MipIndex,
+									0, 0, // X/Y offset
+									MipSizeX, MipSizeY,
+									GLFormat.Format,
+									GLFormat.Type,
+									MipSliceData);
+							}
+
+							return true;
+						});
 					}
-					
-					MipOffset += NumBlocksX * NumBlocksY * NumLayers * GPixelFormats[Format].BlockBytes;
+				}
+				else if (GLFormat.bCompressed && !BulkDataPtr)
+				{
+					// Compressed textures created without using the TexStorage functions cannot be allocated via TexImage without bulk data.
+					// Do nothing here. Texture memory will be allocated when the renderer locks/unlocks the mips for writing.
 				}
 				else
 				{
-					GLenum FirstTarget = bCubeTexture ? GL_TEXTURE_CUBE_MAP_POSITIVE_X : Target;
-					uint32 NumTargets = bCubeTexture ? 6 : 1;
+					// Failed to create immutable storage. Fall back to the standard TexImage functions.
+					// This both allocates the memory and fills in the bulk data simultaneously.
+					EnumerateSubresources(BulkDataPtr, [&](GLenum CurrentTarget, uint32 MipSizeX, uint32 MipSizeY, uint32 MipSizeZ, uint32 MipIndex, uint32 ArraySlice, void const* MipSliceData, uint32 MipSliceSize)
+					{
+						if (GLFormat.bCompressed)
+						{
+							glCompressedTexImage2D(
+								CurrentTarget,
+								MipIndex,
+								GLFormat.InternalFormat[bSRGB],
+								MipSizeX, MipSizeY, 
+								0,
+								MipSliceSize,
+								MipSliceData);
+						}
+						else
+						{
+							glTexImage2D(
+								CurrentTarget,
+								MipIndex,
+								GLFormat.InternalFormat[bSRGB],
+								MipSizeX, MipSizeY,
+								0,
+								GLFormat.Format,
+								GLFormat.Type,
+								MipSliceData);
+						}
 
-					for(uint32 TargetIndex = 0; TargetIndex < NumTargets; TargetIndex++)
+						// Always continue iterating to allocate all mips/slices.
+						return true;
+					});
+				}
+				break;
+
+			case GL_TEXTURE_2D_ARRAY:
+			case GL_TEXTURE_CUBE_MAP_ARRAY:
+			case GL_TEXTURE_3D:
+				{
+					bAllocatedStorage = true; // Always supported if 3D textures are supported.
+
+					const uint32 SizeZ =
+						Target == GL_TEXTURE_3D ? Desc.Depth :
+						Target == GL_TEXTURE_CUBE_MAP_ARRAY ? Desc.ArraySize * 6 :
+						Desc.ArraySize;
+
+					FOpenGL::TexStorage3D(Target, Desc.NumMips, GLFormat.InternalFormat[bSRGB], Desc.Extent.X, Desc.Extent.Y, SizeZ, GLFormat.Format, GLFormat.Type);
+
+					// Texture created with immutable storage. Now fill in the bulk data.
+					if (BulkDataPtr)
 					{
-						glTexSubImage2D(
-							/*Target=*/ FirstTarget + TargetIndex,
-							/*Level=*/ MipIndex,
-							/*XOffset*/ 0,
-							/*YOffset*/ 0,
-							/*SizeX=*/ FMath::Max<uint32>(1,(SizeX >> MipIndex)),
-							/*SizeY=*/ FMath::Max<uint32>(1,(SizeY >> MipIndex)),
-							/*Format=*/ GLFormat.Format,
-							/*Type=*/ GLFormat.Type,
-							/*Data=*/ &Data[MipOffset]
-							);
-						
-						MipOffset += NumBlocksX * NumBlocksY * NumLayers * GPixelFormats[Format].BlockBytes;
+						EnumerateSubresources(BulkDataPtr, [&](GLenum CurrentTarget, uint32 MipSizeX, uint32 MipSizeY, uint32 MipSizeZ, uint32 MipIndex, uint32 ArraySlice, void const* MipSliceData, uint32 MipSliceSize)
+						{
+							// Stop when there's no more bulk data
+							if (MipSliceData == nullptr)
+								return false;
+
+							if (GLFormat.bCompressed)
+							{
+								glCompressedTexSubImage3D(
+									CurrentTarget,
+									MipIndex,
+									0, 0, ArraySlice, // X/Y/Z offset
+									MipSizeX, MipSizeY, MipSizeZ,
+									GLFormat.Format,
+									MipSliceSize,
+									MipSliceData);
+							}
+							else
+							{
+								glTexSubImage3D(
+									CurrentTarget,
+									MipIndex,
+									0, 0, ArraySlice, // X/Y/Z offset
+									MipSizeX, MipSizeY, MipSizeZ,
+									GLFormat.Format,
+									GLFormat.Type,
+									MipSliceData);
+							}
+
+							return true;
+						});
 					}
 				}
+				break;
 
-				glPixelStorei(GL_UNPACK_ALIGNMENT, 4);
-			}
-
-			BulkData->Discard();
-		}
-	}
-	else if (TileMemDepth)
-	{
-#if PLATFORM_ANDROID
-		Target = GL_RENDERBUFFER;
-		glGenRenderbuffers(1, &TextureID);
-		glBindRenderbuffer(GL_RENDERBUFFER, TextureID);
-		glRenderbufferStorageMultisampleEXT(GL_RENDERBUFFER, NumSamplesTileMem, GL_DEPTH24_STENCIL8, SizeX, SizeY);
-		VERIFY_GL(glRenderbufferStorageMultisampleEXT);
-		glBindRenderbuffer(GL_RENDERBUFFER, 0);
-#endif
-	}
-	else
-	{
-		check( BulkData == NULL);
+			case GL_TEXTURE_2D_MULTISAMPLE:
+				{
+					checkf(BulkDataPtr == nullptr, TEXT("Multisample textures cannot be created with initial bulk data."));
 
-		// Try to create an immutable texture and fallback if it fails
-		if (!FOpenGL::TexStorage2DMultisample( Target, NumSamples, GLFormat.InternalFormat[bSRGB], SizeX, SizeY, true))
-		{
-			FOpenGL::TexImage2DMultisample(
-				Target,
-				NumSamples,
-				GLFormat.InternalFormat[bSRGB],
-				SizeX,
-				SizeY,
-				true
-				);
+					// Try to create an immutable storage texture and fallback if it fails
+					const bool FixedSampleLocations = true;
+					if (FOpenGL::TexStorage2DMultisample(Target, Texture->GetNumSamplesStored(), GLFormat.InternalFormat[bSRGB], Desc.Extent.X, Desc.Extent.Y, FixedSampleLocations))
+					{
+						bAllocatedStorage = true;
+					}
+					else
+					{
+						FOpenGL::TexImage2DMultisample(Target, Texture->GetNumSamplesStored(), GLFormat.InternalFormat[bSRGB], Desc.Extent.X, Desc.Extent.Y, FixedSampleLocations);
+					}
+				}
+				break;
+			}
 		}
 	}
-	
-	// Determine the attachment point for the texture.	
-	GLenum Attachment = GL_NONE;
-	if(EnumHasAnyFlags(Flags, TexCreate_RenderTargetable | TexCreate_CPUReadback))
-	{
-		Attachment = GL_COLOR_ATTACHMENT0;
-	}
-	else if(EnumHasAnyFlags(Flags, TexCreate_DepthStencilTargetable))
-	{
-		Attachment = (Format == PF_DepthStencil) ? GL_DEPTH_STENCIL_ATTACHMENT : GL_DEPTH_ATTACHMENT;
-	}
-	else if(EnumHasAnyFlags(Flags, TexCreate_ResolveTargetable))
-	{
-		Attachment = (Format == PF_DepthStencil)
-						? GL_DEPTH_STENCIL_ATTACHMENT
-						: ((Format == PF_ShadowDepth || Format == PF_D24)
-							? GL_DEPTH_ATTACHMENT
-							: GL_COLOR_ATTACHMENT0);
-	}
 
-	switch(Attachment)
-	{
-		case GL_COLOR_ATTACHMENT0:
-			check(GMaxOpenGLColorSamples>=(GLint)NumSamples);
-			break;
-		case GL_DEPTH_ATTACHMENT:
-		case GL_DEPTH_STENCIL_ATTACHMENT:
-			check(GMaxOpenGLDepthSamples>=(GLint)NumSamples);
-			break;
-		default:
-			break;
-	}
 	// @todo: If integer pixel format
 	//check(GMaxOpenGLIntegerSamples>=NumSamples);
+	Texture->SetAllocatedStorage(bAllocatedStorage);
 
-	if (bCubeTexture)
-	{
-		//	FOpenGLTextureCube* TextureCube = new FOpenGLTextureCube(this, TextureID, Target, Attachment, SizeX, SizeY, 0, NumMips, 1, 1, ArraySize, (EPixelFormat)Format, true, bAllocatedStorage, Flags, InClearValue);
-		FOpenGLTextureCube* TextureCube = (FOpenGLTextureCube*)Texture;
-		TextureCube->SetResource(TextureID);
-		TextureCube->Target = Target;
-		TextureCube->Attachment = Attachment;
-		TextureCube->SetAllocatedStorage(bAllocatedStorage);
-	}
-	else
-	{
-		FOpenGLTexture2D* Texture2D = (FOpenGLTexture2D*)Texture;
-		Texture2D->SetResource(TextureID);
-		Texture2D->Target = Target;
-		Texture2D->Attachment = Attachment;
-		Texture2D->SetAllocatedStorage(bAllocatedStorage);
-	}
-
-	OpenGLTextureAllocated(Texture, Flags);
 	// No need to restore texture stage; leave it like this,
 	// and the next draw will take care of cleaning it up; or
 	// next operation that needs the stage will switch something else in on it.
 }
 
-template<typename RHIResourceType>
-void TOpenGLTexture<RHIResourceType>::Resolve(uint32 MipIndex,uint32 ArrayIndex)
+
+void FOpenGLTexture::Resolve(uint32 MipIndex,uint32 ArrayIndex)
 {
 	VERIFY_GL_SCOPE();
-	
-#if UE_BUILD_DEBUG
-	if((FOpenGLTexture2D*)this->GetTexture2D())
-	{
-		check( ((FOpenGLTexture2D*)this->GetTexture2D())->GetNumSamples() == 1 );
-	}
-#endif
+	check(!GetTexture2D() || GetNumSamples() == 1);
 	
 	// Calculate the dimensions of the mip-map.
 	EPixelFormat PixelFormat = this->GetFormat();
@@ -811,14 +894,14 @@ void TOpenGLTexture<RHIResourceType>::Resolve(uint32 MipIndex,uint32 ArrayIndex)
 	const bool bSRGB = EnumHasAnyFlags(this->GetFlags(), TexCreate_SRGB);
 	
 	// Use a texture stage that's not likely to be used for draws, to avoid waiting
-	FOpenGLContextState& ContextState = OpenGLRHI->GetContextStateForCurrentContext();
+	FOpenGLContextState& ContextState = FOpenGLDynamicRHI::Get().GetContextStateForCurrentContext();
 
-	OpenGLRHI->CachedSetupTextureStage(ContextState, FOpenGL::GetMaxCombinedTextureImageUnits() - 1, Target, this->GetResource(), -1, this->GetNumMips());
+	FOpenGLDynamicRHI::Get().CachedSetupTextureStage(ContextState, FOpenGL::GetMaxCombinedTextureImageUnits() - 1, Target, this->GetResource(), -1, this->GetNumMips());
 	
 	glBindBuffer( GL_PIXEL_PACK_BUFFER, PixelBuffer->Resource );
 
 	{
-		if( this->GetSizeZ() )
+		if (GetDesc().IsTextureArray() || GetDesc().IsTexture3D())
 		{
 			// apparently it's not possible to retrieve compressed image from GL_TEXTURE_2D_ARRAY in OpenGL for compressed images
 			// and for uncompressed ones it's not possible to specify the image index
@@ -836,8 +919,8 @@ void TOpenGLTexture<RHIResourceType>::Resolve(uint32 MipIndex,uint32 ArrayIndex)
 			else
 			{
 				// Get framebuffer for texture
-				FOpenGLTextureBase* Texture = this;
-				GLuint SourceFramebuffer = OpenGLRHI->GetOpenGLFramebuffer(1, &Texture, (bCubemap ? &ArrayIndex : nullptr), &MipIndex, nullptr);
+				FOpenGLTexture* Texture = this;
+				GLuint SourceFramebuffer = FOpenGLDynamicRHI::Get().GetOpenGLFramebuffer(1, &Texture, (bCubemap ? &ArrayIndex : nullptr), &MipIndex, nullptr);
 				// Bind the framebuffer
 				glBindFramebuffer(UGL_READ_FRAMEBUFFER, SourceFramebuffer);
 				FOpenGL::ReadBuffer(GL_COLOR_ATTACHMENT0);
@@ -846,7 +929,7 @@ void TOpenGLTexture<RHIResourceType>::Resolve(uint32 MipIndex,uint32 ArrayIndex)
 				glReadPixels(0, 0, MipSizeX, MipSizeY, GLFormat.Format, GLFormat.Type, 0);
 				glPixelStorei(GL_PACK_ALIGNMENT, 4);
 
-				OpenGLRHI->GetContextStateForCurrentContext().Framebuffer = (GLuint)-1;
+				FOpenGLDynamicRHI::Get().GetContextStateForCurrentContext().Framebuffer = (GLuint)-1;
 			}
 		}
 	}
@@ -858,8 +941,7 @@ void TOpenGLTexture<RHIResourceType>::Resolve(uint32 MipIndex,uint32 ArrayIndex)
 	// next operation that needs the stage will switch something else in on it.
 }
 
-template<typename RHIResourceType>
-uint32 TOpenGLTexture<RHIResourceType>::GetLockSize(uint32 InMipIndex, uint32 ArrayIndex, EResourceLockMode LockMode, uint32& DestStride)
+uint32 FOpenGLTexture::GetLockSize(uint32 InMipIndex, uint32 ArrayIndex, EResourceLockMode LockMode, uint32& DestStride)
 {
 	// Calculate the dimensions of the mip-map.
 	EPixelFormat PixelFormat = this->GetFormat();
@@ -875,8 +957,7 @@ uint32 TOpenGLTexture<RHIResourceType>::GetLockSize(uint32 InMipIndex, uint32 Ar
 	return MipBytes;
 }
 
-template<typename RHIResourceType>
-void TOpenGLTexture<RHIResourceType>::Fill2DGLTextureImage(const FOpenGLTextureFormat& GLFormat, const bool bSRGB, uint32 MipIndex, const void* BufferOrPBOOffset, uint32 ImageSize, uint32 ArrayIndex)
+void FOpenGLTexture::Fill2DGLTextureImage(const FOpenGLTextureFormat& GLFormat, const bool bSRGB, uint32 MipIndex, const void* BufferOrPBOOffset, uint32 ImageSize, uint32 ArrayIndex)
 {
 	if (GLFormat.bCompressed)
 	{
@@ -927,17 +1008,10 @@ void TOpenGLTexture<RHIResourceType>::Fill2DGLTextureImage(const FOpenGLTextureF
 	}
 }
 
-template<typename RHIResourceType>
-void* TOpenGLTexture<RHIResourceType>::Lock(uint32 InMipIndex,uint32 ArrayIndex,EResourceLockMode LockMode,uint32& DestStride)
+void* FOpenGLTexture::Lock(uint32 InMipIndex,uint32 ArrayIndex,EResourceLockMode LockMode,uint32& DestStride)
 {
 	VERIFY_GL_SCOPE();
-
-#if UE_BUILD_DEBUG
-	if((FOpenGLTexture2D*)this->GetTexture2D())
-	{
-		check( ((FOpenGLTexture2D*)this->GetTexture2D())->GetNumSamples() == 1 );
-	}
-#endif
+	check(!GetTexture2D() || GetNumSamples() == 1);
 
 	SCOPE_CYCLE_COUNTER(STAT_OpenGLLockTextureTime);
 	
@@ -1038,8 +1112,7 @@ inline float HalfFloatToFloat(uint16 HalfFloat)
 	return Convert.F;
 }
 
-template<typename RHIResourceType>
-void TOpenGLTexture<RHIResourceType>::Unlock(uint32 MipIndex,uint32 ArrayIndex)
+void FOpenGLTexture::Unlock(uint32 MipIndex,uint32 ArrayIndex)
 {
 	VERIFY_GL_SCOPE();
 	SCOPE_CYCLE_COUNTER(STAT_OpenGLUnlockTextureTime);
@@ -1071,8 +1144,8 @@ void TOpenGLTexture<RHIResourceType>::Unlock(uint32 MipIndex,uint32 ArrayIndex)
 		check(this->GetSizeZ() == 0);
 
 		// Use a texture stage that's not likely to be used for draws, to avoid waiting
-		FOpenGLContextState& ContextState = OpenGLRHI->GetContextStateForCurrentContext();
-		OpenGLRHI->CachedSetupTextureStage(ContextState, FOpenGL::GetMaxCombinedTextureImageUnits() - 1, Target, GetResource(), -1, this->GetNumMips());
+		FOpenGLContextState& ContextState = FOpenGLDynamicRHI::Get().GetContextStateForCurrentContext();
+		FOpenGLDynamicRHI::Get().CachedSetupTextureStage(ContextState, FOpenGL::GetMaxCombinedTextureImageUnits() - 1, Target, GetResource(), -1, this->GetNumMips());
 
 		CachedBindPixelUnpackBuffer(GL_PIXEL_UNPACK_BUFFER, 0);
 
@@ -1132,12 +1205,11 @@ void TOpenGLTexture<RHIResourceType>::Unlock(uint32 MipIndex,uint32 ArrayIndex)
 	if (!PixelBuffer->IsLockReadOnly())
 	{
 		// Use a texture stage that's not likely to be used for draws, to avoid waiting
-		FOpenGLContextState& ContextState = OpenGLRHI->GetContextStateForCurrentContext();
-		OpenGLRHI->CachedSetupTextureStage(ContextState, FOpenGL::GetMaxCombinedTextureImageUnits() - 1, Target, GetResource(), -1, this->GetNumMips());
+		FOpenGLContextState& ContextState = FOpenGLDynamicRHI::Get().GetContextStateForCurrentContext();
+		FOpenGLDynamicRHI::Get().CachedSetupTextureStage(ContextState, FOpenGL::GetMaxCombinedTextureImageUnits() - 1, Target, GetResource(), -1, this->GetNumMips());
 
-		if (this->GetSizeZ())
+		if (GetDesc().IsTextureArray() || GetDesc().IsTexture3D())
 		{
-			// texture 2D array
 			if (GLFormat.bCompressed)
 			{
 				FOpenGL::CompressedTexSubImage3D(
@@ -1201,8 +1273,7 @@ float GMaxRestoreTime = 0.0f;
 float GAvgRestoreTime = 0.0f;
 uint32 GAvgRestoreCount = 0;
 
-template<typename RHIResourceType>
-void TOpenGLTexture<RHIResourceType>::RestoreEvictedGLResource(bool bAttemptToRetainMips)
+void FOpenGLTexture::RestoreEvictedGLResource(bool bAttemptToRetainMips)
 {
 //	double StartTime = FPlatformTime::Seconds();
 
@@ -1212,7 +1283,7 @@ void TOpenGLTexture<RHIResourceType>::RestoreEvictedGLResource(bool bAttemptToRe
 	EvictionParamsPtr->bHasRestored = true;
 
 	const FClearValueBinding ClearBinding = this->GetClearBinding();
-	OpenGLRHI->InitializeGLTextureInternal(GetRawResourceName(), this, this->GetSizeX(), this->GetSizeY(), bCubemap, false, false, this->GetFormat(), this->GetNumMips(), this->GetNumSamples(), 0, this->GetFlags(), ClearBinding, nullptr);
+	FOpenGLDynamicRHI::Get().InitializeGLTextureInternal(this, nullptr, 0);
 
 	EPixelFormat PixelFormat = this->GetFormat();
 	const FOpenGLTextureFormat& GLFormat = GOpenGLTextureFormats[PixelFormat];
@@ -1233,7 +1304,9 @@ void TOpenGLTexture<RHIResourceType>::RestoreEvictedGLResource(bool bAttemptToRe
 	// Use the resident streaming mips if our cvar is -1.
 	uint32 DeferTextureCreationKeepLowerMipCount = (uint32)(GOGLDeferTextureCreationKeepLowerMipCount >= 0 ? GOGLDeferTextureCreationKeepLowerMipCount : UTexture::GetStaticMinTextureResidentMipCount());
 
-	uint32 RetainMips = bAttemptToRetainMips && EnumHasAnyFlags(this->GetFlags(), TexCreate_Streamable) && this->GetNumMips() > 1 && !this->IsAliased() ? DeferTextureCreationKeepLowerMipCount : 0;
+	uint32 RetainMips = bAttemptToRetainMips && EnumHasAnyFlags(this->GetFlags(), TexCreate_Streamable) && this->GetNumMips() > 1 && !bAlias
+		? DeferTextureCreationKeepLowerMipCount
+		: 0;
 
 	if (CanBeEvicted())
 	{
@@ -1265,13 +1338,11 @@ void TOpenGLTexture<RHIResourceType>::RestoreEvictedGLResource(bool bAttemptToRe
 // 	GAvgRestoreTime += ThisTime;
 }
 
-template<typename RHIResourceType>
-void TOpenGLTexture<RHIResourceType>::TryEvictGLResource()
+void FOpenGLTexture::TryEvictGLResource()
 {
 	VERIFY_GL_SCOPE();
-	if (CanCreateAsEvicted() && EvictionParamsPtr->bHasRestored)
+	if (bCanCreateAsEvicted && EvictionParamsPtr->bHasRestored)
 	{
-
 		if (CanBeEvicted())
 		{
 			QUICK_SCOPE_CYCLE_COUNTER(STAT_OpenGLTryEvictGLResource);
@@ -1281,13 +1352,13 @@ void TOpenGLTexture<RHIResourceType>::TryEvictGLResource()
 			EvictionParamsPtr->bHasRestored = false;
 			const FClearValueBinding ClearBinding = this->GetClearBinding();
 			// recreate the GL tex resource name (but not allocate the memory)
-			OpenGLRHI->InitializeGLTexture(this, this->GetSizeX(), this->GetSizeY(), bCubemap, false, false, this->GetFormat(), this->GetNumMips(), this->GetNumSamples(), 0, this->GetFlags(), ClearBinding, nullptr);
+			FOpenGLDynamicRHI::Get().InitializeGLTexture(this, nullptr, 0);
 			GTotalEvictedMipMemDuplicated -= EvictionParamsPtr->GetTotalAllocated();
 		}
 	}
 }
 
-static bool CanDeferTextureCreation()
+bool FOpenGLTextureDesc::CanDeferTextureCreation()
 {
 	bool bCanDeferTextureCreation = CVarDeferTextureCreation.GetValueOnAnyThread() != 0;
 #if PLATFORM_ANDROID
@@ -1316,51 +1387,28 @@ static bool CanDeferTextureCreation()
 	return bCanDeferTextureCreation;
 }
 
-template<typename RHIResourceType>
-bool TOpenGLTexture<RHIResourceType>::CanCreateAsEvicted()
-{
-	// can run on RT.
-	bool bRet =
-		CanDeferTextureCreation()
-		&& FOpenGL::SupportsCopyImage()
-		&& this->GetFlags() != TexCreate_None // ignore TexCreate_None
-		&& !EnumHasAnyFlags((ETextureCreateFlags)CVarDeferTextureCreationExcludeMask.GetValueOnAnyThread(), this->GetFlags())  // Anything outside of these flags cannot be evicted.
-		&& Target == GL_TEXTURE_2D
-		&& this->GetTexture2D(); // 2d only.
-
-	if (GOGLTextureEvictLogging)
-	{
-		UE_CLOG(!bRet, LogRHI, Warning, TEXT("CanDeferTextureCreation:%d, SupportsCopyImage:%d, Flags:%llx Mask:%x, Target:%x"),
-			CanDeferTextureCreation(), FOpenGL::SupportsCopyImage(), this->GetFlags(), CVarDeferTextureCreationExcludeMask.GetValueOnAnyThread(), Target);
-	}
-
-	return bRet;
-}
-
-template<typename RHIResourceType>
-bool TOpenGLTexture<RHIResourceType>::CanBeEvicted()
+bool FOpenGLTexture::CanBeEvicted()
 {
 	VERIFY_GL_SCOPE();
-	checkf(!CanCreateAsEvicted() || EvictionParamsPtr.IsValid(), TEXT("%p, CanCreateAsEvicted() %d, EvictionParamsPtr.IsValid() %d"), this, CanCreateAsEvicted(),EvictionParamsPtr.IsValid());
+	checkf(!bCanCreateAsEvicted || EvictionParamsPtr.IsValid(), TEXT("%p, bCanCreateAsEvicted %d, EvictionParamsPtr.IsValid() %d"), this, bCanCreateAsEvicted, EvictionParamsPtr.IsValid());
 
 	// if we're aliased check that there's no eviction data.
-	check(!CanCreateAsEvicted() || !this->IsAliased() || (EvictionParamsPtr->MipImageData.Num() == 0 && EvictionParamsPtr->MipImageData.Num() != this->GetNumMips()));
+	check(!bCanCreateAsEvicted || !bAlias || (EvictionParamsPtr->MipImageData.Num() == 0 && EvictionParamsPtr->MipImageData.Num() != this->GetNumMips()));
 
 	// cant evict if we're aliased, or there are mips are not backed by stored data.
-	bool bRet = CanCreateAsEvicted() && EvictionParamsPtr->MipImageData.Num() == this->GetNumMips() && EvictionParamsPtr->AreAllMipsPresent();
+	bool bRet = bCanCreateAsEvicted && EvictionParamsPtr->MipImageData.Num() == this->GetNumMips() && EvictionParamsPtr->AreAllMipsPresent();
 
 	return bRet;
 }
 
-template<typename RHIResourceType>
-void TOpenGLTexture<RHIResourceType>::CloneViaCopyImage( TOpenGLTexture<RHIResourceType>* Src, uint32 InNumMips, int32 SrcOffset, int32 DstOffset)
+void FOpenGLTexture::CloneViaCopyImage(FOpenGLTexture* Src, uint32 InNumMips, int32 SrcOffset, int32 DstOffset)
 {
 	VERIFY_GL_SCOPE();
 
 	check(FOpenGL::SupportsCopyImage());
 	
-	check(Src->CanCreateAsEvicted() == CanCreateAsEvicted());
-	if (CanCreateAsEvicted())
+	check(Src->bCanCreateAsEvicted == bCanCreateAsEvicted);
+	if (bCanCreateAsEvicted)
 	{
 		// Copy all mips that are present.
 		if (!(!Src->IsEvicted() || Src->EvictionParamsPtr->AreAllMipsPresent()))
@@ -1403,8 +1451,7 @@ void TOpenGLTexture<RHIResourceType>::CloneViaCopyImage( TOpenGLTexture<RHIResou
 	
 }
 
-template<typename RHIResourceType>
-void TOpenGLTexture<RHIResourceType>::CloneViaPBO( TOpenGLTexture<RHIResourceType>* Src, uint32 InNumMips, int32 SrcOffset, int32 DstOffset)
+void FOpenGLTexture::CloneViaPBO(FOpenGLTexture* Src, uint32 InNumMips, int32 SrcOffset, int32 DstOffset)
 {
 	VERIFY_GL_SCOPE();
 	
@@ -1423,7 +1470,7 @@ void TOpenGLTexture<RHIResourceType>::CloneViaPBO( TOpenGLTexture<RHIResourceTyp
 	const uint32 BlockSizeY = GPixelFormats[PixelFormat].BlockSizeY;
 	const uint32 BlockBytes = GPixelFormats[PixelFormat].BlockBytes;
 	
-	FOpenGLContextState& ContextState = OpenGLRHI->GetContextStateForCurrentContext();
+	FOpenGLContextState& ContextState = FOpenGLDynamicRHI::Get().GetContextStateForCurrentContext();
 	
 	for (uint32 ArrayIndex = 0; ArrayIndex < this->GetEffectiveSizeZ(); ArrayIndex++)
 	{
@@ -1461,7 +1508,7 @@ void TOpenGLTexture<RHIResourceType>::CloneViaPBO( TOpenGLTexture<RHIResourceTyp
 			// Transfer data from texture to pixel buffer.
 			// This may be further optimized by caching information if surface content was changed since last lock.
 			// Use a texture stage that's not likely to be used for draws, to avoid waiting
-			OpenGLRHI->CachedSetupTextureStage(ContextState, FOpenGL::GetMaxCombinedTextureImageUnits() - 1, Src->Target, Src->GetResource(), -1, this->GetNumMips());
+			FOpenGLDynamicRHI::Get().CachedSetupTextureStage(ContextState, FOpenGL::GetMaxCombinedTextureImageUnits() - 1, Src->Target, Src->GetResource(), -1, this->GetNumMips());
 			
 			glBindBuffer( GL_PIXEL_PACK_BUFFER, PixelBuffer->Resource );
 			
@@ -1475,633 +1522,216 @@ void TOpenGLTexture<RHIResourceType>::CloneViaPBO( TOpenGLTexture<RHIResourceTyp
 			{
 				glPixelStorei(GL_PACK_ALIGNMENT, 1);
 				FOpenGL::GetTexImage(Src->bCubemap ? GL_TEXTURE_CUBE_MAP_POSITIVE_X + ArrayIndex : Src->Target,
-									 SrcMipIndex,
-									 GLFormat.Format,
-									 GLFormat.Type,
-									 0);	// offset into PBO
-				glPixelStorei(GL_PACK_ALIGNMENT, 4);
-			}
-			
-			// copy the texture data
-			// Upload directly into Dst to avoid out-of-band synchronization caused by glMapBuffer!
-			{
-				CachedBindPixelUnpackBuffer( GL_PIXEL_UNPACK_BUFFER, PixelBuffer->Resource );
-				
-				// Use a texture stage that's not likely to be used for draws, to avoid waiting
-				OpenGLRHI->CachedSetupTextureStage(ContextState, FOpenGL::GetMaxCombinedTextureImageUnits() - 1, Target, GetResource(), -1, this->GetNumMips());
-				
-				if( this->GetSizeZ() )
-				{
-					// texture 2D array
-					if (GLFormat.bCompressed)
-					{
-						FOpenGL::CompressedTexSubImage3D(Target,
-														 DstMipIndex,
-														 0,
-														 0,
-														 ArrayIndex,
-														 MipSizeX,
-														 MipSizeY,
-														 1,
-														 GLFormat.InternalFormat[bSRGB],
-														 PixelBuffer->GetSize(),
-														 0);
-					}
-					else
-					{
-						glPixelStorei(GL_UNPACK_ALIGNMENT, 1);
-						check( FOpenGL::SupportsTexture3D() );
-						FOpenGL::TexSubImage3D(Target,
-											   DstMipIndex,
-											   0,
-											   0,
-											   ArrayIndex,
-											   MipSizeX,
-											   MipSizeY,
-											   1,
-											   GLFormat.Format,
-											   GLFormat.Type,
-											   0);	// offset into PBO
-						glPixelStorei(GL_UNPACK_ALIGNMENT, 4);
-					}
-				}
-				else
-				{
-					if (GLFormat.bCompressed)
-					{
-						if (GetAllocatedStorageForMip(DstMipIndex,ArrayIndex))
-						{
-							glCompressedTexSubImage2D(bCubemap ? GL_TEXTURE_CUBE_MAP_POSITIVE_X + ArrayIndex : Target,
-													  DstMipIndex,
-													  0,
-													  0,
-													  MipSizeX,
-													  MipSizeY,
-													  GLFormat.InternalFormat[bSRGB],
-													  PixelBuffer->GetSize(),
-													  0);	// offset into PBO
-						}
-						else
-						{
-							glCompressedTexImage2D(bCubemap ? GL_TEXTURE_CUBE_MAP_POSITIVE_X + ArrayIndex : Target,
-												   DstMipIndex,
-												   GLFormat.InternalFormat[bSRGB],
-												   MipSizeX,
-												   MipSizeY,
-												   0,
-												   PixelBuffer->GetSize(),
-												   0);	// offset into PBO
-							SetAllocatedStorageForMip(DstMipIndex,ArrayIndex);
-						}
-					}
-					else
-					{
-						// All construction paths should have called TexStorage2D or TexImage2D. So we will
-						// always call TexSubImage2D.
-						check(GetAllocatedStorageForMip(DstMipIndex,ArrayIndex) == true);
-						glPixelStorei(GL_UNPACK_ALIGNMENT, 1);
-						glTexSubImage2D(bCubemap ? GL_TEXTURE_CUBE_MAP_POSITIVE_X + ArrayIndex : Target,
-										DstMipIndex,
-										0,
-										0,
-										MipSizeX,
-										MipSizeY,
-										GLFormat.Format,
-										GLFormat.Type,
-										0);	// offset into PBO
-						glPixelStorei(GL_UNPACK_ALIGNMENT, 4);
-					}
-				}
-			}
-			
-			// need to free PBO if we aren't keeping shadow copies
-			PixelBuffers[BufferIndex] = NULL;
-			
-			// No need to restore texture stage; leave it like this,
-			// and the next draw will take care of cleaning it up; or
-			// next operation that needs the stage will switch something else in on it.
-		}
-	}
-	
-	// Reset the buffer bindings on exit only
-	glBindBuffer( GL_PIXEL_PACK_BUFFER, 0 );
-	CachedBindPixelUnpackBuffer(GL_PIXEL_UNPACK_BUFFER, 0);
-}
-
-/*-----------------------------------------------------------------------------
-	2D texture support.
------------------------------------------------------------------------------*/
-
-/**
-* Creates a 2D RHI texture resource
-* @param SizeX - width of the texture to create
-* @param SizeY - height of the texture to create
-* @param Format - EPixelFormat texture format
-* @param NumMips - number of mips to generate or 0 for full mip pyramid
-* @param Flags - ETextureCreateFlags creation flags
-*/
-FTexture2DRHIRef FOpenGLDynamicRHI::RHICreateTexture2D(uint32 SizeX,uint32 SizeY,uint8 Format,uint32 NumMips,uint32 NumSamples,ETextureCreateFlags Flags, ERHIAccess InResourceState,FRHIResourceCreateInfo& Info)
-{
-	return (FRHITexture2D*)CreateOpenGLTexture(SizeX, SizeY, false, false, false, Format, NumMips, NumSamples, 1, Flags, Info.ClearValueBinding, Info.BulkData);
-}
-
-/**
-* Creates a 2D RHI texture external resource
-* @param SizeX - width of the texture to create
-* @param SizeY - height of the texture to create
-* @param Format - EPixelFormat texture format
-* @param NumMips - number of mips to generate or 0 for full mip pyramid
-* @param Flags - ETextureCreateFlags creation flags
-*/
-FTexture2DRHIRef FOpenGLDynamicRHI::RHICreateTextureExternal2D(uint32 SizeX, uint32 SizeY, uint8 Format, uint32 NumMips, uint32 NumSamples, ETextureCreateFlags Flags, ERHIAccess InResourceState, FRHIResourceCreateInfo& Info)
-{
-	return (FRHITexture2D*)CreateOpenGLTexture(SizeX, SizeY, false, false, true, Format, NumMips, NumSamples, 1, Flags, Info.ClearValueBinding, Info.BulkData);
-}
-
-FTexture2DRHIRef FOpenGLDynamicRHI::RHIAsyncCreateTexture2D(uint32 SizeX, uint32 SizeY, uint8 Format, uint32 NumMips, ETextureCreateFlags Flags, ERHIAccess InResourceState, void** InitialMipData, uint32 NumInitialMips)
-{
-	check(0);
-	return FTexture2DRHIRef();
-}
-
-void FOpenGLDynamicRHI::RHICopySharedMips(FRHITexture2D* DestTexture2D, FRHITexture2D* SrcTexture2D)
-{
-	check(0);
-}
-
-FTexture2DArrayRHIRef FOpenGLDynamicRHI::RHICreateTexture2DArray(uint32 SizeX,uint32 SizeY,uint32 SizeZ,uint8 Format,uint32 NumMips,uint32 NumSamples,ETextureCreateFlags Flags, ERHIAccess InResourceState, FRHIResourceCreateInfo& Info)
-{
-	VERIFY_GL_SCOPE();
-
-	SCOPE_CYCLE_COUNTER(STAT_OpenGLCreateTextureTime);
-
-	check( FOpenGL::SupportsTexture3D() );
-
-	if(NumMips == 0)
-	{
-		NumMips = FindMaxMipmapLevel(SizeX, SizeY);
-	}
-
-	GLuint TextureID = 0;
-	FOpenGL::GenTextures(1, &TextureID);
-
-	const GLenum Target = GL_TEXTURE_2D_ARRAY;
-
-	// Use a texture stage that's not likely to be used for draws, to avoid waiting
-	FOpenGLContextState& ContextState = GetContextStateForCurrentContext();
-	CachedSetupTextureStage(ContextState, FOpenGL::GetMaxCombinedTextureImageUnits() - 1, Target, TextureID, 0, NumMips);
-
-	glTexParameteri(Target, GL_TEXTURE_WRAP_S, GL_REPEAT);
-	glTexParameteri(Target, GL_TEXTURE_WRAP_T, GL_REPEAT);
-	glTexParameteri(Target, GL_TEXTURE_MAG_FILTER, GL_NEAREST);
-	glTexParameteri(Target, GL_TEXTURE_MIN_FILTER, NumMips > 1 ? GL_NEAREST_MIPMAP_NEAREST : GL_NEAREST);
-	if( FOpenGL::SupportsTextureFilterAnisotropic() )
-	{
-		glTexParameteri(Target, GL_TEXTURE_MAX_ANISOTROPY_EXT, 1);
-	}
-	glTexParameteri(Target, GL_TEXTURE_BASE_LEVEL, 0);
-	glTexParameteri(Target, GL_TEXTURE_MAX_LEVEL, NumMips - 1);
-	
-	TextureMipLimits.Add(TextureID, TPair<GLenum, GLenum>(0, NumMips - 1));
-
-	const bool bSRGB = EnumHasAnyFlags(Flags, TexCreate_SRGB);
-	const FOpenGLTextureFormat& GLFormat = GOpenGLTextureFormats[Format];
-	const FPixelFormatInfo& FormatInfo = GPixelFormats[Format];
-	if (GLFormat.InternalFormat[bSRGB] == GL_NONE)
-	{
-		UE_LOG(LogRHI, Fatal,TEXT("Texture format '%s' not supported."), FormatInfo.Name);
-	}
-
-	if (GLFormat.bBGRA && !EnumHasAnyFlags(Flags, TexCreate_RenderTargetable))
-	{
-		glTexParameteri(Target, GL_TEXTURE_SWIZZLE_R, GL_BLUE);
-		glTexParameteri(Target, GL_TEXTURE_SWIZZLE_B, GL_RED);
-	}
-
-	// Make sure PBO is disabled
-	CachedBindPixelUnpackBuffer(ContextState, 0);
-
-	uint8* Data = Info.BulkData ? (uint8*)Info.BulkData->GetResourceBulkData() : NULL;
-	uint32 DataSize = Info.BulkData ? Info.BulkData->GetResourceBulkDataSize() : 0;
-	uint32 MipOffset = 0;
-
-	FOpenGL::TexStorage3D( Target, NumMips, GLFormat.SizedInternalFormat[bSRGB], SizeX, SizeY, SizeZ, GLFormat.Format, GLFormat.Type );
-
-	if (Data)
-	{
-		for(uint32 MipIndex = 0; MipIndex < NumMips; MipIndex++)
-		{
-			const int32 MipSizeX = FMath::Max<int32>(1, (SizeX >> MipIndex));
-			const int32 MipSizeY = FMath::Max<int32>(1, (SizeY >> MipIndex));
-
-			const uint32 MipLinePitch = FMath::DivideAndRoundUp(MipSizeX, FormatInfo.BlockSizeX) * FormatInfo.BlockBytes;
-			const uint32 MipSlicePitch = FMath::DivideAndRoundUp(MipSizeY, FormatInfo.BlockSizeY) * MipLinePitch;
-			const uint32 MipSize = MipSlicePitch * SizeZ;
-
-			if (MipOffset + MipSize > DataSize)
-			{
-				break; // Stop if the texture does not contain the mips.
-			}
-						
-			if (GLFormat.bCompressed)
-			{
-				FOpenGL::CompressedTexSubImage3D(
-					Target,
-					MipIndex,
-					0, 0, 0,
-					MipSizeX, MipSizeY, SizeZ,
-					GLFormat.InternalFormat[bSRGB],
-					MipSize,
-					&Data[MipOffset]);
-			}
-			else
-			{
-				glPixelStorei(GL_UNPACK_ALIGNMENT, 1);
-				FOpenGL::TexSubImage3D(
-					Target,
-					MipIndex,
-					0, 0, 0,
-					MipSizeX, MipSizeY, SizeZ,
-					GLFormat.Format,
-					GLFormat.Type,
-					&Data[MipOffset]
-					);
-				glPixelStorei(GL_UNPACK_ALIGNMENT, 4);
-			}
-
-			MipOffset+= MipSize;
-		}
-
-		Info.BulkData->Discard();
-	}
-	
-	// Determine the attachment point for the texture.	
-	GLenum Attachment = GL_NONE;
-	if(EnumHasAnyFlags(Flags, TexCreate_RenderTargetable))
-	{
-		Attachment = GL_COLOR_ATTACHMENT0;
-	}
-	else if(EnumHasAnyFlags(Flags, TexCreate_DepthStencilTargetable))
-	{
-		Attachment = (Format == PF_DepthStencil) ? GL_DEPTH_STENCIL_ATTACHMENT : GL_DEPTH_ATTACHMENT;
-	}
-	else if(EnumHasAnyFlags(Flags, TexCreate_ResolveTargetable))
-	{
-		Attachment = (Format == PF_DepthStencil)
-			? GL_DEPTH_STENCIL_ATTACHMENT
-			: ((Format == PF_ShadowDepth || Format == PF_D24)
-			? GL_DEPTH_ATTACHMENT
-			: GL_COLOR_ATTACHMENT0);
-	}
-
-	FOpenGLTexture2DArray* Texture = new FOpenGLTexture2DArray(this,TextureID,Target,Attachment,SizeX,SizeY,SizeZ,NumMips,1, 1, SizeZ, (EPixelFormat)Format,false,true,Flags,Info.ClearValueBinding);
-	OpenGLTextureAllocated( Texture, Flags );
-
-	// No need to restore texture stage; leave it like this,
-	// and the next draw will take care of cleaning it up; or
-	// next operation that needs the stage will switch something else in on it.
-
-	return Texture;
-}
-
-FTexture3DRHIRef FOpenGLDynamicRHI::RHICreateTexture3D(uint32 SizeX,uint32 SizeY,uint32 SizeZ,uint8 Format,uint32 NumMips,ETextureCreateFlags Flags, ERHIAccess InResourceState,FRHIResourceCreateInfo& CreateInfo)
-{
-	VERIFY_GL_SCOPE();
-
-	SCOPE_CYCLE_COUNTER(STAT_OpenGLCreateTextureTime);
-
-	check( FOpenGL::SupportsTexture3D() );
-
-	if(NumMips == 0)
-	{
-		NumMips = FindMaxMipmapLevel(SizeX, SizeY, SizeZ);
-	}
-
-	GLuint TextureID = 0;
-	FOpenGL::GenTextures(1, &TextureID);
-
-	const GLenum Target = GL_TEXTURE_3D;
-
-	// Use a texture stage that's not likely to be used for draws, to avoid waiting
-	FOpenGLContextState& ContextState = GetContextStateForCurrentContext();
-	CachedSetupTextureStage(ContextState, FOpenGL::GetMaxCombinedTextureImageUnits() - 1, Target, TextureID, 0, NumMips);
-
-	glTexParameteri(Target, GL_TEXTURE_WRAP_S, GL_REPEAT);
-	glTexParameteri(Target, GL_TEXTURE_WRAP_T, GL_REPEAT);
-	glTexParameteri(Target, GL_TEXTURE_WRAP_R, GL_REPEAT);
-	glTexParameteri(Target, GL_TEXTURE_MAG_FILTER, GL_NEAREST);
-	glTexParameteri(Target, GL_TEXTURE_MIN_FILTER, GL_NEAREST_MIPMAP_NEAREST);
-	if( FOpenGL::SupportsTextureFilterAnisotropic() )
-	{
-		glTexParameteri(Target, GL_TEXTURE_MAX_ANISOTROPY_EXT, 1);
-	}
-	glTexParameteri(Target, GL_TEXTURE_BASE_LEVEL, 0);
-	glTexParameteri(Target, GL_TEXTURE_MAX_LEVEL, NumMips - 1);
-	
-	TextureMipLimits.Add(TextureID, TPair<GLenum, GLenum>(0, NumMips - 1));
-
-	const bool bSRGB = EnumHasAnyFlags(Flags, TexCreate_SRGB);
-	const FOpenGLTextureFormat& GLFormat = GOpenGLTextureFormats[Format];
-	const FPixelFormatInfo& FormatInfo = GPixelFormats[Format];
-	if (GLFormat.InternalFormat[bSRGB] == GL_NONE)
-	{
-		UE_LOG(LogRHI, Fatal, TEXT("Texture format '%s' not supported."), FormatInfo.Name);
-	}
-
-	if (GLFormat.bBGRA && !EnumHasAnyFlags(Flags, TexCreate_RenderTargetable))
-	{
-		glTexParameteri(Target, GL_TEXTURE_SWIZZLE_R, GL_BLUE);
-		glTexParameteri(Target, GL_TEXTURE_SWIZZLE_B, GL_RED);
-	}
-
-
-	// Make sure PBO is disabled
-	CachedBindPixelUnpackBuffer(ContextState,0);
-
-	uint8* Data = CreateInfo.BulkData ? (uint8*)CreateInfo.BulkData->GetResourceBulkData() : nullptr;
-	uint32 DataSize = CreateInfo.BulkData ? CreateInfo.BulkData->GetResourceBulkDataSize() : 0;
-	uint32 MipOffset = 0;
-
-	FOpenGL::TexStorage3D( Target, NumMips, GLFormat.SizedInternalFormat[bSRGB], SizeX, SizeY, SizeZ, GLFormat.Format, GLFormat.Type );
-
-	if (Data)
-	{
-		for (uint32 MipIndex = 0; MipIndex < NumMips; MipIndex++)
-		{
-			const int32 MipSizeX = FMath::Max<int32>(1, (SizeX >> MipIndex));
-			const int32 MipSizeY = FMath::Max<int32>(1, (SizeY >> MipIndex));
-			const int32 MipSizeZ = FMath::Max<int32>(1, (SizeZ >> MipIndex));
-
-			const uint32 MipLinePitch = FMath::DivideAndRoundUp(MipSizeX, FormatInfo.BlockSizeX) * FormatInfo.BlockBytes;
-			const uint32 MipSlicePitch = FMath::DivideAndRoundUp(MipSizeY, FormatInfo.BlockSizeY) * MipLinePitch;
-			const uint32 MipSize = MipSlicePitch * MipSizeZ;
-
-			if (MipOffset + MipSize > DataSize)
-			{
-				break; // Stop if the texture does not contain the mips.
-			}
-
-			if (GLFormat.bCompressed)
-			{
-				int32 RowLength = FMath::DivideAndRoundUp(MipSizeX, FormatInfo.BlockSizeX) * FormatInfo.BlockSizeX;
-				int32 ImageHeight = FMath::DivideAndRoundUp(MipSizeY, FormatInfo.BlockSizeY) * FormatInfo.BlockSizeY;
-
-				FOpenGL::CompressedTexSubImage3D(
-					Target,
-					MipIndex,
-					0, 0, 0,
-					MipSizeX, MipSizeY, MipSizeZ,
-					GLFormat.InternalFormat[bSRGB],
-					MipSize,
-					Data + MipOffset);
-			}
-			else
-			{
-				glPixelStorei(GL_UNPACK_ALIGNMENT, 1);
-				FOpenGL::TexSubImage3D(
-					/*Target=*/ Target,
-					/*Level=*/ MipIndex,
-					0, 0, 0, 
-					MipSizeX, MipSizeY, MipSizeZ,
-					/*Format=*/ GLFormat.Format,
-					/*Type=*/ GLFormat.Type,
-					/*Data=*/ Data + MipOffset);
-				glPixelStorei(GL_UNPACK_ALIGNMENT, 4);
-			}
-
-			MipOffset += MipSize;
-		}
-
-		CreateInfo.BulkData->Discard();
-	}
-	
-	// Determine the attachment point for the texture.	
-	GLenum Attachment = GL_NONE;
-	if(EnumHasAnyFlags(Flags, TexCreate_RenderTargetable))
-	{
-		Attachment = GL_COLOR_ATTACHMENT0;
-	}
-	else if(EnumHasAnyFlags(Flags, TexCreate_DepthStencilTargetable))
-	{
-		Attachment = Format == PF_DepthStencil ? GL_DEPTH_STENCIL_ATTACHMENT : GL_DEPTH_ATTACHMENT;
-	}
-	else if(EnumHasAnyFlags(Flags, TexCreate_ResolveTargetable))
-	{
-		Attachment = (Format == PF_DepthStencil)
-			? GL_DEPTH_STENCIL_ATTACHMENT
-			: ((Format == PF_ShadowDepth || Format == PF_D24)
-			? GL_DEPTH_ATTACHMENT
-			: GL_COLOR_ATTACHMENT0);
-	}
-
-	FOpenGLTexture3D* Texture = new FOpenGLTexture3D(this,TextureID,Target,Attachment,SizeX,SizeY,SizeZ,NumMips,1,1,1, (EPixelFormat)Format,false,true,Flags,CreateInfo.ClearValueBinding);
-	OpenGLTextureAllocated( Texture, Flags );
-
-	// No need to restore texture stage; leave it like this,
-	// and the next draw will take care of cleaning it up; or
-	// next operation that needs the stage will switch something else in on it.
-
-	return Texture;
-}
-
-FShaderResourceViewRHIRef FOpenGLDynamicRHI::RHICreateShaderResourceView(FRHITexture* Texture, const FRHITextureSRVCreateInfo& CreateInfo)
-{
-	const uint32 MipLevel = CreateInfo.MipLevel;
-	const uint32 NumMipLevels = CreateInfo.NumMipLevels;
-	EPixelFormat TextureBaseFormat = Texture->GetFormat();
-	const uint8 Format = (CreateInfo.Format == PF_Unknown) ? TextureBaseFormat : CreateInfo.Format;
-
-	FOpenGLShaderResourceViewProxy *ViewProxy = new FOpenGLShaderResourceViewProxy([this, Texture, MipLevel, NumMipLevels, Format](FRHIShaderResourceView* OwnerRHI) -> FOpenGLShaderResourceView*
-	{
-		if (FRHITexture2D* Texture2DRHI = Texture->GetTexture2D())
-		{
-			FOpenGLTexture2D* Texture2D = ResourceCast(Texture2DRHI);
-			FOpenGLShaderResourceView *View = 0;
-
-			if (FOpenGL::SupportsTextureView())
+									 SrcMipIndex,
+									 GLFormat.Format,
+									 GLFormat.Type,
+									 0);	// offset into PBO
+				glPixelStorei(GL_PACK_ALIGNMENT, 4);
+			}
+			
+			// copy the texture data
+			// Upload directly into Dst to avoid out-of-band synchronization caused by glMapBuffer!
 			{
-				VERIFY_GL_SCOPE();
-
-				GLuint Resource = 0;
-
-				FOpenGL::GenTextures(1, &Resource);
-
-				if (Format != PF_X24_G8)
+				CachedBindPixelUnpackBuffer( GL_PIXEL_UNPACK_BUFFER, PixelBuffer->Resource );
+				
+				// Use a texture stage that's not likely to be used for draws, to avoid waiting
+				FOpenGLDynamicRHI::Get().CachedSetupTextureStage(ContextState, FOpenGL::GetMaxCombinedTextureImageUnits() - 1, Target, GetResource(), -1, this->GetNumMips());
+				
+				if( this->GetSizeZ() )
 				{
-					// Choose original format when PF_Unknown is specified (as stated for FRHITextureSRVCreateInfo::Format)
-					const FOpenGLTextureFormat& GLFormat = GOpenGLTextureFormats[Format == PF_Unknown ? Texture2D->GetFormat() : Format];
-					const bool bSRGB = EnumHasAnyFlags(Texture2D->GetFlags(), TexCreate_SRGB);
-
-					FOpenGL::TextureView(Resource, Texture2D->Target, Texture2D->GetResource(), GLFormat.InternalFormat[bSRGB], MipLevel, NumMipLevels, 0, 1);
+					// texture 2D array
+					if (GLFormat.bCompressed)
+					{
+						FOpenGL::CompressedTexSubImage3D(Target,
+														 DstMipIndex,
+														 0,
+														 0,
+														 ArrayIndex,
+														 MipSizeX,
+														 MipSizeY,
+														 1,
+														 GLFormat.InternalFormat[bSRGB],
+														 PixelBuffer->GetSize(),
+														 0);
+					}
+					else
+					{
+						glPixelStorei(GL_UNPACK_ALIGNMENT, 1);
+						check( FOpenGL::SupportsTexture3D() );
+						FOpenGL::TexSubImage3D(Target,
+											   DstMipIndex,
+											   0,
+											   0,
+											   ArrayIndex,
+											   MipSizeX,
+											   MipSizeY,
+											   1,
+											   GLFormat.Format,
+											   GLFormat.Type,
+											   0);	// offset into PBO
+						glPixelStorei(GL_UNPACK_ALIGNMENT, 4);
+					}
 				}
 				else
 				{
-					// PF_X24_G8 doesn't correspond to a real format under OpenGL
-					// The solution is to create a view with the original format, and convert it to return the stencil index
-					// To match component locations, texture swizzle needs to be setup too
-					const FOpenGLTextureFormat& GLFormat = GOpenGLTextureFormats[Texture2D->GetFormat()];
-
-					// create a second depth/stencil view
-					FOpenGL::TextureView(Resource, Texture2D->Target, Texture2D->GetResource(), GLFormat.InternalFormat[0], MipLevel, NumMipLevels, 0, 1);
-
-					// Use a texture stage that's not likely to be used for draws, to avoid waiting
-					FOpenGLContextState& ContextState = GetContextStateForCurrentContext();
-					CachedSetupTextureStage(ContextState, FOpenGL::GetMaxCombinedTextureImageUnits() - 1, Texture2D->Target, Resource, 0, NumMipLevels);
-
-					//set the texture to return the stencil index, and then force the components to match D3D
-					glTexParameteri(Texture2D->Target, GL_DEPTH_STENCIL_TEXTURE_MODE, GL_STENCIL_INDEX);
-					glTexParameteri(Texture2D->Target, GL_TEXTURE_SWIZZLE_R, GL_ZERO);
-					glTexParameteri(Texture2D->Target, GL_TEXTURE_SWIZZLE_G, GL_RED);
-					glTexParameteri(Texture2D->Target, GL_TEXTURE_SWIZZLE_B, GL_ZERO);
-					glTexParameteri(Texture2D->Target, GL_TEXTURE_SWIZZLE_A, GL_ZERO);
-				}
-
-				View = new FOpenGLShaderResourceView(this, Resource, Texture2D->Target, MipLevel, true);
-			}
-			else
-			{
-				uint32 const Target = Texture2D->Target;
-				GLuint Resource = Texture2D->GetResource();
-
-				FRHITexture2D* DepthStencilTex = nullptr;
-
-				// For stencil sampling we have to use a separate single channel texture to blit stencil data into
-#if PLATFORM_DESKTOP
-				if (FOpenGL::GetFeatureLevel() >= ERHIFeatureLevel::SM5 && Format == PF_X24_G8)
-				{
-					check(NumMipLevels == 1 && MipLevel == 0);
-
-					if (!Texture2D->SRVResource)
+					if (GLFormat.bCompressed)
 					{
-						FOpenGL::GenTextures(1, &Texture2D->SRVResource);
-
-						GLenum const InternalFormat = GL_R8UI;
-						GLenum const ChannelFormat = GL_RED_INTEGER;
-						uint32 const SizeX = Texture2D->GetSizeX();
-						uint32 const SizeY = Texture2D->GetSizeY();
-						GLenum const Type = GL_UNSIGNED_BYTE;
-						ETextureCreateFlags const Flags = TexCreate_None;
-
-						FOpenGLContextState& ContextState = GetContextStateForCurrentContext();
-						CachedSetupTextureStage(ContextState, FOpenGL::GetMaxCombinedTextureImageUnits() - 1, Target, Texture2D->SRVResource, MipLevel, NumMipLevels);
-
-						if (!FOpenGL::TexStorage2D(Target, NumMipLevels, InternalFormat, SizeX, SizeY, ChannelFormat, Type, Flags))
+						if (GetAllocatedStorageForMip(DstMipIndex,ArrayIndex))
 						{
-							glTexImage2D(Target, 0, InternalFormat, SizeX, SizeY, 0, ChannelFormat, Type, nullptr);
+							glCompressedTexSubImage2D(bCubemap ? GL_TEXTURE_CUBE_MAP_POSITIVE_X + ArrayIndex : Target,
+													  DstMipIndex,
+													  0,
+													  0,
+													  MipSizeX,
+													  MipSizeY,
+													  GLFormat.InternalFormat[bSRGB],
+													  PixelBuffer->GetSize(),
+													  0);	// offset into PBO
 						}
-
-						TArray<uint8> ZeroData;
-						ZeroData.AddZeroed(SizeX * SizeY);
-
+						else
+						{
+							glCompressedTexImage2D(bCubemap ? GL_TEXTURE_CUBE_MAP_POSITIVE_X + ArrayIndex : Target,
+												   DstMipIndex,
+												   GLFormat.InternalFormat[bSRGB],
+												   MipSizeX,
+												   MipSizeY,
+												   0,
+												   PixelBuffer->GetSize(),
+												   0);	// offset into PBO
+							SetAllocatedStorageForMip(DstMipIndex,ArrayIndex);
+						}
+					}
+					else
+					{
+						// All construction paths should have called TexStorage2D or TexImage2D. So we will
+						// always call TexSubImage2D.
+						check(GetAllocatedStorageForMip(DstMipIndex,ArrayIndex) == true);
 						glPixelStorei(GL_UNPACK_ALIGNMENT, 1);
-						glTexSubImage2D(
-							Target,
-							0,
-							0,
-							0,
-							SizeX,
-							SizeY,
-							ChannelFormat,
-							Type,
-							ZeroData.GetData());
+						glTexSubImage2D(bCubemap ? GL_TEXTURE_CUBE_MAP_POSITIVE_X + ArrayIndex : Target,
+										DstMipIndex,
+										0,
+										0,
+										MipSizeX,
+										MipSizeY,
+										GLFormat.Format,
+										GLFormat.Type,
+										0);	// offset into PBO
 						glPixelStorei(GL_UNPACK_ALIGNMENT, 4);
-
-						//set the texture to return the stencil index, and then force the components to match D3D
-						glTexParameteri(Target, GL_TEXTURE_SWIZZLE_R, GL_ZERO);
-						glTexParameteri(Target, GL_TEXTURE_SWIZZLE_G, GL_RED);
-						glTexParameteri(Target, GL_TEXTURE_SWIZZLE_B, GL_ZERO);
-						glTexParameteri(Target, GL_TEXTURE_SWIZZLE_A, GL_ZERO);
 					}
-					check(Texture2D->SRVResource);
-
-					Resource = Texture2D->SRVResource;
-					DepthStencilTex = Texture2DRHI;
 				}
-#endif
-
-				View = new FOpenGLShaderResourceView(this, Resource, Target, MipLevel, false);
-				View->Texture2D = DepthStencilTex;
 			}
-			return View;
+			
+			// need to free PBO if we aren't keeping shadow copies
+			PixelBuffers[BufferIndex] = NULL;
+			
+			// No need to restore texture stage; leave it like this,
+			// and the next draw will take care of cleaning it up; or
+			// next operation that needs the stage will switch something else in on it.
 		}
-		else if (FRHITexture2DArray* Texture2DArrayRHI = Texture->GetTexture2DArray())
-		{
-			FOpenGLTexture2DArray* Texture2DArray = ResourceCast(Texture2DArrayRHI);
-			FOpenGLShaderResourceView *View = 0;
+	}
+	
+	// Reset the buffer bindings on exit only
+	glBindBuffer( GL_PIXEL_PACK_BUFFER, 0 );
+	CachedBindPixelUnpackBuffer(GL_PIXEL_UNPACK_BUFFER, 0);
+}
 
-			if (FOpenGL::SupportsTextureView())
-			{
-				VERIFY_GL_SCOPE();
+/*-----------------------------------------------------------------------------
+	2D texture support.
+-----------------------------------------------------------------------------*/
+
+FTextureRHIRef FOpenGLDynamicRHI::RHICreateTexture(const FRHITextureCreateDesc& CreateDesc)
+{
+	return new FOpenGLTexture(CreateDesc);
+}
 
-				GLuint Resource = 0;
+FTextureRHIRef FOpenGLDynamicRHI::RHIAsyncCreateTexture2D(uint32 SizeX, uint32 SizeY, uint8 Format, uint32 NumMips, ETextureCreateFlags Flags, ERHIAccess InResourceState, void** InitialMipData, uint32 NumInitialMips)
+{
+	check(0);
+	return FTexture2DRHIRef();
+}
 
-				FOpenGL::GenTextures(1, &Resource);
-				const FOpenGLTextureFormat& GLFormat = GOpenGLTextureFormats[Texture2DArray->GetFormat()];
-				const bool bSRGB = EnumHasAnyFlags(Texture2DArray->GetFlags(), TexCreate_SRGB);
+void FOpenGLDynamicRHI::RHICopySharedMips(FRHITexture2D* DestTexture2D, FRHITexture2D* SrcTexture2D)
+{
+	check(0);
+}
 
-				FOpenGL::TextureView(Resource, Texture2DArray->Target, Texture2DArray->GetResource(), GLFormat.InternalFormat[bSRGB], MipLevel, 1, 0, 1);
+FOpenGLShaderResourceView::FOpenGLShaderResourceView(FOpenGLTexture* InTexture, const FRHITextureSRVCreateInfo& CreateInfo)
+	: Target      (InTexture->Target)
+	, Texture     (InTexture)
+	, LimitMip    (CreateInfo.MipLevel)
+	, OwnsResource(FOpenGL::SupportsTextureView())
+{
+	check(CreateInfo.NumMipLevels == 1);
 
-				return new FOpenGLShaderResourceView(this, Resource, Texture2DArray->Target, MipLevel, true);
-			}
-			else
-			{
-				return new FOpenGLShaderResourceView(this, Texture2DArray->GetResource(), Texture2DArray->Target, MipLevel, false);
-			}
-		}
-		else if (FRHITextureCube* TextureCubeRHI = Texture->GetTextureCube())
+	if (OwnsResource)
+	{
+		FOpenGL::GenTextures(1, &Resource);
+	}
+
+	const uint32 MipLevel = CreateInfo.MipLevel;
+	const uint32 NumMipLevels = CreateInfo.NumMipLevels;
+	const uint8 Format = (CreateInfo.Format == PF_Unknown)
+		? Texture->GetFormat()
+		: CreateInfo.Format;
+
+	const bool bFormatsMatch = Format == Texture->GetFormat();
+	const bool bSRGB = CreateInfo.SRGBOverride == SRGBO_Default
+		? EnumHasAnyFlags(InTexture->GetDesc().Flags, TexCreate_SRGB)
+		: false;
+
+	if (Format == PF_X24_G8)
+	{
+		RunOnGLRenderContextThread([this, MipLevel, NumMipLevels]()
 		{
-			FOpenGLTextureCube* TextureCube = ResourceCast(TextureCubeRHI);
+			VERIFY_GL_SCOPE();
 			if (FOpenGL::SupportsTextureView())
 			{
-				VERIFY_GL_SCOPE();
-
-				GLuint Resource = 0;
+				// PF_X24_G8 doesn't correspond to a real format under OpenGL
+				// The solution is to create a view with the original format, and convert it to return the stencil index
+				// To match component locations, texture swizzle needs to be setup too
+				const FOpenGLTextureFormat& GLFormat = GOpenGLTextureFormats[Texture->GetFormat()];
 
-				FOpenGL::GenTextures(1, &Resource);
-				const FOpenGLTextureFormat& GLFormat = GOpenGLTextureFormats[TextureCube->GetFormat()];
-				const bool bSRGB = EnumHasAnyFlags(TextureCube->GetFlags(), TexCreate_SRGB);
+				// create a second depth/stencil view
+				FOpenGL::TextureView(Resource, Target, Texture->GetResource(), GLFormat.InternalFormat[0], MipLevel, NumMipLevels, 0, 1);
 
-				FOpenGL::TextureView(Resource, TextureCube->Target, TextureCube->GetResource(), GLFormat.InternalFormat[bSRGB], MipLevel, 1, 0, 6);
-
-				return new FOpenGLShaderResourceView(this, Resource, TextureCube->Target, MipLevel, true);
+				// Use a texture stage that's not likely to be used for draws, to avoid waiting
+				FOpenGLContextState& ContextState = FOpenGLDynamicRHI::Get().GetContextStateForCurrentContext();
+				FOpenGLDynamicRHI::Get().CachedSetupTextureStage(ContextState, FOpenGL::GetMaxCombinedTextureImageUnits() - 1, Target, Resource, 0, NumMipLevels);
+
+				//set the texture to return the stencil index, and then force the components to match D3D
+				glTexParameteri(Target, GL_DEPTH_STENCIL_TEXTURE_MODE, GL_STENCIL_INDEX);
+				glTexParameteri(Target, GL_TEXTURE_SWIZZLE_R, GL_ZERO);
+				glTexParameteri(Target, GL_TEXTURE_SWIZZLE_G, GL_RED);
+				glTexParameteri(Target, GL_TEXTURE_SWIZZLE_B, GL_ZERO);
+				glTexParameteri(Target, GL_TEXTURE_SWIZZLE_A, GL_ZERO);
 			}
 			else
 			{
-				return new FOpenGLShaderResourceView(this, TextureCube->GetResource(), TextureCube->Target, MipLevel, false);
+				UE_LOG(LogRHI, Fatal, TEXT("Cannot create a stencil SRV (PF_X24_G8) when texture views are unsupported."));
 			}
-		}
-		else if (FRHITexture3D* Texture3DRHI = Texture->GetTexture3D())
+		});
+	}
+	else
+	{
+		RunOnGLRenderContextThread([this, Format, MipLevel, bSRGB, bFormatsMatch]()
 		{
-			FOpenGLTexture3D* Texture3D = ResourceCast(Texture3DRHI);
-
-			FOpenGLShaderResourceView *View = 0;
-
+			VERIFY_GL_SCOPE();
 			if (FOpenGL::SupportsTextureView())
 			{
-				VERIFY_GL_SCOPE();
-
-				GLuint Resource = 0;
-
-				FOpenGL::GenTextures(1, &Resource);
-				const FOpenGLTextureFormat& GLFormat = GOpenGLTextureFormats[Texture3D->GetFormat()];
-				const bool bSRGB = EnumHasAnyFlags(Texture3D->GetFlags(), TexCreate_SRGB);
-
-				FOpenGL::TextureView(Resource, Texture3D->Target, Texture3D->GetResource(), GLFormat.InternalFormat[bSRGB], MipLevel, 1, 0, 1);
-
-				return new FOpenGLShaderResourceView(this, Resource, Texture3D->Target, MipLevel, true);
+				FOpenGL::TextureView(Resource, Target, Texture->GetResource(), GOpenGLTextureFormats[Format].InternalFormat[bSRGB], MipLevel, 1, 0, 1);
 			}
 			else
 			{
-				return new FOpenGLShaderResourceView(this, Texture3D->GetResource(), Texture3D->Target, MipLevel, false);
+				checkf(bFormatsMatch, TEXT("SRVs cannot modify the pixel format of a texture when texture views are unsupported."));
+				Resource = Texture->GetResource();
 			}
-		}
-		else
-		{
-			check(false);
-			return nullptr;
-		}
-	});
-	return ViewProxy;
+		});
+	}
+}
+
+FShaderResourceViewRHIRef FOpenGLDynamicRHI::RHICreateShaderResourceView(FRHITexture* Texture, const FRHITextureSRVCreateInfo& CreateInfo)
+{
+	return new FOpenGLShaderResourceView(GetOpenGLTextureFromRHITexture(Texture), CreateInfo);
 }
 
 /** Generates mip maps for the surface. */
@@ -2115,10 +1745,10 @@ void FOpenGLDynamicRHI::RHIGenerateMips(FRHITexture* SurfaceRHI)
 			GPUProfilingData.RegisterGPUWork(0);
 
 			FOpenGLContextState& ContextState = GetContextStateForCurrentContext();
-			FOpenGLTextureBase* Texture = GetOpenGLTextureFromRHITexture(SurfaceRHI);
+			FOpenGLTexture* Texture = GetOpenGLTextureFromRHITexture(SurfaceRHI);
 			// Setup the texture on a disused unit
 			// need to figure out how to setup mips properly in no views case
-			CachedSetupTextureStage(ContextState, FOpenGL::GetMaxCombinedTextureImageUnits() - 1, Texture->Target, Texture->GetResource(), -1, Texture->NumMips); 
+			CachedSetupTextureStage(ContextState, FOpenGL::GetMaxCombinedTextureImageUnits() - 1, Texture->Target, Texture->GetResource(), -1, Texture->GetNumMips());
 			FOpenGL::GenerateMipmap(Texture->Target);
 		});
 	}
@@ -2138,111 +1768,63 @@ void FOpenGLDynamicRHI::RHIGenerateMips(FRHITexture* SurfaceRHI)
  */
 uint32 FOpenGLDynamicRHI::RHIComputeMemorySize(FRHITexture* TextureRHI)
 {
-	if(!TextureRHI)
+	if (!TextureRHI)
 	{
 		return 0;
 	}
 
-	FOpenGLTextureBase* Texture = static_cast<FOpenGLTextureBase*>(TextureRHI->GetTextureBaseRHI());
-	if (!Texture->IsMemorySizeSet())
-	{
-		GetOpenGLTextureFromRHITexture(TextureRHI);
-	}
-	return Texture->GetMemorySize();
+	FOpenGLTexture* Texture = GetOpenGLTextureFromRHITexture(TextureRHI);
+	return Texture->MemorySize;
 }
 
-
-static FTexture2DRHIRef CreateAsyncReallocate2DTextureTarget(FOpenGLDynamicRHI* OGLRHI, FRHITexture2D* Texture2DRHI, int32 NewMipCountIn, int32 NewSizeX, int32 NewSizeY)
+FTexture2DRHIRef FOpenGLDynamicRHI::AsyncReallocateTexture2D_RenderThread(class FRHICommandListImmediate& RHICmdList, FRHITexture2D* Texture2DRHI, int32 NewMipCount, int32 NewSizeX, int32 NewSizeY, FThreadSafeCounter* RequestStatus)
 {
-	FOpenGLTexture2D* Texture2D = FOpenGLDynamicRHI::ResourceCast(Texture2DRHI);
-	uint8 Format = Texture2D->GetFormat();
-	uint32 NumSamples = 1;
-	ETextureCreateFlags Flags = Texture2D->GetFlags();
-	uint32 NewMipCount = (uint32)NewMipCountIn;
-	uint32 OriginalMipCount = Texture2DRHI->GetNumMips();
-	const FClearValueBinding ClearBinding = Texture2DRHI->GetClearBinding();
-	FTexture2DRHIRef NewTexture2DRHI = (FRHITexture2D*)OGLRHI->CreateOpenGLRHITextureOnly(NewSizeX, NewSizeY, false, false, false, Format, NewMipCount, NumSamples, 1, Flags, ClearBinding);
-
-	// CreateOpenGLRHITextureOnly can potentially change some of the input parameters, ensure that's not happening:
-	check(Format == (uint8)Texture2D->GetFormat());
-	check(Flags == Texture2D->GetFlags());
-	check(NumSamples == 1);
-	return NewTexture2DRHI;
+	return this->RHIAsyncReallocateTexture2D(Texture2DRHI, NewMipCount, NewSizeX, NewSizeY, RequestStatus);
 }
 
-static void GLCopyAsyncTexture2D(FOpenGLDynamicRHI* OGLRHI, FRHITexture2D* NewTexture2DRHI, int32 NewSizeX, int32 NewSizeY, FRHITexture2D* SourceTexture2DRHI, FThreadSafeCounter* RequestStatus)
+FTexture2DRHIRef FOpenGLDynamicRHI::RHIAsyncReallocateTexture2D(FRHITexture2D* Texture2DRHI, int32 NewMipCount, int32 NewSizeX, int32 NewSizeY, FThreadSafeCounter* RequestStatus)
 {
-	VERIFY_GL_SCOPE();
-
-	FOpenGLTexture2D* SourceTexture2D = FOpenGLDynamicRHI::ResourceCast(SourceTexture2DRHI);
-	uint8 Format = NewTexture2DRHI->GetFormat();
-	uint32 NumSamples = 1;
-	ETextureCreateFlags Flags = NewTexture2DRHI->GetFlags();
-	uint32 NewMipCount = (uint32)NewTexture2DRHI->GetNumMips();
-	uint32 SourceMipCount = SourceTexture2DRHI->GetNumMips();
-
-	const FClearValueBinding ClearBinding = NewTexture2DRHI->GetClearBinding();
+	check(IsInRenderingThread());
+	FRHICommandListImmediate& RHICmdList = FRHICommandListExecutor::GetImmediateCommandList();
 
+	FOpenGLTexture* OldTexture = GetOpenGLTextureFromRHITexture(Texture2DRHI);
 
-	OGLRHI->InitializeGLTexture(NewTexture2DRHI, NewSizeX, NewSizeY, false, false, false, Format, NewMipCount, 1, 1, Flags, ClearBinding);
+	FRHITextureDesc Desc = OldTexture->GetDesc();
+	int32 SourceMipCount = Desc.NumMips;
 
-	FOpenGLTexture2D* NewTexture2D = FOpenGLDynamicRHI::ResourceCast(NewTexture2DRHI);
+	Desc.Extent = FIntPoint(NewSizeX, NewSizeY);
+	Desc.NumMips = NewMipCount;
 
-	const uint32 BlockSizeX = GPixelFormats[Format].BlockSizeX;
-	const uint32 BlockSizeY = GPixelFormats[Format].BlockSizeY;
-	const uint32 NumBytesPerBlock = GPixelFormats[Format].BlockBytes;
+	FRHITextureCreateDesc CreateDesc(
+		Desc,
+		RHIGetDefaultResourceState(Desc.Flags, false),
+		TEXT("RHIAsyncReallocateTexture2D")
+	);
 
-	// Use the GPU to asynchronously copy the old mip-maps into the new texture.
-	const uint32 NumSharedMips = FMath::Min(SourceMipCount, NewMipCount);
-	const uint32 SourceMipOffset = SourceMipCount - NumSharedMips;
-	const uint32 DestMipOffset = NewMipCount - NumSharedMips;
+	FOpenGLTexture* NewTexture = new FOpenGLTexture(CreateDesc);
 
-	if (FOpenGL::SupportsCopyImage())
+	RHICmdList.EnqueueLambda([OldTexture, NewTexture, SourceMipCount, NewMipCount, RequestStatus](FRHICommandListImmediate& RHICmdList)
 	{
-		FOpenGLTexture2D *NewOGLTexture2D = (FOpenGLTexture2D*)NewTexture2D;
-		FOpenGLTexture2D *OGLTexture2D = (FOpenGLTexture2D*)SourceTexture2D;
-		NewOGLTexture2D->CloneViaCopyImage(OGLTexture2D, NumSharedMips, SourceMipOffset, DestMipOffset);
-	}
-	else
-	{
-		FOpenGLTexture2D *NewOGLTexture2D = (FOpenGLTexture2D*)NewTexture2D;
-		FOpenGLTexture2D *OGLTexture2D = (FOpenGLTexture2D*)SourceTexture2D;
-		NewOGLTexture2D->CloneViaPBO(OGLTexture2D, NumSharedMips, SourceMipOffset, DestMipOffset);
-	}
-
-	// Decrement the thread-safe counter used to track the completion of the reallocation, since D3D handles sequencing the
-	// async mip copies with other D3D calls.
-	RequestStatus->Decrement();
-}
+		VERIFY_GL_SCOPE();
 
-FTexture2DRHIRef FOpenGLDynamicRHI::AsyncReallocateTexture2D_RenderThread(class FRHICommandListImmediate& RHICmdList, FRHITexture2D* Texture2DRHI, int32 NewMipCountIn, int32 NewSizeX, int32 NewSizeY, FThreadSafeCounter* RequestStatus)
-{
-	if (RHICmdList.Bypass() || !IsRunningRHIInSeparateThread())
-	{
-		return RHIAsyncReallocateTexture2D(Texture2DRHI, NewMipCountIn, NewSizeX, NewSizeY, RequestStatus);
-	}
-	else
-	{
-		FTexture2DRHIRef NewTexture2DRHI = CreateAsyncReallocate2DTextureTarget(this, Texture2DRHI, NewMipCountIn, NewSizeX, NewSizeY);
-		FOpenGLTexture2D* Texture2D = FOpenGLDynamicRHI::ResourceCast(NewTexture2DRHI.GetReference());
-		Texture2D->CreationFence.Reset();
+		// Use the GPU to asynchronously copy the old mip-maps into the new texture.
+		const uint32 NumSharedMips = FMath::Min(SourceMipCount, NewMipCount);
+		const uint32 SourceMipOffset = SourceMipCount - NumSharedMips;
+		const uint32 DestMipOffset = NewMipCount - NumSharedMips;
 
-		ALLOC_COMMAND_CL(RHICmdList, FRHICommandGLCommand)([=]() 
+		if (FOpenGL::SupportsCopyImage())
 		{
-			GLCopyAsyncTexture2D(this, NewTexture2DRHI, NewSizeX, NewSizeY, Texture2DRHI, RequestStatus); 
-			Texture2D->CreationFence.WriteAssertFence();
-		});
+			NewTexture->CloneViaCopyImage(OldTexture, NumSharedMips, SourceMipOffset, DestMipOffset);
+		}
+		else
+		{
+			NewTexture->CloneViaPBO(OldTexture, NumSharedMips, SourceMipOffset, DestMipOffset);
+		}
 
-		Texture2D->CreationFence.SetRHIThreadFence();
-		return NewTexture2DRHI;
-	}
-}
+		RequestStatus->Decrement();
+	});
 
-FTexture2DRHIRef FOpenGLDynamicRHI::RHIAsyncReallocateTexture2D(FRHITexture2D* Texture2DRHI, int32 NewMipCountIn, int32 NewSizeX, int32 NewSizeY, FThreadSafeCounter* RequestStatus)
-{
-	FTexture2DRHIRef NewTexture2DRHI = CreateAsyncReallocate2DTextureTarget(this, Texture2DRHI, NewMipCountIn, NewSizeX, NewSizeY);
-	GLCopyAsyncTexture2D(this, NewTexture2DRHI, NewSizeX, NewSizeY, Texture2DRHI, RequestStatus);
-	return NewTexture2DRHI;
+	return NewTexture;
 }
 
 /**
@@ -2274,29 +1856,25 @@ ETextureReallocationStatus FOpenGLDynamicRHI::RHICancelAsyncReallocateTexture2D(
 
 void* FOpenGLDynamicRHI::RHILockTexture2D(FRHITexture2D* TextureRHI,uint32 MipIndex,EResourceLockMode LockMode,uint32& DestStride,bool bLockWithinMiptail)
 {
-	FOpenGLTexture2D* Texture = ResourceCast(TextureRHI);
-	return Texture->Lock(MipIndex,0,LockMode,DestStride);
+	return GetOpenGLTextureFromRHITexture(TextureRHI)->Lock(MipIndex, 0, LockMode, DestStride);
 }
 
-void FOpenGLDynamicRHI::RHIUnlockTexture2D(FRHITexture2D* TextureRHI,uint32 MipIndex,bool bLockWithinMiptail)
+void FOpenGLDynamicRHI::RHIUnlockTexture2D(FRHITexture2D* TextureRHI, uint32 MipIndex, bool bLockWithinMiptail)
 {
-	FOpenGLTexture2D* Texture = ResourceCast(TextureRHI);
-	Texture->Unlock(MipIndex, 0);
+	GetOpenGLTextureFromRHITexture(TextureRHI)->Unlock(MipIndex, 0);
 }
 
-void* FOpenGLDynamicRHI::RHILockTexture2DArray(FRHITexture2DArray* TextureRHI,uint32 TextureIndex,uint32 MipIndex,EResourceLockMode LockMode,uint32& DestStride,bool bLockWithinMiptail)
+void* FOpenGLDynamicRHI::RHILockTexture2DArray(FRHITexture2DArray* TextureRHI, uint32 TextureIndex, uint32 MipIndex, EResourceLockMode LockMode, uint32& DestStride, bool bLockWithinMiptail)
 {
-	FOpenGLTexture2DArray* Texture = ResourceCast(TextureRHI);
-	return Texture->Lock(MipIndex,TextureIndex,LockMode,DestStride);
+	return GetOpenGLTextureFromRHITexture(TextureRHI)->Lock(MipIndex, TextureIndex, LockMode, DestStride);
 }
 
-void FOpenGLDynamicRHI::RHIUnlockTexture2DArray(FRHITexture2DArray* TextureRHI,uint32 TextureIndex,uint32 MipIndex,bool bLockWithinMiptail)
+void FOpenGLDynamicRHI::RHIUnlockTexture2DArray(FRHITexture2DArray* TextureRHI, uint32 TextureIndex, uint32 MipIndex, bool bLockWithinMiptail)
 {
-	FOpenGLTexture2DArray* Texture = ResourceCast(TextureRHI);
-	Texture->Unlock(MipIndex, TextureIndex);
+	GetOpenGLTextureFromRHITexture(TextureRHI)->Unlock(MipIndex, TextureIndex);
 }
 
-void FOpenGLDynamicRHI::RHIUpdateTexture2D(FRHITexture2D* TextureRHI,uint32 MipIndex,const FUpdateTextureRegion2D& UpdateRegionIn,uint32 SourcePitch,const uint8* SourceDataIn)
+void FOpenGLDynamicRHI::RHIUpdateTexture2D(FRHITexture2D* TextureRHI, uint32 MipIndex, const FUpdateTextureRegion2D& UpdateRegionIn, uint32 SourcePitch, const uint8* SourceDataIn)
 {
 	FRHICommandListImmediate& RHICmdList = FRHICommandListExecutor::GetImmediateCommandList();
 	const FUpdateTextureRegion2D UpdateRegion = UpdateRegionIn;
@@ -2315,7 +1893,7 @@ void FOpenGLDynamicRHI::RHIUpdateTexture2D(FRHITexture2D* TextureRHI,uint32 MipI
 	{
 		VERIFY_GL_SCOPE();
 
-		FOpenGLTexture2D* Texture = ResourceCast(TextureRHI);
+		FOpenGLTexture* Texture = GetOpenGLTextureFromRHITexture(TextureRHI);
 
 		// Use a texture stage that's not likely to be used for draws, to avoid waiting
 		FOpenGLContextState& ContextState = GetContextStateForCurrentContext();
@@ -2354,7 +1932,7 @@ void FOpenGLDynamicRHI::RHIUpdateTexture3D(FRHITexture3D* TextureRHI, uint32 Mip
 {
 	VERIFY_GL_SCOPE();
 	check( FOpenGL::SupportsTexture3D() );
-	FOpenGLTexture3D* Texture = ResourceCast(TextureRHI);
+	FOpenGLTexture* Texture = GetOpenGLTextureFromRHITexture(TextureRHI);
 
 	// Use a texture stage that's not likely to be used for draws, to avoid waiting
 	FOpenGLContextState& ContextState = GetContextStateForCurrentContext();
@@ -2470,32 +2048,14 @@ void FOpenGLDynamicRHI::InvalidateUAVResourceInCache(GLuint Resource)
 /*-----------------------------------------------------------------------------
 	Cubemap texture support.
 -----------------------------------------------------------------------------*/
-FTextureCubeRHIRef FOpenGLDynamicRHI::RHICreateTextureCube( uint32 Size, uint8 Format, uint32 NumMips, ETextureCreateFlags Flags, ERHIAccess InResourceState, FRHIResourceCreateInfo& CreateInfo )
-{
-	// not yet supported
-	check(!CreateInfo.BulkData);
-
-	return (FRHITextureCube*)CreateOpenGLTexture(Size,Size,true, false, false, Format, NumMips, 1, 1, Flags, CreateInfo.ClearValueBinding);
-}
-
-FTextureCubeRHIRef FOpenGLDynamicRHI::RHICreateTextureCubeArray( uint32 Size, uint32 ArraySize, uint8 Format, uint32 NumMips, ETextureCreateFlags Flags, ERHIAccess InResourceState, FRHIResourceCreateInfo& CreateInfo )
-{
-	// not yet supported
-	check(!CreateInfo.BulkData);
-
-	return (FRHITextureCube*)CreateOpenGLTexture(Size, Size, true, true, false, Format, NumMips, 1, 6 * ArraySize, Flags, CreateInfo.ClearValueBinding);
-}
-
-void* FOpenGLDynamicRHI::RHILockTextureCubeFace(FRHITextureCube* TextureCubeRHI,uint32 FaceIndex,uint32 ArrayIndex,uint32 MipIndex,EResourceLockMode LockMode,uint32& DestStride,bool bLockWithinMiptail)
+void* FOpenGLDynamicRHI::RHILockTextureCubeFace(FRHITextureCube* TextureCubeRHI, uint32 FaceIndex, uint32 ArrayIndex, uint32 MipIndex, EResourceLockMode LockMode, uint32& DestStride, bool bLockWithinMiptail)
 {
-	FOpenGLTextureCube* TextureCube = ResourceCast(TextureCubeRHI);
-	return TextureCube->Lock(MipIndex,FaceIndex + 6 * ArrayIndex,LockMode,DestStride);
+	return GetOpenGLTextureFromRHITexture(TextureCubeRHI)->Lock(MipIndex, FaceIndex + 6 * ArrayIndex, LockMode, DestStride);
 }
 
-void FOpenGLDynamicRHI::RHIUnlockTextureCubeFace(FRHITextureCube* TextureCubeRHI,uint32 FaceIndex,uint32 ArrayIndex,uint32 MipIndex,bool bLockWithinMiptail)
+void FOpenGLDynamicRHI::RHIUnlockTextureCubeFace(FRHITextureCube* TextureCubeRHI, uint32 FaceIndex, uint32 ArrayIndex, uint32 MipIndex, bool bLockWithinMiptail)
 {
-	FOpenGLTextureCube* TextureCube = ResourceCast(TextureCubeRHI);
-	TextureCube->Unlock(MipIndex,FaceIndex + ArrayIndex * 6);
+	GetOpenGLTextureFromRHITexture(TextureCubeRHI)->Unlock(MipIndex, FaceIndex + ArrayIndex * 6);
 }
 
 void FOpenGLDynamicRHI::RHIBindDebugLabelName(FRHITexture* TextureRHI, const TCHAR* Name)
@@ -2505,7 +2065,7 @@ void FOpenGLDynamicRHI::RHIBindDebugLabelName(FRHITexture* TextureRHI, const TCH
 	if (ShouldRunGLRenderContextOpOnThisThread(RHICmdList))
 	{
 		VERIFY_GL_SCOPE();
-		FOpenGLTextureBase* Texture = GetOpenGLTextureFromRHITexture(TextureRHI);
+		FOpenGLTexture* Texture = GetOpenGLTextureFromRHITexture(TextureRHI);
 		if (Texture->IsEvicted())
 		{
 			Texture->EvictionParamsPtr->SetDebugLabelName(TCHAR_TO_ANSI(Name));
@@ -2523,7 +2083,7 @@ void FOpenGLDynamicRHI::RHIBindDebugLabelName(FRHITexture* TextureRHI, const TCH
 		RunOnGLRenderContextThread([TextureRHI, TextureDebugName = MoveTemp(TextureDebugName)]()
 		{
 			VERIFY_GL_SCOPE();
-			FOpenGLTextureBase* Texture = GetOpenGLTextureFromRHITexture(TextureRHI);
+			FOpenGLTexture* Texture = GetOpenGLTextureFromRHITexture(TextureRHI);
 			if (Texture->IsEvicted())
 			{
 				Texture->EvictionParamsPtr->SetDebugLabelName(TextureDebugName);
@@ -2548,8 +2108,8 @@ void FOpenGLDynamicRHI::RHIVirtualTextureSetFirstMipVisible(FRHITexture2D* Textu
 void FOpenGLDynamicRHI::RHICopyTexture(FRHITexture* SourceTextureRHI, FRHITexture* DestTextureRHI, const FRHICopyTextureInfo& CopyInfo)
 {
 	VERIFY_GL_SCOPE();
-	FOpenGLTextureBase* SourceTexture = GetOpenGLTextureFromRHITexture(SourceTextureRHI);
-	FOpenGLTextureBase* DestTexture = GetOpenGLTextureFromRHITexture(DestTextureRHI);
+	FOpenGLTexture* SourceTexture = GetOpenGLTextureFromRHITexture(SourceTextureRHI);
+	FOpenGLTexture* DestTexture = GetOpenGLTextureFromRHITexture(DestTextureRHI);
 
 	checkf(SourceTexture->Target == DestTexture->Target, TEXT("Cannot copy between different texture targets, SourceTexture Target=%x, Format=%d, Flags=%x; DestTexture Target=%x, Format=%d, Flags=%x"),
 		SourceTexture->Target, SourceTextureRHI->GetFormat(), SourceTextureRHI->GetFlags(),
@@ -2654,7 +2214,7 @@ void FOpenGLDynamicRHI::RHICopyTexture(FRHITexture* SourceTextureRHI, FRHITextur
 	{
 		for (uint32 MipIndex = 0; MipIndex < CopyInfo.NumMips; ++MipIndex)
 		{
-			FOpenGLTextureBase* RenderTargets[1] = { SourceTexture };
+			FOpenGLTexture* RenderTargets[1] = { SourceTexture };
 			uint32 MipLevels[1] = { static_cast<uint32>(SrcMip) };
 			uint32 ArrayIndices[1] = { CopyInfo.SourceSliceIndex + SliceIndex };
 
@@ -2705,138 +2265,75 @@ void FOpenGLDynamicRHI::RHICopyTexture(FRHITexture* SourceTextureRHI, FRHITextur
 
 FTexture2DRHIRef FOpenGLDynamicRHI::RHICreateTexture2DFromResource(EPixelFormat Format, uint32 SizeX, uint32 SizeY, uint32 NumMips, uint32 NumSamples, uint32 NumSamplesTileMem, const FClearValueBinding& ClearValueBinding, GLuint Resource, ETextureCreateFlags TexCreateFlags)
 {
-	FOpenGLTexture2D* Texture2D = new FOpenGLTexture2D(
-		this,
-		Resource,
-		(NumSamples > 1) ? GL_TEXTURE_2D_MULTISAMPLE : GL_TEXTURE_2D,
-		GL_NONE,
-		SizeX,
-		SizeY,
-		0,
-		NumMips,
-		NumSamples,
-		NumSamplesTileMem,
-		1,
+	return new FOpenGLTexture(FRHITextureCreateDesc::Create2D(
+		TEXT("RHICreateTexture2DFromResource"),
+		{ (int32)SizeX, (int32)SizeY },
 		Format,
-		false,
-		false,
+		ClearValueBinding,
 		TexCreateFlags,
-		ClearValueBinding);
-
-	Texture2D->SetAliased(true);
-	OpenGLTextureAllocated(Texture2D, TexCreateFlags);
-	return Texture2D;
+		NumMips,
+		NumSamples
+	), Resource);
 }
 
 FTexture2DRHIRef FOpenGLDynamicRHI::RHICreateTexture2DArrayFromResource(EPixelFormat Format, uint32 SizeX, uint32 SizeY, uint32 ArraySize, uint32 NumMips, uint32 NumSamples, uint32 NumSamplesTileMem, const FClearValueBinding& ClearValueBinding, GLuint Resource, ETextureCreateFlags TexCreateFlags)
 {
-	FOpenGLTexture2D* Texture2DArray = new FOpenGLTexture2D(
-		this,
-		Resource,
-		GL_TEXTURE_2D_ARRAY,
-		GL_NONE,
-		SizeX,
-		SizeY,
-		0,
-		NumMips,
-		NumSamples,
-		NumSamplesTileMem,
-		ArraySize,
+	return new FOpenGLTexture(FRHITextureCreateDesc::Create2DArray(
+		TEXT("RHICreateTexture2DArrayFromResource"),
+		{ (int32)SizeX, (int32)SizeY },
 		Format,
-		false,
-		false,
+		ClearValueBinding,
 		TexCreateFlags,
-		ClearValueBinding);
-
-	Texture2DArray->SetAliased(true);
-	OpenGLTextureAllocated(Texture2DArray, TexCreateFlags);
-	return Texture2DArray;
+		ArraySize,
+		NumMips,
+		NumSamples
+	), Resource);
 }
 
 FTextureCubeRHIRef FOpenGLDynamicRHI::RHICreateTextureCubeFromResource(EPixelFormat Format, uint32 Size, bool bArray, uint32 ArraySize, uint32 NumMips, uint32 NumSamples, uint32 NumSamplesTileMem, const FClearValueBinding& ClearValueBinding, GLuint Resource, ETextureCreateFlags TexCreateFlags)
 {
-	FOpenGLTextureCube* TextureCube = new FOpenGLTextureCube(
-		this,
-		Resource,
-		GL_TEXTURE_CUBE_MAP,
-		GL_NONE,
-		Size,
-		Size,
-		0,
-		NumMips,
-		NumSamples,
-		NumSamplesTileMem,
-		1,
-		Format,
-		false,
-		false,
-		TexCreateFlags,
-		ClearValueBinding);
-
-	TextureCube->SetAliased(true);
-	OpenGLTextureAllocated(TextureCube, TexCreateFlags);
-	return TextureCube;
+	if (bArray)
+	{
+		return new FOpenGLTexture(FRHITextureCreateDesc::CreateCube(
+			TEXT("RHICreateTextureCubeFromResource"),
+			Size,
+			Format,
+			ClearValueBinding,
+			TexCreateFlags,
+			NumMips,
+			NumSamples
+		), Resource);
+	}
+	else
+	{
+		return new FOpenGLTexture(FRHITextureCreateDesc::CreateCubeArray(
+			TEXT("RHICreateTextureCubeFromResource"),
+			Size,
+			Format,
+			ClearValueBinding,
+			TexCreateFlags,
+			ArraySize,
+			NumMips,
+			NumSamples
+		), Resource);
+	}
 }
 
 void FOpenGLDynamicRHI::RHIAliasTextureResources(FTextureRHIRef& DestRHITexture, FTextureRHIRef& SrcRHITexture)
 {
 	VERIFY_GL_SCOPE();
-	FOpenGLTextureBase* DestTexture = GetOpenGLTextureFromRHITexture(DestRHITexture);
-	FOpenGLTextureBase* SrcTexture = GetOpenGLTextureFromRHITexture(SrcRHITexture);
+	FOpenGLTexture* DestTexture = GetOpenGLTextureFromRHITexture(DestRHITexture);
+	FOpenGLTexture* SrcTexture = GetOpenGLTextureFromRHITexture(SrcRHITexture);
 
 	if (DestTexture && SrcTexture)
 	{
-		DestTexture->AliasResources(SrcTexture);
+		DestTexture->AliasResources(*SrcTexture);
 	}
 }
 
 FTextureRHIRef FOpenGLDynamicRHI::RHICreateAliasedTexture(FTextureRHIRef& SourceTexture)
 {
-	SCOPE_CYCLE_COUNTER(STAT_OpenGLCreateTextureTime);
-
-	FTextureRHIRef Result;
-	if (SourceTexture->GetTexture2D() != nullptr)
-	{
-		FOpenGLTexture2D* SourceTexture2D = static_cast<FOpenGLTexture2D*>(SourceTexture.GetReference());
-
-		Result = new FOpenGLTexture2D(this, 0, SourceTexture2D->Target, -1, SourceTexture2D->GetSizeX(), SourceTexture2D->GetSizeY(), 0,
-			SourceTexture2D->GetNumMips(), SourceTexture2D->GetNumSamples(), SourceTexture2D->GetNumSamplesTileMem(), 1, SourceTexture2D->GetFormat(),
-			false, false, SourceTexture2D->GetFlags(), SourceTexture2D->GetClearBinding());
-	}
-	else if (SourceTexture->GetTexture2DArray() != nullptr)
-	{
-		FOpenGLTexture2DArray* SourceTexture2DArray = static_cast<FOpenGLTexture2DArray*>(SourceTexture.GetReference());
-
-		Result = new FOpenGLTexture2DArray(this, 0, SourceTexture2DArray->Target, -1, SourceTexture2DArray->GetSizeX(), SourceTexture2DArray->GetSizeY(), SourceTexture2DArray->GetSizeZ(),
-			SourceTexture2DArray->GetNumMips(), SourceTexture2DArray->GetNumSamples(), 1 /* aka check(InNumSamplesTileMem == 1) in OpenGLResource.h FOpenGLBaseTexture2DArray constructor */,
-			1, SourceTexture2DArray->GetFormat(), false, false, SourceTexture2DArray->GetFlags(), SourceTexture2DArray->GetClearBinding());
-	}
-	else if (SourceTexture->GetTextureCube() != nullptr)
-	{
-		FOpenGLTextureCube* SourceTextureCube = static_cast<FOpenGLTextureCube*>(SourceTexture.GetReference());
-
-		Result = new FOpenGLTextureCube(this, 0, SourceTextureCube->Target, -1, SourceTextureCube->GetSizeX(), SourceTextureCube->GetSizeY(), SourceTextureCube->GetSizeZ(),
-			SourceTextureCube->GetNumMips(), SourceTextureCube->GetNumSamples(), 1 /* OpenGL currently doesn't support multisample cube textures, per OpenGLResource.h FOpenGLBaseTextureCube */,
-			1, SourceTextureCube->GetFormat(), true, false, SourceTextureCube->GetFlags(), SourceTextureCube->GetClearBinding());
-	}
-	else
-	{
-		UE_LOG(LogRHI, Error, TEXT("Currently FOpenGLDynamicRHI::RHICreateAliasedTexture only supports 2D, 2D Array and Cube textures."));
-	}
-
-	if (Result != nullptr)
-	{
-		RHIAliasTextureResources(Result, SourceTexture);
-	}
-
-	if (Result)
-	{
-		FOpenGLTextureBase* BaseTexture = GetOpenGLTextureFromRHITexture(Result);
-		// Init memory size to zero, since we're aliased.
-		BaseTexture->SetMemorySize(0);
-	}
-
-	return Result;
+	return new FOpenGLTexture(*GetOpenGLTextureFromRHITexture(SourceTexture), FOpenGLTexture::AliasResource);
 }
 
 void* FOpenGLDynamicRHI::LockTexture2D_RenderThread(class FRHICommandListImmediate& RHICmdList, FRHITexture2D* Texture, uint32 MipIndex, EResourceLockMode LockMode, uint32& DestStride, bool bLockWithinMiptail, bool bNeedsDefaultRHIFlush)
@@ -2852,11 +2349,11 @@ void* FOpenGLDynamicRHI::LockTexture2D_RenderThread(class FRHICommandListImmedia
 		return this->RHILockTexture2D(Texture, MipIndex, LockMode, DestStride, bLockWithinMiptail);
 		RHITHREAD_GLCOMMAND_EPILOGUE_GET_RETURN(void *);
 		Result = ReturnValue;
-		MipBytes = ResourceCast_Unfenced(Texture)->GetLockSize(MipIndex, 0, LockMode, DestStride);
+		MipBytes = GetOpenGLTextureFromRHITexture(Texture)->GetLockSize(MipIndex, 0, LockMode, DestStride);
 	}
 	else
 	{
-		MipBytes = ResourceCast_Unfenced(Texture)->GetLockSize(MipIndex, 0, LockMode, DestStride);
+		MipBytes = GetOpenGLTextureFromRHITexture(Texture)->GetLockSize(MipIndex, 0, LockMode, DestStride);
 		Result = FMemory::Malloc(MipBytes, 16);
 	}
 	check(Result);
@@ -2907,11 +2404,11 @@ void* FOpenGLDynamicRHI::RHILockTextureCubeFace_RenderThread(class FRHICommandLi
 		return this->RHILockTextureCubeFace(Texture, FaceIndex, ArrayIndex, MipIndex, LockMode, DestStride, bLockWithinMiptail);
 		RHITHREAD_GLCOMMAND_EPILOGUE_GET_RETURN(void *);
 		Result = ReturnValue;
-		MipBytes = ResourceCast_Unfenced(Texture)->GetLockSize(MipIndex, 0, LockMode, DestStride);
+		MipBytes = GetOpenGLTextureFromRHITexture(Texture)->GetLockSize(MipIndex, 0, LockMode, DestStride);
 	}
 	else
 	{
-		MipBytes = ResourceCast_Unfenced(Texture)->GetLockSize(MipIndex, 0, LockMode, DestStride);
+		MipBytes = GetOpenGLTextureFromRHITexture(Texture)->GetLockSize(MipIndex, 0, LockMode, DestStride);
 		Result = FMemory::Malloc(MipBytes, 16);
 	}
 	check(Result);
@@ -2961,11 +2458,11 @@ void* FOpenGLDynamicRHI::LockTexture2DArray_RenderThread(class FRHICommandListIm
 		return this->RHILockTexture2DArray(Texture, ArrayIndex, MipIndex, LockMode, DestStride, bLockWithinMiptail);
 		RHITHREAD_GLCOMMAND_EPILOGUE_GET_RETURN(void*);
 		Result = ReturnValue;
-		MipBytes = ResourceCast_Unfenced(Texture)->GetLockSize(MipIndex, ArrayIndex, LockMode, DestStride);
+		MipBytes = GetOpenGLTextureFromRHITexture(Texture)->GetLockSize(MipIndex, ArrayIndex, LockMode, DestStride);
 	}
 	else
 	{
-		MipBytes = ResourceCast_Unfenced(Texture)->GetLockSize(MipIndex, ArrayIndex, LockMode, DestStride);
+		MipBytes = GetOpenGLTextureFromRHITexture(Texture)->GetLockSize(MipIndex, ArrayIndex, LockMode, DestStride);
 		Result = FMemory::Malloc(MipBytes, 16);
 	}
 	check(Result);
@@ -3034,13 +2531,13 @@ void FTextureEvictionLRU::TickEviction()
 		TextureLRU.Num() && (TextureLRU.GetLeastRecent()->EvictionParamsPtr->FrameLastRendered + GOGLTextureEvictFramesToLive) < GFrameNumberRenderThread && EvictCount < GOGLTexturesToEvictPerFrame
 		;EvictCount++)
 	{
-		FOpenGLTextureBase* RemovedFromLRU = TextureLRU.RemoveLeastRecent();
+		FOpenGLTexture* RemovedFromLRU = TextureLRU.RemoveLeastRecent();
 		RemovedFromLRU->EvictionParamsPtr->LRUNode = FSetElementId();
 		RemovedFromLRU->TryEvictGLResource();
 	}
 }
 
-void FTextureEvictionLRU::Remove(FOpenGLTextureBase* TextureBase)
+void FTextureEvictionLRU::Remove(FOpenGLTexture* TextureBase)
 {
 	if( TextureBase->EvictionParamsPtr.IsValid() )
 	{
@@ -3056,7 +2553,7 @@ void FTextureEvictionLRU::Remove(FOpenGLTextureBase* TextureBase)
 	}
 }
 
-bool FTextureEvictionLRU::Add(FOpenGLTextureBase* TextureBase)
+bool FTextureEvictionLRU::Add(FOpenGLTexture* TextureBase)
 {
 	FScopeLock Lock(&TextureLRULock); 
 	check(TextureBase->EvictionParamsPtr);
@@ -3073,7 +2570,7 @@ bool FTextureEvictionLRU::Add(FOpenGLTextureBase* TextureBase)
 	return false;
 }
 
-void FTextureEvictionLRU::Touch(FOpenGLTextureBase* TextureBase)
+void FTextureEvictionLRU::Touch(FOpenGLTexture* TextureBase)
 {
 	FScopeLock Lock(&TextureLRULock);
 	check(TextureBase->EvictionParamsPtr);
@@ -3083,7 +2580,7 @@ void FTextureEvictionLRU::Touch(FOpenGLTextureBase* TextureBase)
 	TextureBase->EvictionParamsPtr->FrameLastRendered = GFrameNumberRenderThread;
 }
 
-FOpenGLTextureBase* FTextureEvictionLRU::GetLeastRecent()
+FOpenGLTexture* FTextureEvictionLRU::GetLeastRecent()
 {
 	return GetLRUContainer().GetLeastRecent();
 }
