Copying //UE5/Dev-ParallelRendering to //UE5/Main

diff --git a/Engine/Source/Runtime/OpenGLDrv/Private/OpenGLUAV.cpp b/Engine/Source/Runtime/OpenGLDrv/Private/OpenGLUAV.cpp
--- a/Engine/Source/Runtime/OpenGLDrv/Private/OpenGLUAV.cpp
+++ b/Engine/Source/Runtime/OpenGLDrv/Private/OpenGLUAV.cpp
@@ -7,12 +7,6 @@
 #include "OpenGLDrvPrivate.h"
 #include "ClearReplacementShaders.h"
 
-FShaderResourceViewRHIRef FOpenGLDynamicRHI::RHICreateShaderResourceView(FRHIBuffer* BufferRHI, uint32 Stride, uint8 Format)
-{
-	ensureMsgf(Stride == GPixelFormats[Format].BlockBytes, TEXT("provided stride: %i was not consitent with Pixelformat: %s"), Stride, GPixelFormats[Format].Name);
-	return FOpenGLDynamicRHI::RHICreateShaderResourceView(FShaderResourceViewInitializer(BufferRHI, EPixelFormat(Format)));
-}
-
 // Binds the specified buffer range to a texture resource and selects glTexBuffer or glTexBufferRange 
 static void BindGLTexBufferRange(GLenum Target, GLenum InternalFormat, GLuint Buffer, uint32 StartOffsetBytes, uint32 NumElements, uint32 Stride)
 {
@@ -35,58 +29,82 @@ static void BindGLTexBufferRange(GLenum Target, GLenum InternalFormat, GLuint Bu
 	}
 }
 
-FShaderResourceViewRHIRef FOpenGLDynamicRHI::RHICreateShaderResourceView(const FShaderResourceViewInitializer& Initializer)
+FOpenGLShaderResourceView::FOpenGLShaderResourceView(const FShaderResourceViewInitializer& Initializer)
 {
 	FShaderResourceViewInitializer::FBufferShaderResourceViewInitializer Desc = Initializer.AsBufferSRV();
-	FRHIBuffer* BufferRHI = Desc.Buffer;
-	if (BufferRHI == nullptr)
+	Buffer = FOpenGLDynamicRHI::ResourceCast(Desc.Buffer);
+	if (Buffer == nullptr)
 	{
-		return new FOpenGLShaderResourceViewProxy([=](FRHIShaderResourceView* OwnerRHI)
-		{
-			return new FOpenGLShaderResourceView(this, 0, GL_TEXTURE_BUFFER);
-		});
+		return;
 	}
 
 	EPixelFormat Format = Desc.Format;
-	
 	if (Initializer.GetType() == FShaderResourceViewInitializer::EType::IndexBufferSRV && Format == PF_Unknown)
-	{		
-		uint32 Stride = BufferRHI->GetStride();
+	{
+		uint32 Stride = Buffer->GetStride();
 		Format = (Stride == 2) ? PF_R16_UINT : PF_R32_UINT;
 	}
 
 	if (Format != PF_Unknown)
 	{
+		Target = GL_TEXTURE_BUFFER;
+		OwnsResource = true;
+
 		UE_CLOG(!GPixelFormats[Format].Supported, LogRHI, Error, TEXT("Unsupported EPixelFormat %d"), Format);
-		
-		return new FOpenGLShaderResourceViewProxy([=, OGLRHI = this](FRHIShaderResourceView* OwnerRHI)
+
+		RunOnGLRenderContextThread([this, Format, StartOffsetBytes = Desc.StartOffsetBytes, NumElements = Desc.NumElements]()
 		{
-			VERIFY_GL_SCOPE();
-			GLuint TextureID = 0;
-			FOpenGL::GenTextures(1, &TextureID);
-			FOpenGLBuffer* GLBuffer = FOpenGLDynamicRHI::ResourceCast(BufferRHI);
-			const uint32 FormatBPP = GPixelFormats[Format].BlockBytes;
-			const FOpenGLTextureFormat& GLFormat = GOpenGLTextureFormats[Format];
+			FOpenGL::GenTextures(1, &Resource);
+
 			// Use a texture stage that's not likely to be used for draws, to avoid waiting
-			OGLRHI->CachedSetupTextureStage(OGLRHI->GetContextStateForCurrentContext(), FOpenGL::GetMaxCombinedTextureImageUnits() - 1, GL_TEXTURE_BUFFER, TextureID, -1, 1);
-			BindGLTexBufferRange(GL_TEXTURE_BUFFER, GLFormat.InternalFormat[0], GLBuffer->Resource, Desc.StartOffsetBytes, Desc.NumElements, FormatBPP);
-			return new FOpenGLShaderResourceView(OGLRHI, TextureID, GL_TEXTURE_BUFFER, BufferRHI, Format);
+			FOpenGLDynamicRHI::Get().CachedSetupTextureStage(
+				FOpenGLDynamicRHI::Get().GetContextStateForCurrentContext(),
+				FOpenGL::GetMaxCombinedTextureImageUnits() - 1,
+				GL_TEXTURE_BUFFER,
+				Resource,
+				-1,
+				1
+			);
+
+			BindGLTexBufferRange(
+				GL_TEXTURE_BUFFER,
+				GOpenGLTextureFormats[Format].InternalFormat[0],
+				Buffer->Resource,
+				StartOffsetBytes,
+				NumElements,
+				GPixelFormats[Format].BlockBytes
+			);
 		});
 	}
 	else
 	{
 		//TODO: add range views for SSBO
 		ensure(Desc.IsWholeResource());
+		Target = GL_SHADER_STORAGE_BUFFER;
 
-		return new FOpenGLShaderResourceViewProxy([=](FRHIShaderResourceView* OwnerRHI)
+		RunOnGLRenderContextThread([this]()
 		{
-			VERIFY_GL_SCOPE();
-			FOpenGLBuffer* BufferGL = ResourceCast(BufferRHI);
-			return new FOpenGLShaderResourceView(this, BufferGL->Resource, GL_SHADER_STORAGE_BUFFER, BufferRHI);
+			Resource = Buffer->Resource;
 		});
 	}
 }
 
+FShaderResourceViewRHIRef FOpenGLDynamicRHI::RHICreateShaderResourceView(const FShaderResourceViewInitializer& Initializer)
+{
+	return new FOpenGLShaderResourceView(Initializer);
+}
+
+FShaderResourceViewRHIRef FOpenGLDynamicRHI::RHICreateShaderResourceView(FRHIBuffer* BufferRHI, uint32 Stride, uint8 Format)
+{
+	ensureMsgf(Stride == GPixelFormats[Format].BlockBytes, TEXT("provided stride: %i was not consitent with Pixelformat: %s"), Stride, GPixelFormats[Format].Name);
+	return FOpenGLDynamicRHI::RHICreateShaderResourceView(FShaderResourceViewInitializer(BufferRHI, EPixelFormat(Format)));
+}
+
+FShaderResourceViewRHIRef FOpenGLDynamicRHI::RHICreateShaderResourceView(FRHIBuffer* BufferRHI)
+{
+	return FOpenGLDynamicRHI::RHICreateShaderResourceView(FShaderResourceViewInitializer(BufferRHI));
+}
+
 void FOpenGLDynamicRHI::RHIUpdateShaderResourceView(FRHIShaderResourceView* SRV, FRHIBuffer* InBuffer)
 {
 	VERIFY_GL_SCOPE();
@@ -102,24 +120,21 @@ void FOpenGLDynamicRHI::RHIUpdateShaderResourceView(FRHIShaderResourceView* SRV,
 		if (!BufferGL)
 		{
 			FOpenGL::TexBuffer(GL_TEXTURE_BUFFER, GL_R16UI, 0); // format ignored here since we're detaching.
-			SRVGL->BufferRHI = nullptr;
-			SRVGL->ModificationVersion = 0;
+			SRVGL->Buffer = nullptr;
 		}
 		else
 		{
-			uint32 Stride = InBuffer->GetStride();
+			uint32 Stride = BufferGL->GetStride();
 			GLenum Format = (Stride == 2) ? GL_R16UI : GL_R32UI;
-			uint32 NumElements = InBuffer->GetSize() / Stride;
+			uint32 NumElements = BufferGL->GetSize() / Stride;
 			BindGLTexBufferRange(GL_TEXTURE_BUFFER, Format, BufferGL->Resource, 0, NumElements, Stride);
-			SRVGL->BufferRHI = InBuffer;
-			SRVGL->ModificationVersion = BufferGL->ModificationCount;
+			SRVGL->Buffer = BufferGL;
 		}
 	}
 	else if (SRVGL->Target == GL_SHADER_STORAGE_BUFFER)
 	{
 		SRVGL->Resource = BufferGL->Resource;
-		SRVGL->BufferRHI = InBuffer;
-		SRVGL->ModificationVersion = BufferGL->ModificationCount;
+		SRVGL->Buffer = BufferGL;
 	}
 	else
 	{
@@ -144,31 +159,23 @@ void FOpenGLDynamicRHI::RHIUpdateShaderResourceView(FRHIShaderResourceView* SRV,
 	if (!BufferGL)
 	{
 		FOpenGL::TexBuffer(GL_TEXTURE_BUFFER, GLFormat.InternalFormat[0], 0);
-		SRVGL->BufferRHI = nullptr;
-		SRVGL->ModificationVersion = 0;
+		SRVGL->Buffer = nullptr;
 	}
 	else
 	{
 		FOpenGL::TexBuffer(GL_TEXTURE_BUFFER, GLFormat.InternalFormat[0], BufferGL->Resource);
-		SRVGL->BufferRHI = InBuffer;
-		SRVGL->Format = Format;
-		SRVGL->ModificationVersion = BufferGL->ModificationCount;
+		SRVGL->Buffer = BufferGL;
 	}
 }
 
-FShaderResourceViewRHIRef FOpenGLDynamicRHI::RHICreateShaderResourceView(FRHIBuffer* BufferRHI)
-{
-	return FOpenGLDynamicRHI::RHICreateShaderResourceView(FShaderResourceViewInitializer(BufferRHI));
-}
-
 FOpenGLShaderResourceView::~FOpenGLShaderResourceView()
 {
 	if (Resource && OwnsResource)
 	{
-		RunOnGLRenderContextThread([OpenGLRHI= OpenGLRHI, Resource = Resource]()
+		RunOnGLRenderContextThread([Resource = Resource]()
 		{
 			VERIFY_GL_SCOPE();
-			OpenGLRHI->InvalidateTextureResourceInCache(Resource);
+			FOpenGLDynamicRHI::Get().InvalidateTextureResourceInCache(Resource);
 			FOpenGL::DeleteTextures(1, &Resource);
 		});
 	}
@@ -187,7 +194,7 @@ FOpenGLTextureUnorderedAccessView::FOpenGLTextureUnorderedAccessView(FRHITexture
 {
 	VERIFY_GL_SCOPE();
 	
-	FOpenGLTextureBase* Texture = GetOpenGLTextureFromRHITexture(TextureRHI);
+	FOpenGLTexture* Texture = GetOpenGLTextureFromRHITexture(TextureRHI);
 	const FOpenGLTextureFormat& GLFormat = GOpenGLTextureFormats[TextureRHI->GetFormat()];
 
 	check(!Texture->CanBeEvicted() && !Texture->IsEvicted());
diff --git a/Engine/Source/Runtime/OpenGLDrv/Private/OpenGLViewport.cpp b/Engine/Source/Runtime/OpenGLDrv/Private/OpenGLViewport.cpp
--- a/Engine/Source/Runtime/OpenGLDrv/Private/OpenGLViewport.cpp
+++ b/Engine/Source/Runtime/OpenGLDrv/Private/OpenGLViewport.cpp
@@ -22,20 +22,20 @@ bool FOpenGLDynamicRHI::RHIGetAvailableResolutions(FScreenResolutionArray& Resol
 	if (Result)
 	{
 		Resolutions.Sort([](const FScreenResolutionRHI& L, const FScreenResolutionRHI& R)
-							{
-								if (L.Width != R.Width)
-								{
-									return L.Width < R.Width;
-								}
-								else if (L.Height != R.Height)
-								{
-									return L.Height < R.Height;
-								}
-								else
-								{
-									return L.RefreshRate < R.RefreshRate;
-								}
-							});
+		{
+			if (L.Width != R.Width)
+			{
+				return L.Width < R.Width;
+			}
+			else if (L.Height != R.Height)
+			{
+				return L.Height < R.Height;
+			}
+			else
+			{
+				return L.RefreshRate < R.RefreshRate;
+			}
+		});
 	}
 	return Result;
 }
@@ -130,7 +130,7 @@ void FOpenGLDynamicRHI::RHIEndDrawingViewport(FRHIViewport* ViewportRHI,bool bPr
 
 	check(DrawingViewport.GetReference() == Viewport);
 
-	FOpenGLTexture2D* BackBuffer = Viewport->GetBackBuffer();
+	FOpenGLTexture* BackBuffer = Viewport->GetBackBuffer();
 
 	FOpenGLContextState& ContextState = GetContextStateForCurrentContext();
 
@@ -226,7 +226,6 @@ FOpenGLViewport::FOpenGLViewport(FOpenGLDynamicRHI* InOpenGLRHI,void* InWindowHa
 	, bIsFullscreen(false)
 	, PixelFormat(PreferredPixelFormat)
 	, bIsValid(true)
-	, FrameSyncEvent(InOpenGLRHI)
 {
 	check(OpenGLRHI);
 	// @todo lumin: Add a "PLATFORM_HAS_NO_NATIVE_WINDOW" or something
@@ -234,14 +233,26 @@ FOpenGLViewport::FOpenGLViewport(FOpenGLDynamicRHI* InOpenGLRHI,void* InWindowHa
 	check(InWindowHandle);
 #endif
 	check(IsInGameThread());
-	PlatformGlGetError();	// flush out old errors.
+
+	// flush out old errors.
+	PlatformGlGetError();	
+
 	OpenGLRHI->Viewports.Add(this);
+
 	check(PlatformOpenGLCurrentContext(OpenGLRHI->PlatformDevice) == CONTEXT_Shared);
+
 	OpenGLContext = PlatformCreateOpenGLContext(OpenGLRHI->PlatformDevice, InWindowHandle);
 	Resize(InSizeX, InSizeY, bInIsFullscreen);
+
 	check(PlatformOpenGLCurrentContext(OpenGLRHI->PlatformDevice) == CONTEXT_Shared);
 
-	BeginInitResource(&FrameSyncEvent);
+	ENQUEUE_RENDER_COMMAND(CreateFrameSyncEvent)([this](FRHICommandListImmediate& RHICmdList)
+	{
+		RunOnGLRenderContextThread([this]()
+		{
+			FrameSyncEvent = MakeUnique<FOpenGLEventQuery>();
+		});
+	});
 }
 
 FOpenGLViewport::~FOpenGLViewport()
@@ -253,44 +264,71 @@ FOpenGLViewport::~FOpenGLViewport()
 		PlatformRestoreDesktopDisplayMode();
 	}
 
-	BeginReleaseResource(&FrameSyncEvent);
-
 	// Release back buffer, before OpenGL context becomes invalid, making it impossible
 	BackBuffer.SafeRelease();
 	check(!IsValidRef(BackBuffer));
 
-	RunOnGLRenderContextThread([&]() {	PlatformDestroyOpenGLContext(OpenGLRHI->PlatformDevice, OpenGLContext); }, true);
+	RunOnGLRenderContextThread([&]()
+	{
+		FrameSyncEvent = nullptr;
+		PlatformDestroyOpenGLContext(OpenGLRHI->PlatformDevice, OpenGLContext);
+	}, true);
+
 	OpenGLContext = NULL;
 	OpenGLRHI->Viewports.Remove(this);
 }
 
+void FOpenGLViewport::WaitForFrameEventCompletion()
+{
+	VERIFY_GL_SCOPE();
+	FrameSyncEvent->WaitForCompletion();
+}
+
+void FOpenGLViewport::IssueFrameEvent()
+{
+	VERIFY_GL_SCOPE();
+	FrameSyncEvent->IssueEvent();
+}
+
 void FOpenGLViewport::Resize(uint32 InSizeX,uint32 InSizeY,bool bInIsFullscreen)
 {
+	check(IsInGameThread());
 	if ((InSizeX == SizeX) && (InSizeY == SizeY) && (bInIsFullscreen == bIsFullscreen))
 	{
 		return;
 	}
 
-	VERIFY_GL_SCOPE();
+	SizeX = InSizeX;
+	SizeY = InSizeY;
+	bool bWasFullscreen = bIsFullscreen;
+	bIsFullscreen = bInIsFullscreen;
 
-	if (IsValidRef(CustomPresent))
+	ENQUEUE_RENDER_COMMAND(ResizeViewport)([this, InSizeX, InSizeY, bInIsFullscreen, bWasFullscreen](FRHICommandListImmediate& RHICmdList)
 	{
-		CustomPresent->OnBackBufferResize();
-	}
-
-	BackBuffer.SafeRelease();	// when the rest of the engine releases it, its framebuffers will be released too (those the engine knows about)
+		if (IsValidRef(CustomPresent))
+		{
+			CustomPresent->OnBackBufferResize();
+		}
 
-	BackBuffer = (FOpenGLTexture2D*)PlatformCreateBuiltinBackBuffer(OpenGLRHI, InSizeX, InSizeY);
-	if (!BackBuffer)
-	{
-		BackBuffer = (FOpenGLTexture2D*)OpenGLRHI->CreateOpenGLTexture(InSizeX, InSizeY, false, false, false, PixelFormat, 1, 1, 1, TexCreate_RenderTargetable, FClearValueBinding::Transparent);
-	}
+		BackBuffer.SafeRelease();	// when the rest of the engine releases it, its framebuffers will be released too (those the engine knows about)
 
-	PlatformResizeGLContext(OpenGLRHI->PlatformDevice, OpenGLContext, InSizeX, InSizeY, bInIsFullscreen, bIsFullscreen, BackBuffer->Target, BackBuffer->GetResource());
+		BackBuffer = PlatformCreateBuiltinBackBuffer(OpenGLRHI, InSizeX, InSizeY);
+		if (!BackBuffer)
+		{
+			BackBuffer = new FOpenGLTexture(FRHITextureCreateDesc::Create2D(
+				TEXT("FOpenGLViewport"),
+				{ (int32)InSizeX, (int32)InSizeY },
+				PixelFormat,
+				FClearValueBinding::Transparent,
+				TexCreate_RenderTargetable
+			));
+		}
 
-	SizeX = InSizeX;
-	SizeY = InSizeY;
-	bIsFullscreen = bInIsFullscreen;
+		RHICmdList.EnqueueLambda([=](FRHICommandListImmediate&)
+		{
+			PlatformResizeGLContext(OpenGLRHI->PlatformDevice, OpenGLContext, InSizeX, InSizeY, bInIsFullscreen, bWasFullscreen, BackBuffer->Target, BackBuffer->GetResource());
+		});
+	});	
 }
 
 void* FOpenGLViewport::GetNativeWindow(void** AddParam) const
diff --git a/Engine/Source/Runtime/OpenGLDrv/Private/Windows/OpenGLWindows.cpp b/Engine/Source/Runtime/OpenGLDrv/Private/Windows/OpenGLWindows.cpp
--- a/Engine/Source/Runtime/OpenGLDrv/Private/Windows/OpenGLWindows.cpp
+++ b/Engine/Source/Runtime/OpenGLDrv/Private/Windows/OpenGLWindows.cpp
@@ -999,17 +999,20 @@ bool PlatformContextIsCurrent( uint64 QueryContext )
 	return (uint64)GetCurrentContext() == QueryContext;
 }
 
-FRHITexture* PlatformCreateBuiltinBackBuffer(FOpenGLDynamicRHI* OpenGLRHI, uint32 SizeX, uint32 SizeY)
+FOpenGLTexture* PlatformCreateBuiltinBackBuffer(FOpenGLDynamicRHI* OpenGLRHI, uint32 SizeX, uint32 SizeY)
 {
 	if (FOpenGL::IsAndroidGLESCompatibilityModeEnabled())
 	{
-		ETextureCreateFlags Flags = TexCreate_RenderTargetable;
-		FOpenGLTexture2D* Texture2D = new FOpenGLTexture2D(OpenGLRHI, 0, GL_RENDERBUFFER, GL_COLOR_ATTACHMENT0, SizeX, SizeY, 0, 1, 1, 1, 1, PF_B8G8R8A8, false, false, Flags, FClearValueBinding::Transparent);
-		OpenGLTextureAllocated(Texture2D, Flags);
-		return Texture2D;
+		return new FOpenGLTexture(FRHITextureCreateDesc::Create2D(
+			TEXT("PlatformCreateBuiltinBackBuffer"),
+			{ (int32)SizeX, (int32)SizeY },
+			PF_B8G8R8A8,
+			FClearValueBinding::Transparent,
+			TexCreate_RenderTargetable | TexCreate_Presentable
+		));
 	}
 
-	return NULL;
+	return nullptr;
 }
 
 void* PlatformGetWindow(FPlatformOpenGLContext* Context, void** AddParam)
diff --git a/Engine/Source/Runtime/OpenGLDrv/Public/OpenGLDrv.h b/Engine/Source/Runtime/OpenGLDrv/Public/OpenGLDrv.h
--- a/Engine/Source/Runtime/OpenGLDrv/Public/OpenGLDrv.h
+++ b/Engine/Source/Runtime/OpenGLDrv/Public/OpenGLDrv.h
@@ -34,9 +34,6 @@
 #define GL_TEXTURE_EXTERNAL_OES 0x8D65
 #endif
 
-#define OPENGL_USE_BINDABLE_UNIFORMS 0
-#define OPENGL_USE_BLIT_FOR_BACK_BUFFER 1
-
 // OpenGL RHI public headers.
 #include "OpenGLUtil.h"
 #include "OpenGLState.h"
@@ -290,7 +287,10 @@ PRAGMA_DISABLE_DEPRECATION_WARNINGS
 /** The interface which is implemented by the dynamically bound RHI. */
 class OPENGLDRV_API FOpenGLDynamicRHI  final : public IOpenGLDynamicRHI, public IRHICommandContextPSOFallback
 {
+	static inline FOpenGLDynamicRHI* Singleton = nullptr;
+
 public:
+	static inline FOpenGLDynamicRHI& Get() { return *Singleton; }
 
 	friend class FOpenGLViewport;
 
@@ -384,18 +384,12 @@ public:
 	virtual FShaderResourceViewRHIRef RHICreateShaderResourceView(const FShaderResourceViewInitializer& Initializer) final override;
 	virtual void RHIUpdateShaderResourceView(FRHIShaderResourceView* SRV, FRHIBuffer* Buffer, uint32 Stride, uint8 Format) final override;
 	virtual void RHIUpdateShaderResourceView(FRHIShaderResourceView* SRV, FRHIBuffer* Buffer) final override;
-	virtual uint64 RHICalcTexture2DPlatformSize(uint32 SizeX, uint32 SizeY, uint8 Format, uint32 NumMips, uint32 NumSamples, ETextureCreateFlags Flags, const FRHIResourceCreateInfo& CreateInfo, uint32& OutAlign) final override;
-	virtual uint64 RHICalcTexture2DArrayPlatformSize(uint32 SizeX, uint32 SizeY, uint32 ArraySize, uint8 Format, uint32 NumMips, uint32 NumSamples, ETextureCreateFlags Flags, const FRHIResourceCreateInfo& CreateInfo, uint32& OutAlign) final override;
-	virtual uint64 RHICalcTexture3DPlatformSize(uint32 SizeX, uint32 SizeY, uint32 SizeZ, uint8 Format, uint32 NumMips, ETextureCreateFlags Flags, const FRHIResourceCreateInfo& CreateInfo, uint32& OutAlign) final override;
-	virtual uint64 RHICalcTextureCubePlatformSize(uint32 Size, uint8 Format, uint32 NumMips, ETextureCreateFlags Flags, const FRHIResourceCreateInfo& CreateInfo, uint32& OutAlign) final override;
+	virtual FRHICalcTextureSizeResult RHICalcTexturePlatformSize(FRHITextureDesc const& Desc, uint32 FirstMipIndex) final override;
 	virtual void RHIGetTextureMemoryStats(FTextureMemoryStats& OutStats) final override;
 	virtual bool RHIGetTextureMemoryVisualizeData(FColor* TextureData,int32 SizeX,int32 SizeY,int32 Pitch,int32 PixelSize) final override;
-	virtual FTexture2DRHIRef RHICreateTexture2D(uint32 SizeX, uint32 SizeY, uint8 Format, uint32 NumMips, uint32 NumSamples, ETextureCreateFlags Flags, ERHIAccess InResourceState, FRHIResourceCreateInfo& CreateInfo) final override;
-	virtual FTexture2DRHIRef RHICreateTextureExternal2D(uint32 SizeX, uint32 SizeY, uint8 Format, uint32 NumMips, uint32 NumSamples, ETextureCreateFlags Flags, ERHIAccess InResourceState, FRHIResourceCreateInfo& CreateInfo) final override;
-	virtual FTexture2DRHIRef RHIAsyncCreateTexture2D(uint32 SizeX, uint32 SizeY, uint8 Format, uint32 NumMips, ETextureCreateFlags Flags, ERHIAccess InResourceState, void** InitialMipData, uint32 NumInitialMips) final override;
+	virtual FTextureRHIRef RHICreateTexture(const FRHITextureCreateDesc& CreateDesc) final override;
+	virtual FTextureRHIRef RHIAsyncCreateTexture2D(uint32 SizeX, uint32 SizeY, uint8 Format, uint32 NumMips, ETextureCreateFlags Flags, ERHIAccess InResourceState, void** InitialMipData, uint32 NumInitialMips) final override;
 	virtual void RHICopySharedMips(FRHITexture2D* DestTexture2D, FRHITexture2D* SrcTexture2D) final override;
-	virtual FTexture2DArrayRHIRef RHICreateTexture2DArray(uint32 SizeX, uint32 SizeY, uint32 SizeZ, uint8 Format, uint32 NumMips, uint32 NumSamples, ETextureCreateFlags Flags, ERHIAccess InResourceState, FRHIResourceCreateInfo& CreateInfo) final override;
-	virtual FTexture3DRHIRef RHICreateTexture3D(uint32 SizeX, uint32 SizeY, uint32 SizeZ, uint8 Format, uint32 NumMips, ETextureCreateFlags Flags, ERHIAccess InResourceState, FRHIResourceCreateInfo& CreateInfo) final override;
 	virtual FShaderResourceViewRHIRef RHICreateShaderResourceView(FRHITexture* Texture, const FRHITextureSRVCreateInfo& CreateInfo) final override;
 	virtual void RHIGenerateMips(FRHITexture* Texture) final override;
 	virtual uint32 RHIComputeMemorySize(FRHITexture* TextureRHI) final override;
@@ -408,8 +402,6 @@ public:
 	virtual void RHIUnlockTexture2DArray(FRHITexture2DArray* Texture, uint32 TextureIndex, uint32 MipIndex, bool bLockWithinMiptail) final override;
 	virtual void RHIUpdateTexture2D(FRHITexture2D* Texture, uint32 MipIndex, const struct FUpdateTextureRegion2D& UpdateRegion, uint32 SourcePitch, const uint8* SourceData) final override;
 	virtual void RHIUpdateTexture3D(FRHITexture3D* Texture, uint32 MipIndex, const struct FUpdateTextureRegion3D& UpdateRegion, uint32 SourceRowPitch, uint32 SourceDepthPitch, const uint8* SourceData) final override;
-	virtual FTextureCubeRHIRef RHICreateTextureCube(uint32 Size, uint8 Format, uint32 NumMips, ETextureCreateFlags Flags, ERHIAccess InResourceState, FRHIResourceCreateInfo& CreateInfo) final override;
-	virtual FTextureCubeRHIRef RHICreateTextureCubeArray(uint32 Size, uint32 ArraySize, uint8 Format, uint32 NumMips, ETextureCreateFlags Flags, ERHIAccess InResourceState, FRHIResourceCreateInfo& CreateInfo) final override;
 	virtual void* RHILockTextureCubeFace(FRHITextureCube* Texture, uint32 FaceIndex, uint32 ArrayIndex, uint32 MipIndex, EResourceLockMode LockMode, uint32& DestStride, bool bLockWithinMiptail) final override;
 	virtual void RHIUnlockTextureCubeFace(FRHITextureCube* Texture, uint32 FaceIndex, uint32 ArrayIndex, uint32 MipIndex, bool bLockWithinMiptail) final override;
 	virtual void RHIBindDebugLabelName(FRHITexture* Texture, const TCHAR* Name) final override;
@@ -632,15 +624,9 @@ public:
 	void EndRenderQuery_OnThisThread(FOpenGLRenderQuery* Query);
 	void GetRenderQueryResult_OnThisThread(FOpenGLRenderQuery* Query, bool bWait);
 
-
-
 	FOpenGLSamplerState* GetPointSamplerState() const { return (FOpenGLSamplerState*)PointSamplerState.GetReference(); }
 
-	FRHITexture* CreateOpenGLTexture(uint32 SizeX, uint32 SizeY, bool CubeTexture, bool ArrayTexture, bool bIsExternal, uint8 Format, uint32 NumMips, uint32 NumSamples, uint32 ArraySize, ETextureCreateFlags Flags, const FClearValueBinding& InClearValue, FResourceBulkDataInterface* BulkData = NULL);
-	
-	FRHITexture* CreateOpenGLRHITextureOnly(const uint32 SizeX, const uint32 SizeY, const bool bCubeTexture, const bool bArrayTexture, const bool bIsExternal, uint8& Format, uint32& NumMips, uint32& NumSamples, const uint32 ArraySize, ETextureCreateFlags Flags, const FClearValueBinding& InClearValue, FResourceBulkDataInterface* BulkData = nullptr);
-
-	void InitializeGLTexture(FRHITexture* Texture, uint32 SizeX, const uint32 SizeY, const bool bCubeTexture, const bool bArrayTexture, const bool bIsExternal, const uint8 Format, const uint32 NumMips, const uint32 NumSamples, const uint32 ArraySize, const ETextureCreateFlags Flags, const FClearValueBinding& InClearValue, FResourceBulkDataInterface* BulkData = nullptr);
+	void InitializeGLTexture(FOpenGLTexture* Texture, const void* BulkDataPtr, uint64 BulkDataSize);
 
 	void* GetOpenGLCurrentContextHandle();
 
@@ -797,43 +783,10 @@ public:
 	}
 
 
-	virtual FTexture2DRHIRef RHICreateTexture2D_RenderThread(class FRHICommandListImmediate& RHICmdList, uint32 SizeX, uint32 SizeY, uint8 Format, uint32 NumMips, uint32 NumSamples, ETextureCreateFlags Flags, ERHIAccess InResourceState, FRHIResourceCreateInfo& CreateInfo) final override
-	{
-		const bool bCubeTexture = false;
-		const bool bArrayTexture = false;
-		const bool bIsExternal = false;
-		const uint32 ArraySize = 1;
-		FOpenGLTexture2D* Texture2D = (FOpenGLTexture2D*)CreateOpenGLRHITextureOnly(SizeX, SizeY, bCubeTexture, bArrayTexture, bIsExternal, Format, NumMips, NumSamples, ArraySize, Flags, CreateInfo.ClearValueBinding, CreateInfo.BulkData);
-		Texture2D->CreationFence.Reset();
-		RunOnGLRenderContextThread([=]()
-		{
-			// Fill in the GL resources.
-			InitializeGLTexture(Texture2D, SizeX, SizeY, bCubeTexture, bArrayTexture, bIsExternal, Format, NumMips, NumSamples, ArraySize, Flags, CreateInfo.ClearValueBinding, CreateInfo.BulkData);
-			Texture2D->CreationFence.WriteAssertFence();
-		});
-		Texture2D->CreationFence.SetRHIThreadFence();
-		return Texture2D;
-	}
-
-	virtual FTexture2DRHIRef RHICreateTextureExternal2D_RenderThread(class FRHICommandListImmediate& RHICmdList, uint32 SizeX, uint32 SizeY, uint8 Format, uint32 NumMips, uint32 NumSamples, ETextureCreateFlags Flags, ERHIAccess InResourceState, FRHIResourceCreateInfo& CreateInfo) final override
+	virtual FTextureRHIRef RHICreateTexture_RenderThread(class FRHICommandListImmediate& RHICmdList, const FRHITextureCreateDesc& CreateDesc) final override
 	{
-		RHITHREAD_GLCOMMAND_PROLOGUE();
-		return this->RHICreateTextureExternal2D(SizeX, SizeY, Format, NumMips, NumSamples, Flags, InResourceState, CreateInfo);
-		RHITHREAD_GLCOMMAND_EPILOGUE_RETURN(FTexture2DRHIRef);
-	}
-
-	virtual FTexture2DArrayRHIRef RHICreateTexture2DArray_RenderThread(class FRHICommandListImmediate& RHICmdList, uint32 SizeX, uint32 SizeY, uint32 SizeZ, uint8 Format, uint32 NumMips, uint32 NumSamples, ETextureCreateFlags Flags, ERHIAccess InResourceState, FRHIResourceCreateInfo& CreateInfo) final override
-	{
-		RHITHREAD_GLCOMMAND_PROLOGUE();
-		return this->RHICreateTexture2DArray(SizeX, SizeY, SizeZ, Format, NumMips, NumSamples, Flags, InResourceState, CreateInfo);
-		RHITHREAD_GLCOMMAND_EPILOGUE_RETURN(FTexture2DArrayRHIRef);
-	}
-
-	virtual FTexture3DRHIRef RHICreateTexture3D_RenderThread(class FRHICommandListImmediate& RHICmdList, uint32 SizeX, uint32 SizeY, uint32 SizeZ, uint8 Format, uint32 NumMips, ETextureCreateFlags Flags, ERHIAccess InResourceState, FRHIResourceCreateInfo& CreateInfo) final override
-	{
-		RHITHREAD_GLCOMMAND_PROLOGUE();
-		return this->RHICreateTexture3D(SizeX, SizeY, SizeZ, Format, NumMips, Flags, InResourceState, CreateInfo);
-		RHITHREAD_GLCOMMAND_EPILOGUE_RETURN(FTexture3DRHIRef);
+		check(IsInRenderingThread());
+		return this->RHICreateTexture(CreateDesc);
 	}
 
 	virtual FUnorderedAccessViewRHIRef RHICreateUnorderedAccessView_RenderThread(class FRHICommandListImmediate& RHICmdList, FRHIBuffer* Buffer, bool bUseUAVCounter, bool bAppendBuffer) final override
@@ -872,32 +825,6 @@ public:
 		return this->RHICreateShaderResourceView(Buffer);
 	}
 
-	virtual FTextureCubeRHIRef RHICreateTextureCube_RenderThread(class FRHICommandListImmediate& RHICmdList, uint32 Size, uint8 Format, uint32 NumMips, ETextureCreateFlags Flags, ERHIAccess InResourceState, FRHIResourceCreateInfo& CreateInfo) final override
-	{
-		const bool bCubeTexture = true;
-		const bool bArrayTexture = false;
-		const bool bIsExternal = false;
-		const uint32 ArraySize = 1;
-		uint32 NumSamples = 1;
-		FOpenGLTextureCube* TextureCube = (FOpenGLTextureCube*)CreateOpenGLRHITextureOnly(Size, Size, bCubeTexture, bArrayTexture, bIsExternal, Format, NumMips, NumSamples, ArraySize, Flags, CreateInfo.ClearValueBinding, CreateInfo.BulkData);
-		TextureCube->CreationFence.Reset();
-		RunOnGLRenderContextThread([=]()
-		{
-			// Fill in the GL resources.
-			InitializeGLTexture(TextureCube, Size, Size, bCubeTexture, bArrayTexture, bIsExternal, Format, NumMips, NumSamples, ArraySize, Flags, CreateInfo.ClearValueBinding, CreateInfo.BulkData);
-			TextureCube->CreationFence.WriteAssertFence();
-		});
-		TextureCube->CreationFence.SetRHIThreadFence();
-		return TextureCube;
-	}
-
-	virtual FTextureCubeRHIRef RHICreateTextureCubeArray_RenderThread(class FRHICommandListImmediate& RHICmdList, uint32 Size, uint32 ArraySize, uint8 Format, uint32 NumMips, ETextureCreateFlags Flags, ERHIAccess InResourceState, FRHIResourceCreateInfo& CreateInfo) final override
-	{
-		RHITHREAD_GLCOMMAND_PROLOGUE();
-		return this->RHICreateTextureCubeArray(Size, ArraySize, Format, NumMips, Flags, InResourceState, CreateInfo);
-		RHITHREAD_GLCOMMAND_EPILOGUE_RETURN(FTextureCubeRHIRef);
-	}
-
 	virtual FRenderQueryRHIRef RHICreateRenderQuery_RenderThread(class FRHICommandListImmediate& RHICmdList, ERenderQueryType QueryType) final override
 	{
 		return this->RHICreateRenderQuery(QueryType);
@@ -985,7 +912,7 @@ public:
 		return GPUProfilingData;
 	}
 
-	GLuint GetOpenGLFramebuffer(uint32 NumSimultaneousRenderTargets, FOpenGLTextureBase** RenderTargets, const uint32* ArrayIndices, const uint32* MipmapLevels, FOpenGLTextureBase* DepthStencilTarget);
+	GLuint GetOpenGLFramebuffer(uint32 NumSimultaneousRenderTargets, FOpenGLTexture** RenderTargets, const uint32* ArrayIndices, const uint32* MipmapLevels, FOpenGLTexture* DepthStencilTarget);
 	
 private:
 
@@ -1058,7 +985,6 @@ private:
 
 	FOpenGLGPUProfiler GPUProfilingData;
 	friend FOpenGLGPUProfiler;
-//	FOpenGLEventQuery FrameSyncEvent;
 
 	FCriticalSection CustomPresentSection;
 	TRefCountPtr<class FRHICustomPresent> CustomPresent;
@@ -1106,11 +1032,11 @@ private:
 
 public:
 	/** Remember what RHI user wants set on a specific OpenGL texture stage, translating from Stage and TextureIndex for stage pair. */
-	void InternalSetShaderTexture(FOpenGLTextureBase* Texture, FOpenGLShaderResourceView* SRV, GLint TextureIndex, GLenum Target, GLuint Resource, int NumMips, int LimitMip);
+	void InternalSetShaderTexture(FOpenGLTexture* Texture, FOpenGLShaderResourceView* SRV, GLint TextureIndex, GLenum Target, GLuint Resource, int NumMips, int LimitMip);
 	void InternalSetShaderImageUAV(GLint UAVIndex, GLenum Format, GLuint Resource, bool bLayered, GLint Layer, GLenum Access);
 	void InternalSetShaderBufferUAV(GLint UAVIndex, GLuint Resource);
 	void InternalSetSamplerStates(GLint TextureIndex, FOpenGLSamplerState* SamplerState);
-	void InitializeGLTextureInternal(GLuint TextureID, FRHITexture* Texture, uint32 SizeX, const uint32 SizeY, const bool bCubeTexture, const bool bArrayTexture, const bool bIsExternal, const uint8 Format, const uint32 NumMips, const uint32 NumSamples, const uint32 ArraySize, const ETextureCreateFlags Flags, const FClearValueBinding& InClearValue, FResourceBulkDataInterface* BulkData);
+	void InitializeGLTextureInternal(FOpenGLTexture* Texture, void const* BulkDataPtr, uint64 BulkDataSize);
 private:
 
 	void RegisterSharedShaderCodeDelegates();
diff --git a/Engine/Source/Runtime/OpenGLDrv/Public/OpenGLResources.h b/Engine/Source/Runtime/OpenGLDrv/Public/OpenGLResources.h
--- a/Engine/Source/Runtime/OpenGLDrv/Public/OpenGLResources.h
+++ b/Engine/Source/Runtime/OpenGLDrv/Public/OpenGLResources.h
@@ -27,6 +27,7 @@
 
 class FOpenGLDynamicRHI;
 class FOpenGLLinkedProgram;
+class FOpenGLTexture;
 typedef TArray<ANSICHAR> FAnsiCharArray;
 
 
@@ -373,9 +374,6 @@ struct TIsGLProxyObject<TOpenGLShaderProxy<TRHIType, TOGLResourceType>>
 
 typedef void (*BufferBindFunction)( GLenum Type, GLuint Buffer );
 
-template<typename BaseType>
-class TOpenGLTexture;
-
 template <typename BaseType, BufferBindFunction BufBind>
 class TOpenGLBuffer : public BaseType
 {
@@ -415,13 +413,9 @@ public:
 	GLuint Resource;
 	GLenum Type;
 
-	/** Needed on OS X to force a rebind of the texture buffer to the texture name to workaround radr://18379338 */
-	uint64 ModificationCount;
-
 	TOpenGLBuffer()
 		: Resource(0)
 		, Type(0)
-		, ModificationCount(0)
 		, bIsLocked(false)
 		, bIsLockReadOnly(false)
 		, bStreamDraw(false)
@@ -437,7 +431,6 @@ public:
 	: BaseType(InStride,InSize,InUsage)
 	, Resource(0)
 	, Type(InType)
-	, ModificationCount(0)
 	, bIsLocked(false)
 	, bIsLockReadOnly(false)
 	, bStreamDraw(bStreamedDraw)
@@ -782,7 +775,6 @@ public:
 				bLockBufferWasAllocated = false;
 				LockSize = 0;
 			}
-			ModificationCount += (bIsLockReadOnly ? 0 : 1);
 			bIsLocked = false;
 		}
 	}
@@ -797,7 +789,6 @@ public:
 #else
 		LoadData( InOffset, InSize, InData);
 #endif
-		ModificationCount++;
 	}
 
 	bool IsDynamic() const { return EnumHasAnyFlags(this->GetUsage(), BUF_AnyDynamic); }
@@ -1152,49 +1143,10 @@ public:
 	virtual ~FOpenGLBoundShaderState();
 };
 
-
-inline GLenum GetOpenGLTargetFromRHITexture(FRHITexture* Texture)
-{
-	if(!Texture)
-	{
-		return GL_NONE;
-	}
-	else if(Texture->GetTexture2D())
-	{
-		return GL_TEXTURE_2D;
-	}
-	else if(Texture->GetTexture2DArray())
-	{
-		return GL_TEXTURE_2D_ARRAY;
-	}
-	else if(Texture->GetTexture3D())
-	{
-		return GL_TEXTURE_3D;
-	}
-	else if(Texture->GetTextureCube())
-	{
-		return GL_TEXTURE_CUBE_MAP;
-	}
-	else
-	{
-		UE_LOG(LogRHI,Fatal,TEXT("Unknown RHI texture type"));
-		return GL_NONE;
-	}
-}
-
-class OPENGLDRV_API FTextureEvictionInterface
-{
-public:
-	virtual bool CanCreateAsEvicted() = 0;
-	virtual void RestoreEvictedGLResource(bool bAttemptToRetainMips) = 0;
-	virtual bool CanBeEvicted() = 0;
-	virtual void TryEvictGLResource() = 0;
-};
-
 class FTextureEvictionLRU
 {
 private:
-	typedef TPsoLruCache<class FOpenGLTextureBase*, class FOpenGLTextureBase*> FOpenGLTextureLRUContainer;
+	typedef TPsoLruCache<FOpenGLTexture*, FOpenGLTexture*> FOpenGLTextureLRUContainer;
 	FCriticalSection TextureLRULock;
 
 	static FORCEINLINE_DEBUGGABLE FOpenGLTextureLRUContainer& GetLRUContainer()
@@ -1213,11 +1165,11 @@ public:
 	}
 	uint32 Num() const { return GetLRUContainer().Num(); }
 
-	void Remove(class FOpenGLTextureBase* TextureBase);
-	bool Add(class FOpenGLTextureBase* TextureBase);
-	void Touch(class FOpenGLTextureBase* TextureBase);
+	void Remove(FOpenGLTexture* TextureBase);
+	bool Add(FOpenGLTexture* TextureBase);
+	void Touch(FOpenGLTexture* TextureBase);
 	void TickEviction();
-	class FOpenGLTextureBase* GetLeastRecent();
+	FOpenGLTexture* GetLeastRecent();
 };
 class FTextureEvictionParams
 {
@@ -1245,7 +1197,8 @@ public:
 
 	void CloneMipData(const FTextureEvictionParams& Src, uint32 NumMips, int32 SrcOffset, int DstOffset);
 
-	uint32 GetTotalAllocated() const {
+	uint32 GetTotalAllocated() const
+	{
 		uint32 TotalAllocated = 0;
 		for (const auto& MipData : MipImageData)
 		{
@@ -1254,7 +1207,8 @@ public:
 		return TotalAllocated;
 	}
 
-	bool AreAllMipsPresent() const {
+	bool AreAllMipsPresent() const
+	{
 		bool bRet = MipImageData.Num() > 0;
 		for (const auto& MipData : MipImageData)
 		{
@@ -1264,38 +1218,63 @@ public:
 	}
 };
 
-extern uint32 GTotalMipRestores;
-class OPENGLDRV_API FOpenGLTextureBase : public FTextureEvictionInterface
+class FOpenGLTextureDesc
 {
-protected:
-	// storing this as static as we can be in the >10,000s instances range.
-	static class FOpenGLDynamicRHI* OpenGLRHI;
-
 public:
-	// Pointer to current sampler state in this unit
-	class FOpenGLSamplerState* SamplerState;
+	FOpenGLTextureDesc(FRHITextureDesc const& InDesc);
+
+	GLenum Target = GL_NONE;
+	GLenum Attachment = GL_NONE;
+
+	uint8 NumSamplesRendered = 1;
+	uint8 NumSamplesStored   = 1;
+
+	uint32 MemorySize = 0;
+
+	uint8 bCubemap            : 1;
+	uint8 bArrayTexture       : 1;
+	uint8 bStreamable         : 1;
+	uint8 bDepthStencil       : 1;
+	uint8 bCanCreateAsEvicted : 1;
+	uint8 bIsPowerOfTwo       : 1;
+	uint8 bTileMemDepthBuffer : 1;
 
 private:
-	/** The OpenGL texture resource. */
-	GLuint Resource;
+	static bool CanDeferTextureCreation();
+};
 
-	void TryRestoreGLResource()
+class FOpenGLTextureCreateDesc : public FRHITextureCreateDesc, public FOpenGLTextureDesc
+{
+public:
+	FOpenGLTextureCreateDesc(FRHITextureCreateDesc const& CreateDesc)
+		: FRHITextureCreateDesc(CreateDesc)
+		, FOpenGLTextureDesc(CreateDesc)
 	{
-		if (EvictionParamsPtr.IsValid() && !EvictionParamsPtr->bHasRestored)
-		{
-			VERIFY_GL_SCOPE();
-			if (!EvictionParamsPtr->bHasRestored)
-			{
-				RestoreEvictedGLResource(true);
-			}
-			else 
-			{
-				check(CanBeEvicted());
-				FTextureEvictionLRU::Get().Touch(this);
-			}
-		}
+		// Tiled GPUs modify the number of samples the FRHITexture reports.
+		// Ensure the base descriptor is updated with the actual number of samples.
+		FRHITextureCreateDesc::NumSamples = FOpenGLTextureDesc::NumSamplesStored;
 	}
+};
+
+class OPENGLDRV_API FOpenGLTexture : public FRHITexture
+{
+	// Prevent copying
+	FOpenGLTexture(FOpenGLTexture const&) = delete;
+	FOpenGLTexture& operator = (FOpenGLTexture const&) = delete;
+
 public:
+	// Standard constructor.
+	explicit FOpenGLTexture(FOpenGLTextureCreateDesc const& CreateDesc);
+
+	// Constructor for external resources (RHICreateTexture2DFromResource etc).
+	explicit FOpenGLTexture(FOpenGLTextureCreateDesc const& CreateDesc, GLuint Resource);
+
+	// Constructor for RHICreateAliasedTexture
+	enum EAliasConstructorParam { AliasResource };
+	explicit FOpenGLTexture(FOpenGLTexture& Texture, EAliasConstructorParam);
+	void AliasResources(FOpenGLTexture& Texture);
+
+	virtual ~FOpenGLTexture();
 
 	GLuint GetResource()
 	{
@@ -1330,243 +1309,37 @@ public:
 		Resource = InResource;
 	}
 
-	/** The OpenGL texture target. */
-	GLenum Target;
-
-	/** The number of mips in the texture. */
-	uint32 NumMips;
-
-	/** The OpenGL attachment point. This should always be GL_COLOR_ATTACHMENT0 in case of color buffer, but the actual texture may be attached on other color attachments. */
-	GLenum Attachment;
-
-	/** OpenGL 3 Stencil/SRV workaround texture resource */
-	GLuint SRVResource;
-
-	/** Initialization constructor. */
-	FOpenGLTextureBase(
-		FOpenGLDynamicRHI* InOpenGLRHI,
-		GLuint InResource,
-		GLenum InTarget,
-		uint32 InNumMips,
-		GLenum InAttachment
-		)
-	: SamplerState(nullptr)
-	, Resource(InResource)
-	, Target(InTarget)
-	, NumMips(InNumMips)
-	, Attachment(InAttachment)
-	, SRVResource( 0 )
-	, MemorySize( 0 )
-	, bIsPowerOfTwo(false)
-	, bIsAliased(false)
-	, bMemorySizeReady(false)
-	{
-		check(OpenGLRHI == nullptr || OpenGLRHI == InOpenGLRHI);
-		OpenGLRHI = InOpenGLRHI;
-	}
-
-	virtual ~FOpenGLTextureBase()
-	{
-		FTextureEvictionLRU::Get().Remove(this);
-
-		if (EvictionParamsPtr.IsValid())
-		{
-			RunOnGLRenderContextThread([EvictionParamsPtr = MoveTemp(EvictionParamsPtr)]() {
-				// EvictionParamsPtr is deleted on RHIT after this.
-			});
-		}
-	}
-
-	int32 GetMemorySize() const
-	{
-		check(bMemorySizeReady);
-		return MemorySize;
-	}
-
-	void SetMemorySize(uint32 InMemorySize)
-	{
-		check(!bMemorySizeReady);
-		MemorySize = InMemorySize;
-		bMemorySizeReady = true;
-	}
-
-	bool IsMemorySizeSet()
-	{
-		return bMemorySizeReady;
-	}
-
-	void SetIsPowerOfTwo(bool bInIsPowerOfTwo)
-	{
-		bIsPowerOfTwo  = bInIsPowerOfTwo ? 1 : 0;
-	}
-
-	bool IsPowerOfTwo() const
-	{
-		return bIsPowerOfTwo != 0;
-	}
-
-	void SetAliased(const bool bInAliased)
-	{
-		bIsAliased = bInAliased ? 1 : 0;
-	}
+	// Returns true if this texture will perform automatic MSAA resolve on tile-based GPUs.
+	bool IsTiledMSAA() const
+    {
+		return NumSamplesInternal > 1 && NumSamplesInternal <= FOpenGL::GetMaxMSAASamplesTileMem();
+    }
 
-	bool IsAliased() const
+	// Returns the number of samples used during rendering. These samples may be kept in on-chip GPU tiled memory.
+	// For non-tiled GPUs, this is always equal to the number of stored samples.
+	uint32 GetNumSamplesRendered() const
 	{
-		return bIsAliased != 0;
+		return NumSamplesInternal;
 	}
 
-	void AliasResources(class FOpenGLTextureBase* Texture)
+	// Returns the number of samples stored in the actual GPU resource memory.
+	uint32 GetNumSamplesStored() const
 	{
-		VERIFY_GL_SCOPE();
-		// restore the source texture, do not allow the texture to become evicted, the aliasing texture cannot re-create the resource.
-		if (Texture->IsEvicted())
-		{
-			Texture->RestoreEvictedGLResource(false);
-		}
-		Resource = Texture->Resource;
-		SRVResource = Texture->SRVResource;
-		bIsAliased = 1;
+		return IsTiledMSAA() ? 1 : NumSamplesInternal;
 	}
 
-	TUniquePtr<FTextureEvictionParams> EvictionParamsPtr;
-	FOpenGLAssertRHIThreadFence CreationFence;
-	
 	bool IsEvicted() const { VERIFY_GL_SCOPE(); return EvictionParamsPtr.IsValid() && !EvictionParamsPtr->bHasRestored; }
-private:
-	uint32 MemorySize		: 30;
-	uint32 bIsPowerOfTwo	: 1;
-	uint32 bIsAliased : 1;
-	uint32 bMemorySizeReady : 1;
-};
-
-// Textures.
-template<typename BaseType>
-class OPENGLDRV_API TOpenGLTexture : public BaseType, public FOpenGLTextureBase
-{
-public:
-
-	/** Initialization constructor. */
-	TOpenGLTexture(
-		class FOpenGLDynamicRHI* InOpenGLRHI,
-		GLuint InResource,
-		GLenum InTarget,
-		GLenum InAttachment,
-		uint32 InSizeX,
-		uint32 InSizeY,
-		uint32 InSizeZ,
-		uint32 InNumMips,
-		uint32 InNumSamples,
-		uint32 InNumSamplesTileMem, /* For render targets on Android tiled GPUs, the number of samples to use internally */
-		uint32 InArraySize,
-		EPixelFormat InFormat,
-		bool bInCubemap,
-		bool bInAllocatedStorage,
-		ETextureCreateFlags InFlags,
-		const FClearValueBinding& InClearValue
-		)
-	: BaseType(InSizeX,InSizeY,InSizeZ,InNumMips,InNumSamples, InNumSamplesTileMem, InArraySize, InFormat,InFlags, InClearValue)
-	, FOpenGLTextureBase(
-		InOpenGLRHI,
-		InResource,
-		InTarget,
-		InNumMips,
-		InAttachment
-		)
-	, BaseLevel(0)
-	, bCubemap(bInCubemap)
-	{
-		PixelBuffers.AddZeroed(this->GetNumMips() * (bCubemap ? 6 : 1) * GetEffectiveSizeZ());
-		SetAllocatedStorage(bInAllocatedStorage);
-	}
-
-private:
-	void DeleteGLResource()
-	{
-		auto DeleteGLResources = [OpenGLRHI = this->OpenGLRHI, Resource = this->GetRawResourceName(), SRVResource = this->SRVResource, Target = this->Target, Flags = this->GetFlags(), Aliased = this->IsAliased()]()
-		{
-			QUICK_SCOPE_CYCLE_COUNTER(STAT_OpenGLDeleteGLTextureTime);
-			VERIFY_GL_SCOPE();
-			if (Resource != 0)
-			{
-				switch (Target)
-				{
-					case GL_TEXTURE_2D:
-					case GL_TEXTURE_2D_MULTISAMPLE:
-					case GL_TEXTURE_3D:
-					case GL_TEXTURE_CUBE_MAP:
-					case GL_TEXTURE_2D_ARRAY:
-					case GL_TEXTURE_CUBE_MAP_ARRAY:
-	#if PLATFORM_ANDROID
-					case GL_TEXTURE_EXTERNAL_OES:
-	#endif
-					{
-						OpenGLRHI->InvalidateTextureResourceInCache(Resource);
-						if (SRVResource)
-						{
-							OpenGLRHI->InvalidateTextureResourceInCache(SRVResource);
-						}
-
-						if (!Aliased)
-						{
-							FOpenGL::DeleteTextures(1, &Resource);
-							if (SRVResource)
-							{
-								FOpenGL::DeleteTextures(1, &SRVResource);
-							}
-						}
-						break;
-					}
-					case GL_RENDERBUFFER:
-					{
-						if (!(Flags & TexCreate_Presentable))
-						{
-							glDeleteRenderbuffers(1, &Resource);
-						}
-						break;
-					}
-					default:
-					{
-						checkNoEntry();
-					}
-				}
-			}
-		};
-
-		RunOnGLRenderContextThread(MoveTemp(DeleteGLResources));
-	}
-
-public:
-
-	virtual ~TOpenGLTexture()
-	{
-		if (GIsRHIInitialized)
-		{
-			if (IsInActualRenderingThread())
-			{
-				this->CreationFence.WaitFenceRenderThreadOnly();
-			}
-
-			if(!CanCreateAsEvicted())
-			{
-				// TODO: this should run on the RHIT now.
-				ReleaseOpenGLFramebuffers(this->OpenGLRHI, this);
-			}
-
-			DeleteGLResource();
-			OpenGLTextureDeleted(this);
-		}
-	}
 
 	virtual void* GetTextureBaseRHI() override final
 	{
-		return static_cast<FOpenGLTextureBase*>(this);
+		return this;
 	}
 
 	/**
 	 * Locks one of the texture's mip-maps.
 	 * @return A pointer to the specified texture data.
 	 */
-	void* Lock(uint32 MipIndex,uint32 ArrayIndex,EResourceLockMode LockMode,uint32& DestStride);
+	void* Lock(uint32 MipIndex, uint32 ArrayIndex, EResourceLockMode LockMode, uint32& DestStride);
 
 	/**
 	* Returns the size of the memory block that is returned from Lock, threadsafe
@@ -1574,19 +1347,13 @@ public:
 	uint32 GetLockSize(uint32 MipIndex, uint32 ArrayIndex, EResourceLockMode LockMode, uint32& DestStride);
 
 	/** Unlocks a previously locked mip-map. */
-	void Unlock(uint32 MipIndex,uint32 ArrayIndex);
-
-	// Accessors.
-	bool IsDynamic() const { return (this->GetFlags() & TexCreate_Dynamic) != 0; }
-	bool IsCubemap() const { return bCubemap != 0; }
-	bool IsStaging() const { return EnumHasAnyFlags(this->GetFlags(), TexCreate_CPUReadback); }
-
+	void Unlock(uint32 MipIndex, uint32 ArrayIndex);
 
 	/** FRHITexture override.  See FRHITexture::GetNativeResource() */
 	virtual void* GetNativeResource() const override
 	{
 		// this must become a full GL resource here, calling the non-const GetResourceRef ensures this.
-		return const_cast<void*>(reinterpret_cast<const void*>(&const_cast<TOpenGLTexture*>(this)->GetResourceRef()));
+		return const_cast<void*>(reinterpret_cast<const void*>(&const_cast<FOpenGLTexture*>(this)->GetResourceRef()));
 	}
 
 	/**
@@ -1611,40 +1378,93 @@ public:
 	/**
 	 * Clone texture from a source using CopyImageSubData
 	 */
-	void CloneViaCopyImage( TOpenGLTexture* Src, uint32 InNumMips, int32 SrcOffset, int32 DstOffset);
+	void CloneViaCopyImage(FOpenGLTexture* Src, uint32 InNumMips, int32 SrcOffset, int32 DstOffset);
 
 	/**
 	 * Clone texture from a source going via PBOs
 	 */
-	void CloneViaPBO( TOpenGLTexture* Src, uint32 InNumMips, int32 SrcOffset, int32 DstOffset);
+	void CloneViaPBO(FOpenGLTexture* Src, uint32 InNumMips, int32 SrcOffset, int32 DstOffset);
 
 	/**
 	 * Resolved the specified face for a read Lock, for non-renderable, CPU readable surfaces this eliminates the readback inside Lock itself.
 	 */
-	void Resolve(uint32 MipIndex,uint32 ArrayIndex);
+	void Resolve(uint32 MipIndex, uint32 ArrayIndex);
+
+	// Texture eviction
+	void RestoreEvictedGLResource(bool bAttemptToRetainMips);
+	bool CanBeEvicted();
+	void TryEvictGLResource();
 
-	/*
-	 * FTextureEvictionInterface
-	 */
-	virtual void RestoreEvictedGLResource(bool bAttemptToRetainMips) override;
-	virtual bool CanCreateAsEvicted() override;
-	virtual bool CanBeEvicted() override;
-	virtual void TryEvictGLResource() override;
 private:
+	static void UpdateTextureStats(FOpenGLTexture* Texture, bool bAllocating);
+
+	void TryRestoreGLResource()
+	{
+		if (EvictionParamsPtr.IsValid() && !EvictionParamsPtr->bHasRestored)
+		{
+			VERIFY_GL_SCOPE();
+			if (!EvictionParamsPtr->bHasRestored)
+			{
+				RestoreEvictedGLResource(true);
+			}
+			else
+			{
+				check(CanBeEvicted());
+				FTextureEvictionLRU::Get().Touch(this);
+			}
+		}
+	}
+
+	void DeleteGLResource();
+
 	void Fill2DGLTextureImage(const FOpenGLTextureFormat& GLFormat, const bool bSRGB, uint32 MipIndex, const void* LockedBuffer, uint32 LockedSize, uint32 ArrayIndex);
 
-	TArray< TRefCountPtr<FOpenGLPixelBuffer> > PixelBuffers;
+	uint32 GetEffectiveSizeZ() const
+	{
+		FRHITextureDesc const& Desc = GetDesc();
+
+		return Desc.IsTexture3D()
+			? Desc.Depth
+			: Desc.ArraySize;
+	}
+
+private:
+	/** The OpenGL texture resource. */
+	GLuint Resource = GL_NONE;
+
+public:
+	/** The OpenGL texture target. */
+	GLenum const Target = 0;
+
+	/** The OpenGL attachment point. This should always be GL_COLOR_ATTACHMENT0 in case of color buffer, but the actual texture may be attached on other color attachments. */
+	GLenum const Attachment = 0;
 
-	uint32 GetEffectiveSizeZ( void ) { return this->GetSizeZ() ? this->GetSizeZ() : 1; }
+	TUniquePtr<FTextureEvictionParams> EvictionParamsPtr;
+
+	// Pointer to current sampler state in this unit
+	class FOpenGLSamplerState* SamplerState = nullptr;
 
-	/** Index of the largest mip-map in the texture */
-	uint32 BaseLevel;
+private:
+	TArray<TRefCountPtr<FOpenGLPixelBuffer>> PixelBuffers;
 
 	/** Bitfields marking whether we have allocated storage for each mip */
 	TBitArray<TInlineAllocator<1> > bAllocatedStorage;
 
-	/** Whether the texture is a cube-map. */
-	const uint32 bCubemap : 1;
+public:
+	uint32 const MemorySize;
+
+private:
+	uint8 const NumSamplesInternal;
+
+public:
+	uint8 const bIsPowerOfTwo       : 1;
+	uint8 const bCanCreateAsEvicted : 1;
+	uint8 const bStreamable         : 1;
+	uint8 const bCubemap            : 1;
+	uint8 const bArrayTexture       : 1;
+	uint8 const bDepthStencil       : 1;
+	uint8 const bTileMemDepthBuffer : 1;
+	uint8 const bAlias              : 1;
 };
 
 template <typename T>
@@ -1653,7 +1473,7 @@ struct TIsGLResourceWithFence
 	enum
 	{
 		Value = TOr<
-		TPointerIsConvertibleFromTo<T, const FOpenGLTextureBase>
+		TPointerIsConvertibleFromTo<T, const FOpenGLTexture>
 		//		,TIsDerivedFrom<T, FRHITexture>
 		>::Value
 	};
@@ -1668,172 +1488,16 @@ static typename TEnableIf<TIsGLResourceWithFence<T>::Value>::Type CheckRHITFence
 	Resource->CreationFence.WaitFenceRenderThreadOnly();
 }
 
-class OPENGLDRV_API FOpenGLBaseTexture2D : public FRHITexture2D
-{
-public:
-	FOpenGLBaseTexture2D(uint32 InSizeX, uint32 InSizeY, uint32 InSizeZ, uint32 InNumMips, uint32 InNumSamples, uint32 InNumSamplesTileMem, uint32 InArraySize, EPixelFormat InFormat, ETextureCreateFlags InFlags, const FClearValueBinding& InClearValue)
-	: FRHITexture2D(InSizeX,InSizeY,InNumMips,InNumSamples,InFormat,InFlags, InClearValue)
-	, SampleCount(InNumSamples)
-	, SampleCountTileMem(InNumSamplesTileMem)
-	{}
-	uint32 GetSizeZ() const { return 0; }
-	uint32 GetNumSamples() const { return SampleCount; }
-	uint32 GetNumSamplesTileMem() const { return SampleCountTileMem; }
-private:
-	uint32 SampleCount;
-	/* For render targets on Android tiled GPUs, the number of samples to use internally */
-	uint32 SampleCountTileMem;
-};
-
-class FOpenGLBaseTexture2DArray : public FRHITexture2DArray
-{
-public:
-	FOpenGLBaseTexture2DArray(uint32 InSizeX, uint32 InSizeY, uint32 InSizeZ, uint32 InNumMips, uint32 InNumSamples, uint32 InNumSamplesTileMem, uint32 InArraySize, EPixelFormat InFormat, ETextureCreateFlags InFlags, const FClearValueBinding& InClearValue)
-	: FRHITexture2DArray(InSizeX,InSizeY,InSizeZ,InNumMips,InNumSamples,InFormat,InFlags, InClearValue)
-	{
-		check(InNumSamples == 1);	// OpenGL supports multisampled texture arrays, but they're currently not implemented in OpenGLDrv.
-		check(InNumSamplesTileMem == 1);
-	}
-};
-
-class FOpenGLBaseTextureCube : public FRHITextureCube
-{
-public:
-	FOpenGLBaseTextureCube(uint32 InSizeX, uint32 InSizeY, uint32 InSizeZ, uint32 InNumMips, uint32 InNumSamples, uint32 InNumSamplesTileMem, uint32 InArraySize, EPixelFormat InFormat, ETextureCreateFlags InFlags, const FClearValueBinding& InClearValue)
-	: FRHITextureCube(InSizeX,InNumMips,InFormat,InFlags,InClearValue)
-	, ArraySize(InArraySize)
-	{
-		check(InNumSamples == 1);	// OpenGL doesn't currently support multisampled cube textures
-		check(InNumSamplesTileMem == 1);
-	}
-	uint32 GetSizeX() const { return GetSize(); }
-	uint32 GetSizeY() const { return GetSize(); } //-V524
-	uint32 GetSizeZ() const { return ArraySize > 1 ? ArraySize : 0; }
-
-	uint32 GetArraySize() const {return ArraySize;}
-private:
-	uint32 ArraySize;
-};
-
-class FOpenGLBaseTexture3D : public FRHITexture3D
-{
-public:
-	FOpenGLBaseTexture3D(uint32 InSizeX, uint32 InSizeY, uint32 InSizeZ, uint32 InNumMips, uint32 InNumSamples, uint32 InNumSamplesTileMem, uint32 InArraySize, EPixelFormat InFormat, ETextureCreateFlags InFlags, const FClearValueBinding& InClearValue)
-	: FRHITexture3D(InSizeX,InSizeY,InSizeZ,InNumMips,InFormat,InFlags,InClearValue)
-	{
-		check(InNumSamples == 1);	// Can't have multisampled texture 3D. Not supported anywhere.
-		check(InNumSamplesTileMem == 1);
-	}
-};
-
-typedef TOpenGLTexture<FOpenGLBaseTexture2D>			FOpenGLTexture2D;
-typedef TOpenGLTexture<FOpenGLBaseTexture2DArray>		FOpenGLTexture2DArray;
-typedef TOpenGLTexture<FOpenGLBaseTexture3D>			FOpenGLTexture3D;
-typedef TOpenGLTexture<FOpenGLBaseTextureCube>			FOpenGLTextureCube;
-
-/** Given a pointer to a RHI texture that was created by the OpenGL RHI, returns a pointer to the FOpenGLTextureBase it encapsulates. */
-inline FOpenGLTextureBase* GetOpenGLTextureFromRHITexture(FRHITexture* Texture)
-{
-	if(!Texture)
-	{
-		return NULL;
-	}
-	else
-	{
-		CheckRHITFence(static_cast<FOpenGLTextureBase*>(Texture->GetTextureBaseRHI()));
-		return static_cast<FOpenGLTextureBase*>(Texture->GetTextureBaseRHI());
-	}
-}
-
-inline uint32 GetOpenGLTextureSizeXFromRHITexture(FRHITexture* Texture)
-{
-	if(!Texture)
-	{
-		return 0;
-	}
-	CheckRHITFence(static_cast<FOpenGLTextureBase*>(Texture->GetTextureBaseRHI()));
-	if(Texture->GetTexture2D())
-	{
-		return ((FOpenGLTexture2D*)Texture)->GetSizeX();
-	}
-	else if(Texture->GetTexture2DArray())
-	{
-		return ((FOpenGLTexture2DArray*)Texture)->GetSizeX();
-	}
-	else if(Texture->GetTexture3D())
-	{
-		return ((FOpenGLTexture3D*)Texture)->GetSizeX();
-	}
-	else if(Texture->GetTextureCube())
-	{
-		return ((FOpenGLTextureCube*)Texture)->GetSize();
-	}
-	else
-	{
-		UE_LOG(LogRHI,Fatal,TEXT("Unknown RHI texture type"));
-		return 0;
-	}
-}
-
-inline uint32 GetOpenGLTextureSizeYFromRHITexture(FRHITexture* Texture)
-{
-	if(!Texture)
-	{
-		return 0;
-	}
-
-	CheckRHITFence(static_cast<FOpenGLTextureBase*>(Texture->GetTextureBaseRHI()));
-	if(Texture->GetTexture2D())
-	{
-		return ((FOpenGLTexture2D*)Texture)->GetSizeY();
-	}
-	else if(Texture->GetTexture2DArray())
-	{
-		return ((FOpenGLTexture2DArray*)Texture)->GetSizeY();
-	}
-	else if(Texture->GetTexture3D())
-	{
-		return ((FOpenGLTexture3D*)Texture)->GetSizeY();
-	}
-	else if(Texture->GetTextureCube())
-	{
-		return ((FOpenGLTextureCube*)Texture)->GetSize();
-	}
-	else
-	{
-		UE_LOG(LogRHI,Fatal,TEXT("Unknown RHI texture type"));
-		return 0;
-	}
-}
-
-inline uint32 GetOpenGLTextureSizeZFromRHITexture(FRHITexture* Texture)
+/** Given a pointer to a RHI texture that was created by the OpenGL RHI, returns a pointer to the FOpenGLTexture it encapsulates. */
+inline FOpenGLTexture* GetOpenGLTextureFromRHITexture(FRHITexture* TextureRHI)
 {
-	if(!Texture)
-	{
-		return 0;
-	}
-
-	CheckRHITFence(Texture);
-	if(Texture->GetTexture2D())
-	{
-		return 0;
-	}
-	else if(Texture->GetTexture2DArray())
+	if (!TextureRHI)
 	{
-		return ((FOpenGLTexture2DArray*)Texture)->GetSizeZ();
-	}
-	else if(Texture->GetTexture3D())
-	{
-		return ((FOpenGLTexture3D*)Texture)->GetSizeZ();
-	}
-	else if(Texture->GetTextureCube())
-	{
-		return ((FOpenGLTextureCube*)Texture)->GetSizeZ();
+		return nullptr;
 	}
 	else
 	{
-		UE_LOG(LogRHI,Fatal,TEXT("Unknown RHI texture type"));
-		return 0;
+		return static_cast<FOpenGLTexture*>(TextureRHI->GetTextureBaseRHI());
 	}
 }
 
@@ -1954,124 +1618,38 @@ public:
 	virtual uint32 GetBufferSize() override;
 };
 
-class FOpenGLShaderResourceView : public FRefCountedObject
+class FOpenGLShaderResourceView : public FRHIShaderResourceView
 {
 public:
+	explicit FOpenGLShaderResourceView() = default;
+	explicit FOpenGLShaderResourceView(const FShaderResourceViewInitializer& Initializer);
+	explicit FOpenGLShaderResourceView(FOpenGLTexture* Texture, const FRHITextureSRVCreateInfo& CreateInfo);
+
+	virtual ~FOpenGLShaderResourceView();
 
 	/** OpenGL texture the buffer is bound with */
-	GLuint Resource;
-	GLenum Target;
+	GLuint Resource = GL_NONE;
+	GLenum Target = GL_TEXTURE_BUFFER;
 
 	/** Needed on GL <= 4.2 to copy stencil data out of combined depth-stencil surfaces. */
-	FTexture2DRHIRef Texture2D;
-
-	int32 LimitMip;
-
-	/** Needed on OS X to force a rebind of the texture buffer to the texture name to workaround radr://18379338 */
-	FBufferRHIRef BufferRHI;
-	uint64 ModificationVersion;
-	uint8 Format;
-
-	FOpenGLShaderResourceView( FOpenGLDynamicRHI* InOpenGLRHI, GLuint InResource, GLenum InTarget )
-	:	Resource(InResource)
-	,	Target(InTarget)
-	,	LimitMip(-1)
-	,	ModificationVersion(0)
-	,	Format(0)
-	,	OpenGLRHI(InOpenGLRHI)
-	,	OwnsResource(InTarget != GL_SHADER_STORAGE_BUFFER)
-	{}
-
-	FOpenGLShaderResourceView(FOpenGLDynamicRHI* InOpenGLRHI, GLuint InResource, GLenum InTarget, FRHIBuffer* InBuffer)
-		: Resource(InResource)
-		, Target(InTarget)
-		, LimitMip(-1)
-		, BufferRHI(InBuffer)
-		, ModificationVersion(0)
-		, Format(0)
-		, OpenGLRHI(InOpenGLRHI)
-		, OwnsResource(InTarget != GL_SHADER_STORAGE_BUFFER)
-	{
-		if (BufferRHI)
-		{
-			FOpenGLBuffer* GLBuffer = (FOpenGLBuffer*)BufferRHI.GetReference();
-			ModificationVersion = GLBuffer->ModificationCount;
-		}
-	}
-
-	FOpenGLShaderResourceView( FOpenGLDynamicRHI* InOpenGLRHI, GLuint InResource, GLenum InTarget, FRHIBuffer* InBuffer, uint8 InFormat )
-	:	Resource(InResource)
-	,	Target(InTarget)
-	,	LimitMip(-1)
-	,	BufferRHI(InBuffer)
-	,	ModificationVersion(0)
-	,	Format(InFormat)
-	,	OpenGLRHI(InOpenGLRHI)
-	,	OwnsResource(true)
-	{
-		check(InTarget != GL_SHADER_STORAGE_BUFFER);
-		if (BufferRHI)
-		{
-			FOpenGLBuffer* GLBuffer = (FOpenGLBuffer*)BufferRHI.GetReference();
-			ModificationVersion = GLBuffer->ModificationCount;
-		}
-	}
-
-	FOpenGLShaderResourceView( FOpenGLDynamicRHI* InOpenGLRHI, GLuint InResource, GLenum InTarget, GLuint Mip, bool InOwnsResource )
-	:	Resource(InResource)
-	,	Target(InTarget)
-	,	LimitMip(Mip)
-	,	ModificationVersion(0)
-	,	Format(0)
-	,	OpenGLRHI(InOpenGLRHI)
-	,	OwnsResource(InOwnsResource)
-	{
-		check(InTarget != GL_SHADER_STORAGE_BUFFER);
-	}
+	TRefCountPtr<FOpenGLTexture> Texture;
+	TRefCountPtr<FOpenGLBuffer> Buffer;
 
-	virtual ~FOpenGLShaderResourceView( void );
+	int32 LimitMip = -1;
 
-protected:
-	FOpenGLDynamicRHI* OpenGLRHI;
-	bool OwnsResource;
-};
-
-// this class is required to remove the SRV from the shader cache upon deletion
-class FOpenGLShaderResourceViewProxy : public TOpenGLResourceProxy<FRHIShaderResourceView, FOpenGLShaderResourceView>
-{
-public:
-	FOpenGLShaderResourceViewProxy(TFunction<FOpenGLShaderResourceView*(FRHIShaderResourceView*)> CreateFunc)
-		: TOpenGLResourceProxy<FRHIShaderResourceView, FOpenGLShaderResourceView>(CreateFunc)
-	{}
-
-	virtual ~FOpenGLShaderResourceViewProxy()
-	{
-
-	}
-};
-
-template<>
-struct TIsGLProxyObject<FOpenGLShaderResourceViewProxy>
-{
-	enum { Value = true };
+private:
+	bool OwnsResource = false;
 };
 
-void OPENGLDRV_API OpenGLTextureDeleted(FRHITexture* Texture);
-void OPENGLDRV_API OpenGLTextureAllocated( FRHITexture* Texture , ETextureCreateFlags Flags);
-
-void OPENGLDRV_API ReleaseOpenGLFramebuffers(FOpenGLDynamicRHI* Device, FRHITexture* TextureRHI);
+void OPENGLDRV_API ReleaseOpenGLFramebuffers(FRHITexture* TextureRHI);
 
 /** A OpenGL event query resource. */
-class FOpenGLEventQuery : public FRenderResource
+class FOpenGLEventQuery
 {
 public:
-
 	/** Initialization constructor. */
-	FOpenGLEventQuery(class FOpenGLDynamicRHI* InOpenGLRHI)
-		: OpenGLRHI(InOpenGLRHI)
-		, Sync(UGLsync())
-	{
-	}
+	FOpenGLEventQuery();
+	~FOpenGLEventQuery();
 
 	/** Issues an event for the query to poll. */
 	void IssueEvent();
@@ -2079,13 +1657,8 @@ public:
 	/** Waits for the event query to finish. */
 	void WaitForCompletion();
 
-	// FRenderResource interface.
-	virtual void InitDynamicRHI() override;
-	virtual void ReleaseDynamicRHI() override;
-
 private:
-	FOpenGLDynamicRHI* OpenGLRHI;
-	UGLsync Sync;
+	UGLsync Sync = {};
 };
 
 class FOpenGLViewport : public FRHIViewport
@@ -2099,18 +1672,11 @@ public:
 
 	// Accessors.
 	FIntPoint GetSizeXY() const { return FIntPoint(SizeX, SizeY); }
-	FOpenGLTexture2D *GetBackBuffer() const { return BackBuffer; }
+	FOpenGLTexture* GetBackBuffer() const { return BackBuffer; }
 	bool IsFullscreen( void ) const { return bIsFullscreen; }
 
-	virtual void WaitForFrameEventCompletion() override
-	{
-		FrameSyncEvent.WaitForCompletion();
-	}
-
-	virtual void IssueFrameEvent() override
-	{
-		FrameSyncEvent.IssueEvent();
-	}
+	virtual void WaitForFrameEventCompletion() override;
+	virtual void IssueFrameEvent() override;
 
 	virtual void* GetNativeWindow(void** AddParam) const override;
 
@@ -2133,8 +1699,8 @@ private:
 	bool bIsFullscreen;
 	EPixelFormat PixelFormat;
 	bool bIsValid;
-	TRefCountPtr<FOpenGLTexture2D> BackBuffer;
-	FOpenGLEventQuery FrameSyncEvent;
+	TRefCountPtr<FOpenGLTexture> BackBuffer;
+	TUniquePtr<FOpenGLEventQuery> FrameSyncEvent;
 	FCustomPresentRHIRef CustomPresent;
 };
 
@@ -2221,27 +1787,6 @@ struct TOpenGLResourceTraits<FRHIBoundShaderState>
 	typedef FOpenGLBoundShaderState TConcreteType;
 };
 template<>
-struct TOpenGLResourceTraits<FRHITexture3D>
-{
-	typedef FOpenGLTexture3D TConcreteType;
-};
-
-template<>
-struct TOpenGLResourceTraits<FRHITexture2D>
-{
-	typedef FOpenGLTexture2D TConcreteType;
-};
-template<>
-struct TOpenGLResourceTraits<FRHITexture2DArray>
-{
-	typedef FOpenGLTexture2DArray TConcreteType;
-};
-template<>
-struct TOpenGLResourceTraits<FRHITextureCube>
-{
-	typedef FOpenGLTextureCube TConcreteType;
-};
-template<>
 struct TOpenGLResourceTraits<FRHIRenderQuery>
 {
 	typedef FOpenGLRenderQuery TConcreteType;
@@ -2259,8 +1804,7 @@ struct TOpenGLResourceTraits<FRHIBuffer>
 template<>
 struct TOpenGLResourceTraits<FRHIShaderResourceView>
 {
-	//typedef FOpenGLShaderResourceView TConcreteType;
-	typedef FOpenGLShaderResourceViewProxy TConcreteType;
+	typedef FOpenGLShaderResourceView TConcreteType;
 };
 template<>
 struct TOpenGLResourceTraits<FRHIUnorderedAccessView>
diff --git a/Engine/Source/Runtime/OpenGLDrv/Public/OpenGLState.h b/Engine/Source/Runtime/OpenGLDrv/Public/OpenGLState.h
--- a/Engine/Source/Runtime/OpenGLDrv/Public/OpenGLState.h
+++ b/Engine/Source/Runtime/OpenGLDrv/Public/OpenGLState.h
@@ -192,7 +192,7 @@ public:
 
 struct FTextureStage
 {
-	class FOpenGLTextureBase* Texture;
+	class FOpenGLTexture* Texture;
 	class FOpenGLShaderResourceView* SRV;
 	GLenum Target;
 	GLuint Resource;
@@ -434,10 +434,10 @@ struct FOpenGLRHIState final : public FOpenGLCommonState
 
 	// Pending framebuffer setup
 	int32							FirstNonzeroRenderTarget;
-	FOpenGLTextureBase*				RenderTargets[MaxSimultaneousRenderTargets];
+	FOpenGLTexture*					RenderTargets[MaxSimultaneousRenderTargets];
 	uint32							RenderTargetMipmapLevels[MaxSimultaneousRenderTargets];
 	uint32							RenderTargetArrayIndex[MaxSimultaneousRenderTargets];
-	FOpenGLTextureBase*				DepthStencil;
+	FOpenGLTexture*					DepthStencil;
 	ERenderTargetStoreAction		StencilStoreAction;
 	uint32							DepthTargetWidth;
 	uint32							DepthTargetHeight;
diff --git a/Engine/Source/Runtime/RHI/Private/RHI.cpp b/Engine/Source/Runtime/RHI/Private/RHI.cpp
--- a/Engine/Source/Runtime/RHI/Private/RHI.cpp
+++ b/Engine/Source/Runtime/RHI/Private/RHI.cpp
@@ -2036,7 +2036,7 @@ void FRHIRenderPassInfo::Validate() const
 	} \
 
 // static
-bool FRHITextureCreateInfo::Validate(const FRHITextureCreateInfo& Desc, const TCHAR* Name, bool bFatal)
+bool FRHITextureDesc::Validate(const FRHITextureCreateInfo& Desc, const TCHAR* Name, bool bFatal)
 {
 	// Validate texture's pixel format.
 	{
@@ -2107,6 +2107,37 @@ bool FRHITextureCreateInfo::Validate(const FRHITextureCreateInfo& Desc, const TC
 
 #undef ValidateTextureDesc
 
+uint64 FRHITextureDesc::CalcMemorySizeEstimate(uint32 FirstMipIndex, uint32 LastMipIndex) const
+{
+#if DO_CHECK
+	Validate(*this, TEXT("CalcMemorySizeEstimate"), /* bFatal = */true);
+#endif
+	check(FirstMipIndex < NumMips && FirstMipIndex <= LastMipIndex && LastMipIndex < NumMips);
+
+	uint64 MemorySize = 0;
+	for (uint32 MipIndex = FirstMipIndex; MipIndex <= LastMipIndex; ++MipIndex)
+	{
+		FIntVector MipSizeInBlocks = FIntVector(
+			FMath::DivideAndRoundUp(FMath::Max(Extent.X >> MipIndex, 1), GPixelFormats[Format].BlockSizeX),
+			FMath::DivideAndRoundUp(FMath::Max(Extent.Y >> MipIndex, 1), GPixelFormats[Format].BlockSizeY),
+			FMath::DivideAndRoundUp(FMath::Max(Depth    >> MipIndex, 1), GPixelFormats[Format].BlockSizeZ)
+		);
+
+		uint32 NumBlocksInMip = MipSizeInBlocks.X * MipSizeInBlocks.Y * MipSizeInBlocks.Z;
+		MemorySize += NumBlocksInMip * GPixelFormats[Format].BlockBytes;
+	}
+
+	MemorySize *= ArraySize;
+	MemorySize *= NumSamples;
+
+	if (IsTextureCube())
+	{
+		MemorySize *= 6;
+	}
+
+	return MemorySize;
+}
+
 static FRHIPanicEvent RHIPanicEvent;
 FRHIPanicEvent& RHIGetPanicDelegate()
 {
diff --git a/Engine/Source/Runtime/RHI/Private/RHICommandList.cpp b/Engine/Source/Runtime/RHI/Private/RHICommandList.cpp
--- a/Engine/Source/Runtime/RHI/Private/RHICommandList.cpp
+++ b/Engine/Source/Runtime/RHI/Private/RHICommandList.cpp
@@ -2528,32 +2528,11 @@ FRHIShaderLibraryRef FDynamicRHI::RHICreateShaderLibrary_RenderThread(class FRHI
 	return GDynamicRHI->RHICreateShaderLibrary(Platform, FilePath, Name);
 }
 
-FTexture2DRHIRef FDynamicRHI::RHICreateTexture2D_RenderThread(class FRHICommandListImmediate& RHICmdList, uint32 SizeX, uint32 SizeY, uint8 Format, uint32 NumMips, uint32 NumSamples, ETextureCreateFlags Flags, ERHIAccess InResourceState, FRHIResourceCreateInfo& CreateInfo)
+FTextureRHIRef FDynamicRHI::RHICreateTexture_RenderThread(class FRHICommandListImmediate& RHICmdList, const FRHITextureCreateDesc& CreateDesc)
 {
-	CSV_SCOPED_TIMING_STAT(RHITStalls, RHICreateTexture2D_RenderThread);
+	CSV_SCOPED_TIMING_STAT(RHITStalls, RHICreateTexture_RenderThread);
 	FScopedRHIThreadStaller StallRHIThread(RHICmdList);
-	return GDynamicRHI->RHICreateTexture2D(SizeX, SizeY, Format, NumMips, NumSamples, Flags, InResourceState, CreateInfo);
-}
-
-FTexture2DRHIRef FDynamicRHI::RHICreateTextureExternal2D_RenderThread(class FRHICommandListImmediate& RHICmdList, uint32 SizeX, uint32 SizeY, uint8 Format, uint32 NumMips, uint32 NumSamples, ETextureCreateFlags Flags, ERHIAccess InResourceState, FRHIResourceCreateInfo& CreateInfo)
-{
-	CSV_SCOPED_TIMING_STAT(RHITStalls, RHICreateTextureExternal2D);
-	FScopedRHIThreadStaller StallRHIThread(RHICmdList);
-	return GDynamicRHI->RHICreateTextureExternal2D(SizeX, SizeY, Format, NumMips, NumSamples, Flags, InResourceState, CreateInfo);
-}
-
-FTexture2DArrayRHIRef FDynamicRHI::RHICreateTexture2DArray_RenderThread(class FRHICommandListImmediate& RHICmdList, uint32 SizeX, uint32 SizeY, uint32 SizeZ, uint8 Format, uint32 NumMips, uint32 NumSamples, ETextureCreateFlags Flags, ERHIAccess InResourceState, FRHIResourceCreateInfo& CreateInfo)
-{
-	CSV_SCOPED_TIMING_STAT(RHITStalls, RHICreateTexture2DArray);
-	FScopedRHIThreadStaller StallRHIThread(RHICmdList);
-	return GDynamicRHI->RHICreateTexture2DArray(SizeX, SizeY, SizeZ, Format, NumMips, NumSamples, Flags, InResourceState, CreateInfo);
-}
-
-FTexture3DRHIRef FDynamicRHI::RHICreateTexture3D_RenderThread(class FRHICommandListImmediate& RHICmdList, uint32 SizeX, uint32 SizeY, uint32 SizeZ, uint8 Format, uint32 NumMips, ETextureCreateFlags Flags, ERHIAccess InResourceState, FRHIResourceCreateInfo& CreateInfo)
-{
-	CSV_SCOPED_TIMING_STAT(RHITStalls, RHICreateTexture3D);
-	FScopedRHIThreadStaller StallRHIThread(RHICmdList);
-	return GDynamicRHI->RHICreateTexture3D(SizeX, SizeY, SizeZ, Format, NumMips, Flags, InResourceState, CreateInfo);
+	return GDynamicRHI->RHICreateTexture(CreateDesc);
 }
 
 FUnorderedAccessViewRHIRef FDynamicRHI::RHICreateUnorderedAccessView_RenderThread(class FRHICommandListImmediate& RHICmdList, FRHIBuffer* Buffer, bool bUseUAVCounter, bool bAppendBuffer)
@@ -2626,20 +2605,6 @@ FShaderResourceViewRHIRef FDynamicRHI::RHICreateShaderResourceViewFMask_RenderTh
 	return GDynamicRHI->RHICreateShaderResourceViewFMask(Texture2DRHI);
 }
 
-FTextureCubeRHIRef FDynamicRHI::RHICreateTextureCube_RenderThread(class FRHICommandListImmediate& RHICmdList, uint32 Size, uint8 Format, uint32 NumMips, ETextureCreateFlags Flags, ERHIAccess InResourceState, FRHIResourceCreateInfo& CreateInfo)
-{
-	CSV_SCOPED_TIMING_STAT(RHITStalls, RHICreateTextureCube_RenderThread);
-	FScopedRHIThreadStaller StallRHIThread(RHICmdList);
-	return GDynamicRHI->RHICreateTextureCube(Size, Format, NumMips, Flags, InResourceState, CreateInfo);
-}
-
-FTextureCubeRHIRef FDynamicRHI::RHICreateTextureCubeArray_RenderThread(class FRHICommandListImmediate& RHICmdList, uint32 Size, uint32 ArraySize, uint8 Format, uint32 NumMips, ETextureCreateFlags Flags, ERHIAccess InResourceState, FRHIResourceCreateInfo& CreateInfo)
-{
-	CSV_SCOPED_TIMING_STAT(RHITStalls, RHICreateTextureCubeArray_RenderThread);
-	FScopedRHIThreadStaller StallRHIThread(RHICmdList);
-	return GDynamicRHI->RHICreateTextureCubeArray(Size, ArraySize, Format, NumMips, Flags, InResourceState, CreateInfo);
-}
-
 FRenderQueryRHIRef FDynamicRHI::RHICreateRenderQuery_RenderThread(class FRHICommandListImmediate& RHICmdList, ERenderQueryType QueryType)
 {
 	CSV_SCOPED_TIMING_STAT(RHITStalls, RHICreateRenderQuery_RenderThread);
diff --git a/Engine/Source/Runtime/RHI/Private/RHIValidation.cpp b/Engine/Source/Runtime/RHI/Private/RHIValidation.cpp
--- a/Engine/Source/Runtime/RHI/Private/RHIValidation.cpp
+++ b/Engine/Source/Runtime/RHI/Private/RHIValidation.cpp
@@ -67,6 +67,23 @@ namespace RHIValidation
 
 		return Init.Strings;
 	}
+
+	FTextureResource::FTextureResource(FRHITextureCreateDesc const& CreateDesc)
+		: FTextureResource()
+	{
+		InitBarrierTracking(CreateDesc);
+	}
+
+	void FTextureResource::InitBarrierTracking(FRHITextureCreateDesc const& CreateDesc)
+	{
+		InitBarrierTracking(
+			CreateDesc.NumMips,
+			CreateDesc.ArraySize * (CreateDesc.IsTextureCube() ? 6 : 1),
+			CreateDesc.Format,
+			CreateDesc.Flags,
+			CreateDesc.InitialState,
+			CreateDesc.DebugName);
+	}
 }
 
 TSet<uint32> FValidationRHI::SeenFailureHashes;
diff --git a/Engine/Source/Runtime/RHI/Public/DynamicRHI.h b/Engine/Source/Runtime/RHI/Public/DynamicRHI.h
--- a/Engine/Source/Runtime/RHI/Public/DynamicRHI.h
+++ b/Engine/Source/Runtime/RHI/Public/DynamicRHI.h
@@ -462,73 +462,25 @@ public:
 
 	virtual void RHIUpdateShaderResourceView(FRHIShaderResourceView* SRV, FRHIBuffer* Buffer);
 
-	/**
-	* Computes the total size of a 2D texture with the specified parameters.
-	*
-	* @param SizeX - width of the texture to compute
-	* @param SizeY - height of the texture to compute
-	* @param Format - EPixelFormat texture format
-	* @param NumMips - number of mips to compute or 0 for full mip pyramid
-	* @param NumSamples - number of MSAA samples, usually 1
-	* @param Flags - ETextureCreateFlags creation flags
-	* @param OutAlign - Alignment required for this texture.  Output parameter.
-	*/
-	// FlushType: Thread safe
-	virtual uint64 RHICalcTexture2DPlatformSize(uint32 SizeX, uint32 SizeY, uint8 Format, uint32 NumMips, uint32 NumSamples, ETextureCreateFlags Flags, const FRHIResourceCreateInfo& CreateInfo, uint32& OutAlign) = 0;
-
-	/**
-	* Computes the total size of a 2D array texture with the specified parameters.
-	*
-	* @param SizeX - width of the texture to compute
-	* @param SizeY - height of the texture to compute
-	* @param ArraySize - number of textures in the array.
-	* @param Format - EPixelFormat texture format
-	* @param NumMips - number of mips to compute or 0 for full mip pyramid
-	* @param NumSamples - number of MSAA samples, usually 1
-	* @param Flags - ETextureCreateFlags creation flags
-	* @param OutAlign - Alignment required for this texture.  Output parameter.
-	*/
-	// FlushType: Thread safe
-	virtual uint64 RHICalcTexture2DArrayPlatformSize(uint32 SizeX, uint32 SizeY, uint32 ArraySize, uint8 Format, uint32 NumMips, uint32 NumSamples, ETextureCreateFlags Flags, const FRHIResourceCreateInfo& CreateInfo, uint32& OutAlign);
-	
-	/**
-	* Computes the total size of a virtual memory (VM) based 2D texture with the specified parameters.
-	*
-	* @param Mip0Width - width of the top mip
-	* @param Mip0Height - height of the top mip
-	* @param Format - EPixelFormat texture format
-	* @param NumMips - number of mips the texture has
-	* @param FirstMipIdx - index of the first resident mip
-	* @param NumSamples - number of MSAA samples, usually 1
-	* @param Flags - ETextureCreateFlags creation flags
-	* @param OutAlign - Alignment required for this texture. Output parameter.
-	*/
-	// FlushType: Thread safe
-	virtual uint64 RHICalcVMTexture2DPlatformSize(uint32 Mip0Width, uint32 Mip0Height, uint8 Format, uint32 NumMips, uint32 FirstMipIdx, uint32 NumSamples, ETextureCreateFlags Flags, uint32& OutAlign);
+	struct FRHICalcTextureSizeResult
+	{
+		// The total size of the texture, in bytes.
+		uint64 Size;
 
-	/**
-	* Computes the total size of a 3D texture with the specified parameters.
-	* @param SizeX - width of the texture to create
-	* @param SizeY - height of the texture to create
-	* @param SizeZ - depth of the texture to create
-	* @param Format - EPixelFormat texture format
-	* @param NumMips - number of mips to generate or 0 for full mip pyramid
-	* @param Flags - ETextureCreateFlags creation flags
-	* @param OutAlign - Alignment required for this texture.  Output parameter.
-	*/
-	// FlushType: Thread safe
-	virtual uint64 RHICalcTexture3DPlatformSize(uint32 SizeX, uint32 SizeY, uint32 SizeZ, uint8 Format, uint32 NumMips, ETextureCreateFlags Flags, const FRHIResourceCreateInfo& CreateInfo, uint32& OutAlign) = 0;
+		// The required address alignment for the texture.
+		uint32 Align;
+	};
 
 	/**
-	* Computes the total size of a cube texture with the specified parameters.
-	* @param Size - width/height of the texture to create
-	* @param Format - EPixelFormat texture format
-	* @param NumMips - number of mips to generate or 0 for full mip pyramid
-	* @param Flags - ETextureCreateFlags creation flags
-	* @param OutAlign - Alignment required for this texture.  Output parameter.
+	* Computes the total GPU memory a texture resource with the specified parameters will occupy on the current RHI platform.
+	* Also returns the required alignment for the resource.
+	*
+	* @param Desc          - The texture descriptor (width, height, format etc)
+	* @param FirstMipIndex - The index of the most detailed mip to consider in the memory size calculation.
+	* @return              - The computed size and alignment of the platform texture resource.
 	*/
 	// FlushType: Thread safe
-	virtual uint64 RHICalcTextureCubePlatformSize(uint32 Size, uint8 Format, uint32 NumMips, ETextureCreateFlags Flags, const FRHIResourceCreateInfo& CreateInfo, uint32& OutAlign) = 0;
+	virtual FRHICalcTextureSizeResult RHICalcTexturePlatformSize(FRHITextureDesc const& Desc, uint32 FirstMipIndex) = 0;
 
 	/**
 	* Gets the minimum alignment (in bytes) required for creating a shader resource view on a buffer-backed resource.
@@ -559,31 +511,9 @@ public:
 	virtual bool RHIGetTextureMemoryVisualizeData(FColor* TextureData, int32 SizeX, int32 SizeY, int32 Pitch, int32 PixelSize) = 0;
 
 	/**
-	* Creates a 2D RHI texture resource
-	* @param SizeX - width of the texture to create
-	* @param SizeY - height of the texture to create
-	* @param Format - EPixelFormat texture format
-	* @param NumMips - number of mips to generate or 0 for full mip pyramid
-	* @param NumSamples - number of MSAA samples, usually 1
-	* @param Flags - ETextureCreateFlags creation flags
+	* Creates an RHI texture resource.
 	*/
-	// FlushType: Wait RHI Thread
-	virtual FTexture2DRHIRef RHICreateTexture2D(uint32 SizeX, uint32 SizeY, uint8 Format, uint32 NumMips, uint32 NumSamples, ETextureCreateFlags Flags, ERHIAccess InResourceState, FRHIResourceCreateInfo& CreateInfo) = 0;
-
-	/**
-	* Creates a 2D RHI texture external resource
-	* @param SizeX - width of the texture to create
-	* @param SizeY - height of the texture to create
-	* @param Format - EPixelFormat texture format
-	* @param NumMips - number of mips to generate or 0 for full mip pyramid
-	* @param NumSamples - number of MSAA samples, usually 1
-	* @param Flags - ETextureCreateFlags creation flags
-	*/
-	// FlushType: Wait RHI Thread
-	virtual FTexture2DRHIRef RHICreateTextureExternal2D(uint32 SizeX, uint32 SizeY, uint8 Format, uint32 NumMips, uint32 NumSamples, ETextureCreateFlags Flags, ERHIAccess InResourceState, FRHIResourceCreateInfo& CreateInfo)
-	{
-		return nullptr;
-	}
+	virtual FTextureRHIRef RHICreateTexture(const FRHITextureCreateDesc& CreateDesc) = 0;
 
 	/**
 	* Thread-safe function that can be used to create a texture outside of the
@@ -599,7 +529,7 @@ public:
 	* @returns a reference to a 2D texture resource
 	*/
 	// FlushType: Thread safe
-	virtual FTexture2DRHIRef RHIAsyncCreateTexture2D(uint32 SizeX, uint32 SizeY, uint8 Format, uint32 NumMips, ETextureCreateFlags Flags, ERHIAccess InResourceState, void** InitialMipData, uint32 NumInitialMips) = 0;
+	virtual FTextureRHIRef RHIAsyncCreateTexture2D(uint32 SizeX, uint32 SizeY, uint8 Format, uint32 NumMips, ETextureCreateFlags Flags, ERHIAccess InResourceState, void** InitialMipData, uint32 NumInitialMips) = 0;
 
 	/**
 	* Copies shared mip levels from one texture to another. The textures must have
@@ -611,31 +541,6 @@ public:
 	// FlushType: Flush RHI Thread
 	virtual void RHICopySharedMips(FRHITexture2D* DestTexture2D, FRHITexture2D* SrcTexture2D) = 0;
 
-	/**
-	* Creates a Array RHI texture resource
-	* @param SizeX - width of the texture to create
-	* @param SizeY - height of the texture to create
-	* @param SizeZ - depth of the texture to create
-	* @param Format - EPixelFormat texture format
-	* @param NumMips - number of mips to generate or 0 for full mip pyramid
-	* @param NumSamples - number of MSAA samples, usually 1
-	* @param Flags - ETextureCreateFlags creation flags
-	*/
-	// FlushType: Wait RHI Thread
-	virtual FTexture2DArrayRHIRef RHICreateTexture2DArray(uint32 SizeX, uint32 SizeY, uint32 SizeZ, uint8 Format, uint32 NumMips, uint32 NumSamples, ETextureCreateFlags Flags, ERHIAccess InResourceState, FRHIResourceCreateInfo& CreateInfo) = 0;
-
-	/**
-	* Creates a 3d RHI texture resource
-	* @param SizeX - width of the texture to create
-	* @param SizeY - height of the texture to create
-	* @param SizeZ - depth of the texture to create
-	* @param Format - EPixelFormat texture format
-	* @param NumMips - number of mips to generate or 0 for full mip pyramid
-	* @param Flags - ETextureCreateFlags creation flags
-	*/
-	// FlushType: Wait RHI Thread
-	virtual FTexture3DRHIRef RHICreateTexture3D(uint32 SizeX, uint32 SizeY, uint32 SizeZ, uint8 Format, uint32 NumMips, ETextureCreateFlags Flags, ERHIAccess InResourceState, FRHIResourceCreateInfo& CreateInfo) = 0;
-
 	/**
 	* @param Ref may be 0
 	*/
@@ -812,27 +717,6 @@ public:
 	// FlushType: Flush RHI Thread
 	virtual void RHIUpdateTexture3D(FRHITexture3D* Texture, uint32 MipIndex, const struct FUpdateTextureRegion3D& UpdateRegion, uint32 SourceRowPitch, uint32 SourceDepthPitch, const uint8* SourceData) = 0;
 
-	/**
-	* Creates a Cube RHI texture resource
-	* @param Size - width/height of the texture to create
-	* @param Format - EPixelFormat texture format
-	* @param NumMips - number of mips to generate or 0 for full mip pyramid
-	* @param Flags - ETextureCreateFlags creation flags
-	*/
-	// FlushType: Wait RHI Thread
-	virtual FTextureCubeRHIRef RHICreateTextureCube(uint32 Size, uint8 Format, uint32 NumMips, ETextureCreateFlags Flags, ERHIAccess InResourceState, FRHIResourceCreateInfo& CreateInfo) = 0;
-
-	/**
-	* Creates a Cube Array RHI texture resource
-	* @param Size - width/height of the texture to create
-	* @param ArraySize - number of array elements of the texture to create
-	* @param Format - EPixelFormat texture format
-	* @param NumMips - number of mips to generate or 0 for full mip pyramid
-	* @param Flags - ETextureCreateFlags creation flags
-	*/
-	// FlushType: Wait RHI Thread
-	virtual FTextureCubeRHIRef RHICreateTextureCubeArray(uint32 Size, uint32 ArraySize, uint8 Format, uint32 NumMips, ETextureCreateFlags Flags, ERHIAccess InResourceState, FRHIResourceCreateInfo& CreateInfo) = 0;
-
 	/**
 	* Locks an RHI texture's mip-map for read/write operations on the CPU
 	* @param Texture - the RHI texture resource to lock
@@ -1178,12 +1062,9 @@ public:
 	virtual void UpdateTexture3D_RenderThread(class FRHICommandListImmediate& RHICmdList, FRHITexture3D* Texture, uint32 MipIndex, const struct FUpdateTextureRegion3D& UpdateRegion, uint32 SourceRowPitch, uint32 SourceDepthPitch, const uint8* SourceData);
 
 	virtual FRHIShaderLibraryRef RHICreateShaderLibrary_RenderThread(class FRHICommandListImmediate& RHICmdList, EShaderPlatform Platform, FString FilePath, FString Name);
-	virtual FTexture2DRHIRef RHICreateTexture2D_RenderThread(class FRHICommandListImmediate& RHICmdList, uint32 SizeX, uint32 SizeY, uint8 Format, uint32 NumMips, uint32 NumSamples, ETextureCreateFlags Flags, ERHIAccess InResourceState, FRHIResourceCreateInfo& CreateInfo);
-	virtual FTexture2DRHIRef RHICreateTextureExternal2D_RenderThread(class FRHICommandListImmediate& RHICmdList, uint32 SizeX, uint32 SizeY, uint8 Format, uint32 NumMips, uint32 NumSamples, ETextureCreateFlags Flags, ERHIAccess InResourceState, FRHIResourceCreateInfo& CreateInfo);
 
-	virtual FTexture2DArrayRHIRef RHICreateTexture2DArray_RenderThread(class FRHICommandListImmediate& RHICmdList, uint32 SizeX, uint32 SizeY, uint32 SizeZ, uint8 Format, uint32 NumMips, uint32 NumSamples, ETextureCreateFlags Flags, ERHIAccess InResourceState, FRHIResourceCreateInfo& CreateInfo);
-	
-	virtual FTexture3DRHIRef RHICreateTexture3D_RenderThread(class FRHICommandListImmediate& RHICmdList, uint32 SizeX, uint32 SizeY, uint32 SizeZ, uint8 Format, uint32 NumMips, ETextureCreateFlags Flags, ERHIAccess InResourceState, FRHIResourceCreateInfo& CreateInfo);
+	virtual FTextureRHIRef RHICreateTexture_RenderThread(class FRHICommandListImmediate& RHICmdList, const FRHITextureCreateDesc& CreateDesc);
+
 	virtual FUnorderedAccessViewRHIRef RHICreateUnorderedAccessView_RenderThread(class FRHICommandListImmediate& RHICmdList, FRHIBuffer* Buffer, bool bUseUAVCounter, bool bAppendBuffer);
 	virtual FUnorderedAccessViewRHIRef RHICreateUnorderedAccessView_RenderThread(class FRHICommandListImmediate& RHICmdList, FRHITexture* Texture, uint32 MipLevel, uint16 FirstArraySlice, uint16 NumArraySlices);
 	virtual FUnorderedAccessViewRHIRef RHICreateUnorderedAccessView_RenderThread(class FRHICommandListImmediate& RHICmdList, FRHITexture* Texture, uint32 MipLevel, uint8 Format, uint16 FirstArraySlice, uint16 NumArraySlices);
@@ -1194,8 +1075,6 @@ public:
 	virtual FShaderResourceViewRHIRef RHICreateShaderResourceView_RenderThread(class FRHICommandListImmediate& RHICmdList, FRHIBuffer* Buffer);
 	virtual FShaderResourceViewRHIRef RHICreateShaderResourceViewWriteMask_RenderThread(class FRHICommandListImmediate& RHICmdList, FRHITexture2D* Texture2D);
 	virtual FShaderResourceViewRHIRef RHICreateShaderResourceViewFMask_RenderThread(class FRHICommandListImmediate& RHICmdList, FRHITexture2D* Texture2D);
-	virtual FTextureCubeRHIRef RHICreateTextureCube_RenderThread(class FRHICommandListImmediate& RHICmdList, uint32 Size, uint8 Format, uint32 NumMips, ETextureCreateFlags Flags, ERHIAccess InResourceState, FRHIResourceCreateInfo& CreateInfo);
-	virtual FTextureCubeRHIRef RHICreateTextureCubeArray_RenderThread(class FRHICommandListImmediate& RHICmdList, uint32 Size, uint32 ArraySize, uint8 Format, uint32 NumMips, ETextureCreateFlags Flags, ERHIAccess InResourceState, FRHIResourceCreateInfo& CreateInfo);
 	virtual FRenderQueryRHIRef RHICreateRenderQuery_RenderThread(class FRHICommandListImmediate& RHICmdList, ERenderQueryType QueryType);
 
 	
@@ -1422,29 +1301,130 @@ FORCEINLINE void RHIUpdateUniformBuffer(FRHIUniformBuffer* UniformBufferRHI, con
 	return GDynamicRHI->RHIUpdateUniformBuffer(UniformBufferRHI, Contents);
 }
 
+FORCEINLINE FDynamicRHI::FRHICalcTextureSizeResult RHICalcTexturePlatformSize(FRHITextureDesc const& Desc, uint32 FirstMipIndex = 0)
+{
+	return GDynamicRHI->RHICalcTexturePlatformSize(Desc, FirstMipIndex);
+}
+
+//UE_DEPRECATED(5.1, "The separate RHICalcTexture... functions for each texture type are deprecated. Use RHICalcTexturePlatformSize instead.")
 FORCEINLINE uint64 RHICalcTexture2DPlatformSize(uint32 SizeX, uint32 SizeY, uint8 Format, uint32 NumMips, uint32 NumSamples, ETextureCreateFlags Flags, const FRHIResourceCreateInfo& CreateInfo, uint32& OutAlign)
 {
-	return GDynamicRHI->RHICalcTexture2DPlatformSize(SizeX, SizeY, Format, NumMips, NumSamples, Flags, CreateInfo, OutAlign);
+	const uint32 Depth = 1;
+	const uint32 ArraySize = 1;
+
+	FRHITextureDesc Desc(
+		ETextureDimension::Texture2D,
+		Flags,
+		(EPixelFormat)Format,
+		CreateInfo.ClearValueBinding,
+		{ (int32)SizeX, (int32)SizeY },
+		Depth,
+		ArraySize,
+		NumMips,
+		NumSamples,
+		CreateInfo.ExtData
+	);
+
+	auto Result = RHICalcTexturePlatformSize(Desc, 0);
+	OutAlign = Result.Align;
+	return Result.Size;
 }
 
+//UE_DEPRECATED(5.1, "The separate RHICalcTexture... functions for each texture type are deprecated. Use RHICalcTexturePlatformSize instead.")
 FORCEINLINE uint64 RHICalcTexture2DArrayPlatformSize(uint32 SizeX, uint32 SizeY, uint32 ArraySize, uint8 Format, uint32 NumMips, uint32 NumSamples, ETextureCreateFlags Flags, const FRHIResourceCreateInfo& CreateInfo, uint32& OutAlign)
 {
-	return GDynamicRHI->RHICalcTexture2DArrayPlatformSize(SizeX, SizeY, ArraySize, Format, NumMips, NumSamples, Flags, CreateInfo, OutAlign);
+	const uint32 Depth = 1;
+
+	FRHITextureDesc Desc(
+		ETextureDimension::Texture2DArray,
+		Flags,
+		(EPixelFormat)Format,
+		CreateInfo.ClearValueBinding,
+		{ (int32)SizeX, (int32)SizeY },
+		Depth,
+		ArraySize,
+		NumMips,
+		NumSamples,
+		CreateInfo.ExtData
+	);
+
+	auto Result = RHICalcTexturePlatformSize(Desc, 0);
+	OutAlign = Result.Align;
+	return Result.Size;
 }
 
+//UE_DEPRECATED(5.1, "The separate RHICalcTexture... functions for each texture type are deprecated. Use RHICalcTexturePlatformSize instead.")
 FORCEINLINE uint64 RHICalcVMTexture2DPlatformSize(uint32 Mip0Width, uint32 Mip0Height, uint8 Format, uint32 NumMips, uint32 FirstMipIdx, uint32 NumSamples, ETextureCreateFlags Flags, uint32& OutAlign)
 {
-	return GDynamicRHI->RHICalcVMTexture2DPlatformSize(Mip0Width, Mip0Height, Format, NumMips, FirstMipIdx, NumSamples, Flags, OutAlign);
+	const uint32 Depth     = 1;
+	const uint32 ArraySize = 1;
+	const uint32 ExtData   = 0;
+
+	FRHITextureDesc Desc(
+		ETextureDimension::Texture2D,
+		Flags,
+		(EPixelFormat)Format,
+		FClearValueBinding::None,
+		{ (int32)Mip0Width, (int32)Mip0Height },
+		Depth,
+		ArraySize,
+		NumMips,
+		NumSamples,
+		ExtData
+	);
+
+	auto Result = RHICalcTexturePlatformSize(Desc, FirstMipIdx);
+	OutAlign = Result.Align;
+	return Result.Size;
 }
 
+//UE_DEPRECATED(5.1, "The separate RHICalcTexture... functions for each texture type are deprecated. Use RHICalcTexturePlatformSize instead.")
 FORCEINLINE uint64 RHICalcTexture3DPlatformSize(uint32 SizeX, uint32 SizeY, uint32 SizeZ, uint8 Format, uint32 NumMips, ETextureCreateFlags Flags, const FRHIResourceCreateInfo& CreateInfo, uint32& OutAlign)
 {
-	return GDynamicRHI->RHICalcTexture3DPlatformSize(SizeX, SizeY, SizeZ, Format, NumMips, Flags, CreateInfo, OutAlign);
+	const uint32 ArraySize  = 1;
+	const uint32 NumSamples = 1;
+
+	FRHITextureDesc Desc(
+		ETextureDimension::Texture3D,
+		Flags,
+		(EPixelFormat)Format,
+		CreateInfo.ClearValueBinding,
+		{ (int32)SizeX, (int32)SizeY },
+		SizeZ,
+		ArraySize,
+		NumMips,
+		NumSamples,
+		CreateInfo.ExtData
+	);
+
+	auto Result = RHICalcTexturePlatformSize(Desc, 0);
+	OutAlign = Result.Align;
+	return Result.Size;
 }
 
+//UE_DEPRECATED(5.1, "The separate RHICalcTexture... functions for each texture type are deprecated. Use RHICalcTexturePlatformSize instead.")
 FORCEINLINE uint64 RHICalcTextureCubePlatformSize(uint32 Size, uint8 Format, uint32 NumMips, ETextureCreateFlags Flags, const FRHIResourceCreateInfo& CreateInfo, uint32& OutAlign)
 {
-	return GDynamicRHI->RHICalcTextureCubePlatformSize(Size, Format, NumMips, Flags, CreateInfo, OutAlign);
+	const uint32 Depth      = 1;
+	const uint32 ArraySize  = 1;
+	const uint32 NumSamples = 1;
+
+	FRHITextureDesc Desc(
+		ETextureDimension::TextureCube,
+		Flags,
+		(EPixelFormat)Format,
+		CreateInfo.ClearValueBinding,
+		{ (int32)Size, (int32)Size },
+		Depth,
+		ArraySize,
+		NumMips,
+		NumSamples,
+		CreateInfo.ExtData
+	);
+
+	auto Result = RHICalcTexturePlatformSize(Desc, 0);
+	OutAlign = Result.Align;
+	return Result.Size;
 }
 
 FORCEINLINE uint64 RHIGetMinimumAlignmentForBufferBackedSRV(EPixelFormat Format)
diff --git a/Engine/Source/Runtime/RHI/Public/RHI.h b/Engine/Source/Runtime/RHI/Public/RHI.h
--- a/Engine/Source/Runtime/RHI/Public/RHI.h
+++ b/Engine/Source/Runtime/RHI/Public/RHI.h
@@ -1724,6 +1724,26 @@ struct FClearValueBinding
 		return false;
 	}
 
+	friend inline uint32 GetTypeHash(FClearValueBinding const& Binding)
+	{
+		uint32 Hash = GetTypeHash(Binding.ColorBinding);
+
+		if (Binding.ColorBinding == EClearBinding::EColorBound)
+		{
+			Hash = HashCombine(Hash, GetTypeHash(Binding.Value.Color[0]));
+			Hash = HashCombine(Hash, GetTypeHash(Binding.Value.Color[1]));
+			Hash = HashCombine(Hash, GetTypeHash(Binding.Value.Color[2]));
+			Hash = HashCombine(Hash, GetTypeHash(Binding.Value.Color[3]));
+		}
+		else if (Binding.ColorBinding == EClearBinding::EDepthStencilBound)
+		{
+			Hash = HashCombine(Hash, GetTypeHash(Binding.Value.DSValue.Depth  ));
+			Hash = HashCombine(Hash, GetTypeHash(Binding.Value.DSValue.Stencil));
+		}
+
+		return Hash;
+	}
+
 	EClearBinding ColorBinding;
 
 	union ClearValueType
