Fixed: incorrect CRC check for OpenGL shaders

diff --git a/Engine/Source/Runtime/OpenGLDrv/Private/OpenGLShaders.cpp b/Engine/Source/Runtime/OpenGLDrv/Private/OpenGLShaders.cpp
--- a/Engine/Source/Runtime/OpenGLDrv/Private/OpenGLShaders.cpp
+++ b/Engine/Source/Runtime/OpenGLDrv/Private/OpenGLShaders.cpp
@@ -989,12 +989,14 @@ void OPENGLDRV_API GLSLToDeviceCompatibleGLSL(FAnsiCharArray& GlslCodeOriginal,
 		}
 	}
 
+#if	DEBUG_GL_SHADERS
 	if (ShaderName.IsEmpty() == false)
 	{
 		AppendCString(GlslCode, "// ");
 		AppendCString(GlslCode, TCHAR_TO_ANSI(ShaderName.GetCharArray().GetData()));
 		AppendCString(GlslCode, "\n");
 	}
+#endif
 
 	if (bEmitMobileMultiView && GSupportsMobileMultiView && TypeEnum == GL_VERTEX_SHADER)
 	{
@@ -1130,7 +1132,8 @@ RHIType* CreateProxyShader(TArrayView<const uint8> Code, const FSHAHash& Hash)
 			GLSLToPlatform(Header, TypeEnum, GlslCodeOriginal, Finalglsl);
 			GetOpenGLCompiledShaderCache().Add(ShaderCodeKey, MakeShareable(new FOpenGLCompiledShaderValue(0, Finalglsl)));
 		}
-#if (UE_BUILD_DEBUG || UE_BUILD_DEVELOPMENT)
+		// With debug shaders we insert a shader name into the source and that can make it unique failing CRC check
+#if (UE_BUILD_DEBUG || UE_BUILD_DEVELOPMENT) && !DEBUG_GL_SHADERS 
 		else
 		{
 			QUICK_SCOPE_CYCLE_COUNTER(STAT_GLCheckShaderCodeCRC);
