Refactored usage of out of scope strings in RHICreateAliasedTexture

diff --git a/Engine/Source/Runtime/D3D12RHI/Private/D3D12Texture.cpp b/Engine/Source/Runtime/D3D12RHI/Private/D3D12Texture.cpp
--- a/Engine/Source/Runtime/D3D12RHI/Private/D3D12Texture.cpp
+++ b/Engine/Source/Runtime/D3D12RHI/Private/D3D12Texture.cpp
@@ -2818,7 +2818,8 @@ FD3D12Texture* FD3D12DynamicRHI::CreateAliasedD3D12Texture2D(FD3D12Texture* Sour
 	const DXGI_FORMAT PlatformShaderResourceFormat = FindShaderResourceDXGIFormat(PlatformResourceFormat, bSRGB);
 	const DXGI_FORMAT PlatformRenderTargetFormat = FindShaderResourceDXGIFormat(PlatformResourceFormat, bSRGB);
 
-	FRHITextureCreateDesc CreateDesc(SourceTexture->GetDesc(), ERHIAccess::SRVMask, TEXT("AliasedTexture"));
+	const FString Name = SourceTexture->GetName().ToString() + TEXT("Alias");
+	FRHITextureCreateDesc CreateDesc(SourceTexture->GetDesc(), ERHIAccess::SRVMask, *Name);
 
 	FD3D12Texture* Texture2D = Adapter->CreateLinkedObject<FD3D12Texture>(Device->GetGPUMask(), [&](FD3D12Device* Device)
 	{
diff --git a/Engine/Source/Runtime/OpenGLDrv/Private/OpenGLTexture.cpp b/Engine/Source/Runtime/OpenGLDrv/Private/OpenGLTexture.cpp
--- a/Engine/Source/Runtime/OpenGLDrv/Private/OpenGLTexture.cpp
+++ b/Engine/Source/Runtime/OpenGLDrv/Private/OpenGLTexture.cpp
@@ -294,8 +294,8 @@ FOpenGLTextureDesc::FOpenGLTextureDesc(FRHITextureDesc const& InDesc)
 }
 
 // Constructor for RHICreateAliasedTexture
-FOpenGLTexture::FOpenGLTexture(FOpenGLTexture& Other, EAliasConstructorParam)
-	: FRHITexture(FRHITextureCreateDesc(Other.GetDesc(), ERHIAccess::SRVMask, *Other.GetName().ToString()))
+FOpenGLTexture::FOpenGLTexture(FOpenGLTexture& Other, const FString& Name, EAliasConstructorParam)
+	: FRHITexture(FRHITextureCreateDesc(Other.GetDesc(), ERHIAccess::SRVMask, *Name))
 	, Target             (Other.Target)
 	, Attachment         (Other.Attachment)
 	, MemorySize         (0)
@@ -2334,7 +2334,8 @@ void FOpenGLDynamicRHI::RHIAliasTextureResources(FTextureRHIRef& DestRHITexture,
 
 FTextureRHIRef FOpenGLDynamicRHI::RHICreateAliasedTexture(FTextureRHIRef& SourceTexture)
 {
-	return new FOpenGLTexture(*GetOpenGLTextureFromRHITexture(SourceTexture), FOpenGLTexture::AliasResource);
+	const FString Name = SourceTexture->GetName().ToString() + TEXT("Alias");
+	return new FOpenGLTexture(*GetOpenGLTextureFromRHITexture(SourceTexture), *Name, FOpenGLTexture::AliasResource);
 }
 
 void* FOpenGLDynamicRHI::LockTexture2D_RenderThread(class FRHICommandListImmediate& RHICmdList, FRHITexture2D* Texture, uint32 MipIndex, EResourceLockMode LockMode, uint32& DestStride, bool bLockWithinMiptail, bool bNeedsDefaultRHIFlush)
diff --git a/Engine/Source/Runtime/OpenGLDrv/Public/OpenGLResources.h b/Engine/Source/Runtime/OpenGLDrv/Public/OpenGLResources.h
--- a/Engine/Source/Runtime/OpenGLDrv/Public/OpenGLResources.h
+++ b/Engine/Source/Runtime/OpenGLDrv/Public/OpenGLResources.h
@@ -1271,7 +1271,7 @@ public:
 
 	// Constructor for RHICreateAliasedTexture
 	enum EAliasConstructorParam { AliasResource };
-	explicit FOpenGLTexture(FOpenGLTexture& Texture, EAliasConstructorParam);
+	explicit FOpenGLTexture(FOpenGLTexture& Texture, const FString& Name, EAliasConstructorParam);
 	void AliasResources(FOpenGLTexture& Texture);
 
 	virtual ~FOpenGLTexture();
diff --git a/Engine/Source/Runtime/VulkanRHI/Private/VulkanRHI.cpp b/Engine/Source/Runtime/VulkanRHI/Private/VulkanRHI.cpp
--- a/Engine/Source/Runtime/VulkanRHI/Private/VulkanRHI.cpp
+++ b/Engine/Source/Runtime/VulkanRHI/Private/VulkanRHI.cpp
@@ -1499,7 +1499,8 @@ void FVulkanDynamicRHI::RHIAliasTextureResources(FTextureRHIRef& DestTextureRHI,
 
 FTextureRHIRef FVulkanDynamicRHI::RHICreateAliasedTexture(FTextureRHIRef& SourceTextureRHI)
 {
-	FRHITextureCreateDesc Desc(SourceTextureRHI->GetDesc(), ERHIAccess::SRVMask, *SourceTextureRHI->GetName().ToString());
+	const FString Name = SourceTextureRHI->GetName().ToString() + TEXT("Alias");
+	FRHITextureCreateDesc Desc(SourceTextureRHI->GetDesc(), ERHIAccess::SRVMask, *Name);
 	return new FVulkanTexture(*Device, Desc, SourceTextureRHI);
 }
 
diff --git a/Engine/Source/Runtime/Windows/D3D11RHI/Private/D3D11Texture.cpp b/Engine/Source/Runtime/Windows/D3D11RHI/Private/D3D11Texture.cpp
--- a/Engine/Source/Runtime/Windows/D3D11RHI/Private/D3D11Texture.cpp
+++ b/Engine/Source/Runtime/Windows/D3D11RHI/Private/D3D11Texture.cpp
@@ -2008,8 +2008,8 @@ FD3D11Texture::FD3D11Texture(
 	UpdateD3D11TextureStats(*this, true);
 }
 
-FD3D11Texture::FD3D11Texture(FD3D11Texture const& Other, EAliasResourceParam)
-	: FRHITexture(FRHITextureCreateDesc(Other.GetDesc(), ERHIAccess::SRVMask, TEXT("FD3D11Texture::FD3D11Texture Alias")))
+FD3D11Texture::FD3D11Texture(FD3D11Texture const& Other, const FString& Name, EAliasResourceParam)
+	: FRHITexture(FRHITextureCreateDesc(Other.GetDesc(), ERHIAccess::SRVMask, *Name))
 	, bAlias(true)
 {
 	AliasResource(Other);
@@ -2052,8 +2052,9 @@ FTextureRHIRef FD3D11DynamicRHI::RHICreateAliasedTexture(FTextureRHIRef& SrcText
 {
 	FD3D11Texture* SrcTexture = GetD3D11TextureFromRHITexture(SrcTextureRHI);
 	check(SrcTexture);
+	const FString Name = SrcTextureRHI->GetName().ToString() + TEXT("Alias");
 
-	return new FD3D11Texture(*SrcTexture, FD3D11Texture::CreateAlias);
+	return new FD3D11Texture(*SrcTexture, *Name, FD3D11Texture::CreateAlias);
 }
 
 void FD3D11DynamicRHI::RHICopyTexture(FRHITexture* SourceTextureRHI, FRHITexture* DestTextureRHI, const FRHICopyTextureInfo& CopyInfo)
diff --git a/Engine/Source/Runtime/Windows/D3D11RHI/Public/D3D11Resources.h b/Engine/Source/Runtime/Windows/D3D11RHI/Public/D3D11Resources.h
--- a/Engine/Source/Runtime/Windows/D3D11RHI/Public/D3D11Resources.h
+++ b/Engine/Source/Runtime/Windows/D3D11RHI/Public/D3D11Resources.h
@@ -157,7 +157,7 @@ public:
 	);
 
 	enum EAliasResourceParam { CreateAlias };
-	explicit FD3D11Texture(FD3D11Texture const& Other, EAliasResourceParam);
+	explicit FD3D11Texture(FD3D11Texture const& Other, const FString& Name, EAliasResourceParam);
 	void AliasResource(FD3D11Texture const& Other);
 
 	virtual ~FD3D11Texture();
