Remove TextureView path in GL RHI, as it's not supported in GL ES

diff --git a/Engine/Source/Runtime/OpenGLDrv/Private/OpenGL4.cpp b/Engine/Source/Runtime/OpenGLDrv/Private/OpenGL4.cpp
--- a/Engine/Source/Runtime/OpenGLDrv/Private/OpenGL4.cpp
+++ b/Engine/Source/Runtime/OpenGLDrv/Private/OpenGL4.cpp
@@ -17,7 +17,6 @@ GLint FOpenGL4::MaxComputeUAVUnits = -1;
 GLint FOpenGL4::MaxPixelUAVUnits = -1;
 
 bool FOpenGL4::bSupportsGPUMemoryInfo = false;
-bool FOpenGL4::bSupportsTextureView = true;
 
 void FOpenGL4::ProcessQueryGLInt()
 {
@@ -57,7 +56,6 @@ void FOpenGL4::ProcessExtensions( const FString& ExtensionsString )
 
 
 	bSupportsGPUMemoryInfo = ExtensionsString.Contains(TEXT("GL_NVX_gpu_memory_info"));
-	bSupportsTextureView = ExtensionsString.Contains(TEXT("GL_ARB_texture_view")) || (MajorVersion == 4 && MinorVersion >= 3) || (MajorVersion > 4);
 
 	//Process Queries after extensions to avoid queries that use functionality that might not be present
 	ProcessQueryGLInt();
diff --git a/Engine/Source/Runtime/OpenGLDrv/Private/OpenGLCommands.cpp b/Engine/Source/Runtime/OpenGLDrv/Private/OpenGLCommands.cpp
--- a/Engine/Source/Runtime/OpenGLDrv/Private/OpenGLCommands.cpp
+++ b/Engine/Source/Runtime/OpenGLDrv/Private/OpenGLCommands.cpp
@@ -498,9 +498,8 @@ void FOpenGLDynamicRHI::CachedSetupTextureStageInner(FOpenGLContextState& Contex
 	}
 	
 	// Use the texture SRV's LimitMip value to specify the mip available for sampling
-	// This requires SupportsTextureBaseLevel & is a fallback for SupportsTextureView
-	// which should be preferred.
-	if(Target != GL_NONE && Target != GL_TEXTURE_BUFFER && Target != GL_TEXTURE_EXTERNAL_OES && !FOpenGL::SupportsTextureView())
+	// This requires SupportsTextureBaseLevel & is a fallback for TextureView
+	if (Target != GL_NONE && Target != GL_TEXTURE_BUFFER && Target != GL_TEXTURE_EXTERNAL_OES)
 	{
 		TPair<GLenum, GLenum>* MipLimits;
 		
@@ -638,9 +637,7 @@ void FOpenGLDynamicRHI::SetupTexturesForDraw( FOpenGLContextState& ContextState,
 			
 #if UE_BUILD_DEBUG
 			// Use the texture SRV's LimitMip value to specify the mip available for sampling
-			// This requires SupportsTextureBaseLevel & is a fallback for SupportsTextureView
-			// which should be preferred.
-			if(!FOpenGL::SupportsTextureView())
+			// This requires SupportsTextureBaseLevel & is a fallback for TextureView
 			{
 				// When trying to limit the mip available for sampling (as part of texture SRV)
 				// ensure that the texture is bound to only one sampler, or that all samplers
diff --git a/Engine/Source/Runtime/OpenGLDrv/Private/OpenGLTexture.cpp b/Engine/Source/Runtime/OpenGLDrv/Private/OpenGLTexture.cpp
--- a/Engine/Source/Runtime/OpenGLDrv/Private/OpenGLTexture.cpp
+++ b/Engine/Source/Runtime/OpenGLDrv/Private/OpenGLTexture.cpp
@@ -1660,72 +1660,25 @@ FOpenGLShaderResourceView::FOpenGLShaderResourceView(FOpenGLTexture* InTexture,
 	, Target      (InTexture->Target)
 	, Texture     (InTexture)
 	, LimitMip    (CreateInfo.MipLevel)
-	, OwnsResource(FOpenGL::SupportsTextureView())
+	, OwnsResource(false)
 {
-	check(CreateInfo.NumMipLevels == 1);
-
-	if (OwnsResource)
-	{
-		FOpenGL::GenTextures(1, &Resource);
-	}
-
-	const uint32 MipLevel = CreateInfo.MipLevel;
-	const uint32 NumMipLevels = CreateInfo.NumMipLevels;
 	const uint8 Format = (CreateInfo.Format == PF_Unknown)
 		? Texture->GetFormat()
 		: CreateInfo.Format;
-
+	
 	const bool bFormatsMatch = Format == Texture->GetFormat();
-	const bool bSRGB = CreateInfo.SRGBOverride == SRGBO_Default
-		? EnumHasAnyFlags(InTexture->GetDesc().Flags, TexCreate_SRGB)
-		: false;
 
 	if (Format == PF_X24_G8)
 	{
-		RunOnGLRenderContextThread([this, MipLevel, NumMipLevels]()
-		{
-			VERIFY_GL_SCOPE();
-			if (FOpenGL::SupportsTextureView())
-			{
-				// PF_X24_G8 doesn't correspond to a real format under OpenGL
-				// The solution is to create a view with the original format, and convert it to return the stencil index
-				// To match component locations, texture swizzle needs to be setup too
-				const FOpenGLTextureFormat& GLFormat = GOpenGLTextureFormats[Texture->GetFormat()];
-
-				// create a second depth/stencil view
-				FOpenGL::TextureView(Resource, Target, Texture->GetResource(), GLFormat.InternalFormat[0], MipLevel, NumMipLevels, 0, 1);
-
-				// Use a texture stage that's not likely to be used for draws, to avoid waiting
-				FOpenGLContextState& ContextState = FOpenGLDynamicRHI::Get().GetContextStateForCurrentContext();
-				FOpenGLDynamicRHI::Get().CachedSetupTextureStage(ContextState, FOpenGL::GetMaxCombinedTextureImageUnits() - 1, Target, Resource, 0, NumMipLevels);
-
-				//set the texture to return the stencil index, and then force the components to match D3D
-				glTexParameteri(Target, GL_DEPTH_STENCIL_TEXTURE_MODE, GL_STENCIL_INDEX);
-				glTexParameteri(Target, GL_TEXTURE_SWIZZLE_R, GL_ZERO);
-				glTexParameteri(Target, GL_TEXTURE_SWIZZLE_G, GL_RED);
-				glTexParameteri(Target, GL_TEXTURE_SWIZZLE_B, GL_ZERO);
-				glTexParameteri(Target, GL_TEXTURE_SWIZZLE_A, GL_ZERO);
-			}
-			else
-			{
-				UE_LOG(LogRHI, Fatal, TEXT("Cannot create a stencil SRV (PF_X24_G8) when texture views are unsupported."));
-			}
-		});
+		UE_LOG(LogRHI, Fatal, TEXT("Cannot create a stencil SRV (PF_X24_G8) when texture views are unsupported."));
 	}
 	else
 	{
-		RunOnGLRenderContextThread([this, Format, MipLevel, bSRGB, bFormatsMatch]()
+		RunOnGLRenderContextThread([this, bFormatsMatch]()
 		{
 			VERIFY_GL_SCOPE();
-			if (FOpenGL::SupportsTextureView())
-			{
-				FOpenGL::TextureView(Resource, Target, Texture->GetResource(), GOpenGLTextureFormats[Format].InternalFormat[bSRGB], MipLevel, 1, 0, 1);
-			}
-			else
-			{
-				checkf(bFormatsMatch, TEXT("SRVs cannot modify the pixel format of a texture when texture views are unsupported."));
-				Resource = Texture->GetResource();
-			}
+			checkf(bFormatsMatch, TEXT("SRVs cannot modify the pixel format of a texture when texture views are unsupported."));
+			Resource = Texture->GetResource();
 		});
 	}
 }
diff --git a/Engine/Source/Runtime/OpenGLDrv/Public/OpenGL.h b/Engine/Source/Runtime/OpenGLDrv/Public/OpenGL.h
--- a/Engine/Source/Runtime/OpenGLDrv/Public/OpenGL.h
+++ b/Engine/Source/Runtime/OpenGLDrv/Public/OpenGL.h
@@ -123,7 +123,6 @@ public:
 	static FORCEINLINE bool SupportsSeparateAlphaBlend()				{ return bSupportsDrawBuffersBlend; }
 	static FORCEINLINE void EnableSupportsClipControl()					{ bSupportsClipControl = true; }
 	static FORCEINLINE bool SupportsClipControl()						{ return bSupportsClipControl; }
-	static FORCEINLINE bool SupportsTextureView()						{ return false; }
 	static FORCEINLINE bool SupportsSeamlessCubeMap()					{ return false; }
 	static FORCEINLINE bool SupportsDrawIndirect()						{ return false; }
 	static FORCEINLINE bool SupportsGenerateMipmap()					{ return false; }
diff --git a/Engine/Source/Runtime/OpenGLDrv/Public/OpenGL4.h b/Engine/Source/Runtime/OpenGLDrv/Public/OpenGL4.h
--- a/Engine/Source/Runtime/OpenGLDrv/Public/OpenGL4.h
+++ b/Engine/Source/Runtime/OpenGLDrv/Public/OpenGL4.h
@@ -15,8 +15,6 @@
 struct FOpenGL4 : public FOpenGL3
 {
 	static FORCEINLINE bool SupportsDrawIndirect()						{ return true; }
-	static FORCEINLINE bool SupportsTextureView()						{ return bSupportsTextureView; }
-
 
 	// Optional
 
@@ -124,6 +122,4 @@ protected:
 	static GLint MaxPixelUAVUnits;
 
 	static bool bSupportsGPUMemoryInfo;
-	static bool bSupportsTextureView;
-
 };
