Log value of GL_MAX_SHADER_PIXEL_LOCAL_STORAGE_SIZE_EX

diff --git a/Engine/Source/Runtime/OpenGLDrv/Private/OpenGLDevice.cpp b/Engine/Source/Runtime/OpenGLDrv/Private/OpenGLDevice.cpp
--- a/Engine/Source/Runtime/OpenGLDrv/Private/OpenGLDevice.cpp
+++ b/Engine/Source/Runtime/OpenGLDrv/Private/OpenGLDevice.cpp
@@ -742,11 +742,16 @@ static void InitRHICapabilitiesForGL()
 		GSupportsQuadBufferStereo = (Result == GL_TRUE);
 	}
 
-	if( FOpenGL::SupportsTextureFilterAnisotropic())
+	if (FOpenGL::SupportsTextureFilterAnisotropic())
 	{
 		LOG_AND_GET_GL_INT_TEMP(GL_MAX_TEXTURE_MAX_ANISOTROPY_EXT, 0);
 		GMaxOpenGLTextureFilterAnisotropic = Value_GL_MAX_TEXTURE_MAX_ANISOTROPY_EXT;
 	}
+
+	if (FOpenGL::SupportsPixelLocalStorage())
+	{
+		LOG_AND_GET_GL_INT_TEMP(GL_MAX_SHADER_PIXEL_LOCAL_STORAGE_SIZE_EXT, 0);
+	}
 #undef LOG_AND_GET_GL_INT_TEMP
 
 	GMaxOpenGLColorSamples = Value_GL_MAX_COLOR_TEXTURE_SAMPLES;
diff --git a/Engine/Source/Runtime/OpenGLDrv/Public/OpenGL.h b/Engine/Source/Runtime/OpenGLDrv/Public/OpenGL.h
--- a/Engine/Source/Runtime/OpenGLDrv/Public/OpenGL.h
+++ b/Engine/Source/Runtime/OpenGLDrv/Public/OpenGL.h
@@ -636,4 +636,10 @@ protected:
 
 #ifndef TEXTURE_ASTC_DECODE_PRECISION_EXT
 #define TEXTURE_ASTC_DECODE_PRECISION_EXT 0x8F69
+#endif
+
+#ifndef GL_EXT_shader_pixel_local_storage
+#define GL_MAX_SHADER_PIXEL_LOCAL_STORAGE_FAST_SIZE_EXT 0x8F63
+#define GL_MAX_SHADER_PIXEL_LOCAL_STORAGE_SIZE_EXT 0x8F67
+#define GL_SHADER_PIXEL_LOCAL_STORAGE_EXT 0x8F64
 #endif
\ No newline at end of file
