Added better state tracking for a case where both real UBs and eUBs are used together in GL

diff --git a/Engine/Source/Runtime/OpenGLDrv/Private/OpenGLCommands.cpp b/Engine/Source/Runtime/OpenGLDrv/Private/OpenGLCommands.cpp
--- a/Engine/Source/Runtime/OpenGLDrv/Private/OpenGLCommands.cpp
+++ b/Engine/Source/Runtime/OpenGLDrv/Private/OpenGLCommands.cpp
@@ -1073,6 +1073,11 @@ void FOpenGLDynamicRHI::RHISetShaderUniformBuffer(FRHIGraphicsShader* ShaderRHI,
 		PendingState.BoundUniformBuffers[Stage][BufferIndex] = BufferRHI;
 		PendingState.DirtyUniformBuffers[Stage] |= 1 << BufferIndex;
 		PendingState.bAnyDirtyGraphicsUniformBuffers = true;
+		
+		if (!GUseEmulatedUniformBuffers || !((FOpenGLUniformBuffer*)BufferRHI)->bIsEmulatedUniformBuffer)
+		{
+			PendingState.bAnyDirtyRealUniformBuffers[Stage] = true;
+		}
 	}
 }
 
@@ -1088,6 +1093,11 @@ void FOpenGLDynamicRHI::RHISetShaderUniformBuffer(FRHIComputeShader* ComputeShad
 	VERIFY_GL_SCOPE();
 	PendingState.BoundUniformBuffers[SF_Compute][BufferIndex] = BufferRHI;
 	PendingState.DirtyUniformBuffers[SF_Compute] |= 1 << BufferIndex;
+
+	if (!GUseEmulatedUniformBuffers || !((FOpenGLUniformBuffer*)BufferRHI)->bIsEmulatedUniformBuffer)
+	{
+		PendingState.bAnyDirtyRealUniformBuffers[SF_Compute] = true;
+	}
 }
 
 void FOpenGLDynamicRHI::RHISetShaderParameter(FRHIGraphicsShader* ShaderRHI,uint32 BufferIndex,uint32 BaseIndex,uint32 NumBytes,const void* NewValue)
@@ -1996,54 +2006,35 @@ void FOpenGLDynamicRHI::OnUniformBufferDeletion( GLuint UniformBufferResource, u
 	}
 }
 
-FORCEINLINE void FOpenGLDynamicRHI::CommitNonComputeShaderConstants()
+void FOpenGLDynamicRHI::CommitNonComputeShaderConstants()
 {
+	VERIFY_GL_SCOPE();
+
+	FOpenGLLinkedProgram* LinkedProgram = PendingState.BoundShaderState->LinkedProgram;
 	if (GUseEmulatedUniformBuffers)
 	{
-		CommitNonComputeShaderConstantsSlowPath();
+		PendingState.ShaderParameters[CrossCompiler::SHADER_STAGE_VERTEX].CommitPackedUniformBuffers(LinkedProgram, CrossCompiler::SHADER_STAGE_VERTEX, PendingState.BoundUniformBuffers[SF_Vertex], PendingState.BoundShaderState->GetVertexShader()->UniformBuffersCopyInfo);
+		PendingState.ShaderParameters[CrossCompiler::SHADER_STAGE_PIXEL].CommitPackedUniformBuffers(LinkedProgram, CrossCompiler::SHADER_STAGE_PIXEL, PendingState.BoundUniformBuffers[SF_Pixel], PendingState.BoundShaderState->GetPixelShader()->UniformBuffersCopyInfo);
+		if (PendingState.BoundShaderState->GetGeometryShader())
+		{
+			PendingState.ShaderParameters[CrossCompiler::SHADER_STAGE_GEOMETRY].CommitPackedUniformBuffers(LinkedProgram, CrossCompiler::SHADER_STAGE_GEOMETRY, PendingState.BoundUniformBuffers[SF_Geometry], PendingState.BoundShaderState->GetGeometryShader()->UniformBuffersCopyInfo);
+		}
 	}
 	
-	FOpenGLLinkedProgram* LinkedProgram = PendingState.BoundShaderState->LinkedProgram;
 	if (LinkedProgram == PendingState.LinkedProgramAndDirtyFlag)
 	{
 		return;
 	}
-	CommitNonComputeShaderConstantsFastPath(LinkedProgram);
-	PendingState.LinkedProgramAndDirtyFlag = LinkedProgram;
-}
-
-
-// Only used for emulated uniform buffers
-void FOpenGLDynamicRHI::CommitNonComputeShaderConstantsSlowPath()
-{
-	VERIFY_GL_SCOPE();
-	PendingState.LinkedProgramAndDirtyFlag = nullptr;
-
-	FOpenGLLinkedProgram* LinkedProgram = PendingState.BoundShaderState->LinkedProgram;
 	
-	PendingState.ShaderParameters[CrossCompiler::SHADER_STAGE_VERTEX].CommitPackedUniformBuffers(LinkedProgram, CrossCompiler::SHADER_STAGE_VERTEX, PendingState.BoundUniformBuffers[SF_Vertex], PendingState.BoundShaderState->GetVertexShader()->UniformBuffersCopyInfo);
+	// commit packed global only if current program has changed or any global parameter has changed (RHISetShaderParameter)
 	PendingState.ShaderParameters[CrossCompiler::SHADER_STAGE_VERTEX].CommitPackedGlobals(LinkedProgram, CrossCompiler::SHADER_STAGE_VERTEX);
-
-	PendingState.ShaderParameters[CrossCompiler::SHADER_STAGE_PIXEL].CommitPackedUniformBuffers(LinkedProgram, CrossCompiler::SHADER_STAGE_PIXEL, PendingState.BoundUniformBuffers[SF_Pixel], PendingState.BoundShaderState->GetPixelShader()->UniformBuffersCopyInfo);
 	PendingState.ShaderParameters[CrossCompiler::SHADER_STAGE_PIXEL].CommitPackedGlobals(LinkedProgram, CrossCompiler::SHADER_STAGE_PIXEL);
-
 	if (PendingState.BoundShaderState->GetGeometryShader())
 	{
-		PendingState.ShaderParameters[CrossCompiler::SHADER_STAGE_GEOMETRY].CommitPackedUniformBuffers(LinkedProgram, CrossCompiler::SHADER_STAGE_GEOMETRY, PendingState.BoundUniformBuffers[SF_Geometry], PendingState.BoundShaderState->GetGeometryShader()->UniformBuffersCopyInfo);
 		PendingState.ShaderParameters[CrossCompiler::SHADER_STAGE_GEOMETRY].CommitPackedGlobals(LinkedProgram, CrossCompiler::SHADER_STAGE_GEOMETRY);
 	}
-}
 
-void FOpenGLDynamicRHI::CommitNonComputeShaderConstantsFastPath(FOpenGLLinkedProgram* LinkedProgram)
-{
-	VERIFY_GL_SCOPE();
-	PendingState.ShaderParameters[CrossCompiler::SHADER_STAGE_VERTEX].CommitPackedGlobals(LinkedProgram, CrossCompiler::SHADER_STAGE_VERTEX);
-	PendingState.ShaderParameters[CrossCompiler::SHADER_STAGE_PIXEL].CommitPackedGlobals(LinkedProgram, CrossCompiler::SHADER_STAGE_PIXEL);
-
-	if (PendingState.BoundShaderState->GetGeometryShader())
-	{
-		PendingState.ShaderParameters[CrossCompiler::SHADER_STAGE_GEOMETRY].CommitPackedGlobals(LinkedProgram, CrossCompiler::SHADER_STAGE_GEOMETRY);
-	}
+	PendingState.LinkedProgramAndDirtyFlag = LinkedProgram;
 }
 
 void FOpenGLDynamicRHI::CommitComputeShaderConstants(FOpenGLComputeShader* ComputeShader)
diff --git a/Engine/Source/Runtime/OpenGLDrv/Private/OpenGLShaders.cpp b/Engine/Source/Runtime/OpenGLDrv/Private/OpenGLShaders.cpp
--- a/Engine/Source/Runtime/OpenGLDrv/Private/OpenGLShaders.cpp
+++ b/Engine/Source/Runtime/OpenGLDrv/Private/OpenGLShaders.cpp
@@ -1223,7 +1223,9 @@ void FOpenGLDynamicRHI::BindUniformBufferBase(FOpenGLContextState& ContextState,
 			Buffer = GLUB->Resource;
 
 			if (GLUB->bIsEmulatedUniformBuffer)
+			{
 				continue;
+			}
 
 			Size = GLUB->GetSize();
 #if SUBALLOCATED_CONSTANT_BUFFER
@@ -3638,44 +3640,58 @@ void FOpenGLDynamicRHI::BindPendingShaderState( FOpenGLContextState& ContextStat
 		ForceUniformBindingUpdate = FOpenGL::SupportsSeparateShaderObjects();
 	}
 
-	int32 NextUniformBufferIndex = OGL_FIRST_UNIFORM_BUFFER;
+	if (PendingState.bAnyDirtyRealUniformBuffers[SF_Vertex] || 
+		PendingState.bAnyDirtyRealUniformBuffers[SF_Pixel] || 
+		PendingState.bAnyDirtyRealUniformBuffers[SF_Geometry])
+	{
+		int32 NextUniformBufferIndex = OGL_FIRST_UNIFORM_BUFFER;
 
-	static_assert(SF_NumGraphicsFrequencies == 5 && SF_NumFrequencies == 10, "Unexpected SF_ ordering");
-	static_assert(SF_RayGen > SF_NumGraphicsFrequencies, "SF_NumGraphicsFrequencies be the number of frequencies supported in OpenGL");
+		static_assert(SF_NumGraphicsFrequencies == 5 && SF_NumFrequencies == 10, "Unexpected SF_ ordering");
+		static_assert(SF_RayGen > SF_NumGraphicsFrequencies, "SF_NumGraphicsFrequencies be the number of frequencies supported in OpenGL");
 
-	int32 NumUniformBuffers[SF_NumGraphicsFrequencies];
+		int32 NumUniformBuffers[SF_NumGraphicsFrequencies];
 
-	PendingState.BoundShaderState->GetNumUniformBuffers(NumUniformBuffers);
-			
-	PendingState.BoundShaderState->LinkedProgram->VerifyUniformBlockBindings(CrossCompiler::SHADER_STAGE_VERTEX, NextUniformBufferIndex);
-
-	BindUniformBufferBase(
-		ContextState,
-		NumUniformBuffers[SF_Vertex],
-		PendingState.BoundUniformBuffers[SF_Vertex],
-		NextUniformBufferIndex,
-		ForceUniformBindingUpdate);
-	NextUniformBufferIndex += NumUniformBuffers[SF_Vertex];
-
-	PendingState.BoundShaderState->LinkedProgram->VerifyUniformBlockBindings(CrossCompiler::SHADER_STAGE_PIXEL, NextUniformBufferIndex);
-	BindUniformBufferBase(
-		ContextState,
-		NumUniformBuffers[SF_Pixel],
-		PendingState.BoundUniformBuffers[SF_Pixel],
-		NextUniformBufferIndex,
-		ForceUniformBindingUpdate);
-	NextUniformBufferIndex += NumUniformBuffers[SF_Pixel];
-
-	if (NumUniformBuffers[SF_Geometry] >= 0)
-	{
-		PendingState.BoundShaderState->LinkedProgram->VerifyUniformBlockBindings(CrossCompiler::SHADER_STAGE_GEOMETRY, NextUniformBufferIndex);
-		BindUniformBufferBase(
-			ContextState,
-			NumUniformBuffers[SF_Geometry],
-			PendingState.BoundUniformBuffers[SF_Geometry],
-			NextUniformBufferIndex,
-			ForceUniformBindingUpdate);
-		NextUniformBufferIndex += NumUniformBuffers[SF_Geometry];
+		PendingState.BoundShaderState->GetNumUniformBuffers(NumUniformBuffers);
+
+		if (PendingState.bAnyDirtyRealUniformBuffers[SF_Vertex])
+		{
+			PendingState.BoundShaderState->LinkedProgram->VerifyUniformBlockBindings(CrossCompiler::SHADER_STAGE_VERTEX, NextUniformBufferIndex);
+			BindUniformBufferBase(
+				ContextState,
+				NumUniformBuffers[SF_Vertex],
+				PendingState.BoundUniformBuffers[SF_Vertex],
+				NextUniformBufferIndex,
+				ForceUniformBindingUpdate);
+		}
+		NextUniformBufferIndex += NumUniformBuffers[SF_Vertex];
+
+		if (PendingState.bAnyDirtyRealUniformBuffers[SF_Pixel])
+		{
+			PendingState.BoundShaderState->LinkedProgram->VerifyUniformBlockBindings(CrossCompiler::SHADER_STAGE_PIXEL, NextUniformBufferIndex);
+			BindUniformBufferBase(
+				ContextState,
+				NumUniformBuffers[SF_Pixel],
+				PendingState.BoundUniformBuffers[SF_Pixel],
+				NextUniformBufferIndex,
+				ForceUniformBindingUpdate);
+		}
+		NextUniformBufferIndex += NumUniformBuffers[SF_Pixel];
+
+		if (NumUniformBuffers[SF_Geometry] >= 0 && PendingState.bAnyDirtyRealUniformBuffers[SF_Geometry])
+		{
+			PendingState.BoundShaderState->LinkedProgram->VerifyUniformBlockBindings(CrossCompiler::SHADER_STAGE_GEOMETRY, NextUniformBufferIndex);
+			BindUniformBufferBase(
+				ContextState,
+				NumUniformBuffers[SF_Geometry],
+				PendingState.BoundUniformBuffers[SF_Geometry],
+				NextUniformBufferIndex,
+				ForceUniformBindingUpdate);
+			NextUniformBufferIndex += NumUniformBuffers[SF_Geometry];
+		}
+
+		PendingState.bAnyDirtyRealUniformBuffers[SF_Vertex] = false;
+		PendingState.bAnyDirtyRealUniformBuffers[SF_Pixel] = false;
+		PendingState.bAnyDirtyRealUniformBuffers[SF_Geometry] = false;
 	}
 
 	if (FOpenGL::SupportsBindlessTexture())
@@ -3841,13 +3857,18 @@ void FOpenGLDynamicRHI::BindPendingComputeShaderState(FOpenGLContextState& Conte
 		ForceUniformBindingUpdate = true;
 	}
 
-	ComputeShader->LinkedProgram->VerifyUniformBlockBindings(CrossCompiler::SHADER_STAGE_COMPUTE, OGL_FIRST_UNIFORM_BUFFER);
-	BindUniformBufferBase(
-		ContextState,
-		ComputeShader->Bindings.NumUniformBuffers,
-		PendingState.BoundUniformBuffers[SF_Compute],
-		OGL_FIRST_UNIFORM_BUFFER,
-		ForceUniformBindingUpdate);
+	if (PendingState.bAnyDirtyRealUniformBuffers[SF_Compute])
+	{
+		ComputeShader->LinkedProgram->VerifyUniformBlockBindings(CrossCompiler::SHADER_STAGE_COMPUTE, OGL_FIRST_UNIFORM_BUFFER);
+		BindUniformBufferBase(
+			ContextState,
+			ComputeShader->Bindings.NumUniformBuffers,
+			PendingState.BoundUniformBuffers[SF_Compute],
+			OGL_FIRST_UNIFORM_BUFFER,
+			ForceUniformBindingUpdate);
+
+		PendingState.bAnyDirtyRealUniformBuffers[SF_Compute] = 0;
+	}
 	SetupBindlessTextures( ContextState, ComputeShader->LinkedProgram->Samplers );
 }
 
@@ -4014,7 +4035,9 @@ void FOpenGLShaderParameterCache::CommitPackedUniformBuffers(FOpenGLLinkedProgra
 			check(UniformBuffer);
 
 			if (!UniformBuffer->bIsEmulatedUniformBuffer)
+			{
 				continue;
+			}
 
 			const uint32* RESTRICT SourceData = UniformBuffer->EmulatedBufferData->Data.GetData();
 			for (int32 InfoIndex = LastInfoIndex; InfoIndex < UniformBuffersCopyInfo.Num(); ++InfoIndex)
diff --git a/Engine/Source/Runtime/OpenGLDrv/Private/OpenGLState.cpp b/Engine/Source/Runtime/OpenGLDrv/Private/OpenGLState.cpp
--- a/Engine/Source/Runtime/OpenGLDrv/Private/OpenGLState.cpp
+++ b/Engine/Source/Runtime/OpenGLDrv/Private/OpenGLState.cpp
@@ -488,6 +488,7 @@ void FOpenGLRHIState::InitializeResources(int32 NumCombinedTextures, int32 NumCo
 	for (int32 Frequency = 0; Frequency < SF_NumStandardFrequencies; ++Frequency)
 	{
 		DirtyUniformBuffers[Frequency] = MAX_uint16;
+		bAnyDirtyRealUniformBuffers[Frequency] = true;
 	}
 	bAnyDirtyGraphicsUniformBuffers = true;
 }
diff --git a/Engine/Source/Runtime/OpenGLDrv/Public/OpenGLDrv.h b/Engine/Source/Runtime/OpenGLDrv/Public/OpenGLDrv.h
--- a/Engine/Source/Runtime/OpenGLDrv/Public/OpenGLDrv.h
+++ b/Engine/Source/Runtime/OpenGLDrv/Public/OpenGLDrv.h
@@ -1015,8 +1015,6 @@ private:
 	void CommitGraphicsResourceTablesInner();
 	void CommitComputeResourceTables(FOpenGLComputeShader* ComputeShader);
 	void CommitNonComputeShaderConstants();
-	void CommitNonComputeShaderConstantsSlowPath();
-	void CommitNonComputeShaderConstantsFastPath(FOpenGLLinkedProgram* LinkedProgram);
 	void CommitComputeShaderConstants(FOpenGLComputeShader* ComputeShader);
 	void SetPendingBlendStateForActiveRenderTargets( FOpenGLContextState& ContextState );
 	
diff --git a/Engine/Source/Runtime/OpenGLDrv/Public/OpenGLState.h b/Engine/Source/Runtime/OpenGLDrv/Public/OpenGLState.h
--- a/Engine/Source/Runtime/OpenGLDrv/Public/OpenGLState.h
+++ b/Engine/Source/Runtime/OpenGLDrv/Public/OpenGLState.h
@@ -466,6 +466,8 @@ struct FOpenGLRHIState final : public FOpenGLCommonState
 	/** Track the currently bound uniform buffers. */
 	FUniformBufferRHIRef BoundUniformBuffers[SF_NumStandardFrequencies][MAX_UNIFORM_BUFFERS_PER_SHADER_STAGE];
 
+	/** Array to track if any real (not emulated) uniform buffers have been bound since the last draw call */
+	bool bAnyDirtyRealUniformBuffers[SF_NumStandardFrequencies];
 	/** Bit array to track which uniform buffers have changed since the last draw call. */
 	bool bAnyDirtyGraphicsUniformBuffers;
 	uint16 DirtyUniformBuffers[SF_NumStandardFrequencies];
