Removed support for separate shader objects from GL RHI as it was never used for GLES

diff --git a/Engine/Source/Developer/ShaderFormatOpenGL/Private/OpenGLShaderCompiler.cpp b/Engine/Source/Developer/ShaderFormatOpenGL/Private/OpenGLShaderCompiler.cpp
--- a/Engine/Source/Developer/ShaderFormatOpenGL/Private/OpenGLShaderCompiler.cpp
+++ b/Engine/Source/Developer/ShaderFormatOpenGL/Private/OpenGLShaderCompiler.cpp
@@ -82,15 +82,6 @@ static FORCEINLINE bool IsPCESPlatform(GLSLVersion Version)
 	return (Version == GLSL_150_ES3_1);
 }
 
-// This function should match OpenGLShaderPlatformSeparable
-bool FOpenGLFrontend::SupportsSeparateShaderObjects(GLSLVersion Version)
-{
-	// Only desktop shader platforms can use separable shaders for now,
-	// the generated code relies on macros supplied at runtime to determine whether
-	// shaders may be separable and/or linked.
-	return Version == GLSL_150_ES3_1;
-}
-
 /*------------------------------------------------------------------------------
 	Shader compiling.
 ------------------------------------------------------------------------------*/
@@ -1163,11 +1154,6 @@ uint32 FOpenGLFrontend::CalculateCrossCompilerFlags(GLSLVersion Version, const b
 		CCFlags |= HLSLCC_ExpandUBMemberArrays;
 	}
 
-	if (SupportsSeparateShaderObjects(Version))
-	{
-		CCFlags |= HLSLCC_SeparateShaderObjects;
-	}
-
 	if (CompilerFlags.Contains(CFLAG_UsesExternalTexture))
 	{
 		CCFlags |= HLSLCC_UsesExternalTexture;
@@ -3320,7 +3306,6 @@ enum class EPlatformType
 struct FDeviceCapabilities
 {
 	EPlatformType TargetPlatform = EPlatformType::Android;
-	bool bSupportsSeparateShaderObjects;
 };
 
 void FOpenGLFrontend::FillDeviceCapsOfflineCompilation(struct FDeviceCapabilities& Capabilities, const GLSLVersion ShaderVersion) const
@@ -3334,7 +3319,6 @@ void FOpenGLFrontend::FillDeviceCapsOfflineCompilation(struct FDeviceCapabilitie
 	else
 	{
 		Capabilities.TargetPlatform = EPlatformType::Desktop;
-		Capabilities.bSupportsSeparateShaderObjects = true;
 	}
 }
 
@@ -3380,7 +3364,7 @@ inline bool OpenGLShaderPlatformSeparable(const GLSLVersion InShaderPlatform)
 	switch (InShaderPlatform)
 	{
 		case GLSL_150_ES3_1:
-			return true;
+			return false;
 
 		case GLSL_ES3_1_ANDROID:
 			return false;
@@ -3439,7 +3423,6 @@ TSharedPtr<ANSICHAR> FOpenGLFrontend::PrepareCodeForOfflineCompilation(const GLS
 	// OpenGL SM5 shader platforms require location declarations for the layout, but don't necessarily use SSOs
 	if (Capabilities.TargetPlatform == EPlatformType::Desktop)
 	{
-		StrOutSource.Append(TEXT("#extension GL_ARB_separate_shader_objects : enable\n"));
 		StrOutSource.Append(TEXT("#define INTERFACE_BLOCK(Pos, Interp, Modifiers, Semantic, PreType, PostType) layout(location=Pos) Interp Modifiers struct { PreType PostType; }\n"));
 	}
 	else
diff --git a/Engine/Source/Developer/ShaderFormatOpenGL/Private/ShaderFormatOpenGL.cpp b/Engine/Source/Developer/ShaderFormatOpenGL/Private/ShaderFormatOpenGL.cpp
--- a/Engine/Source/Developer/ShaderFormatOpenGL/Private/ShaderFormatOpenGL.cpp
+++ b/Engine/Source/Developer/ShaderFormatOpenGL/Private/ShaderFormatOpenGL.cpp
@@ -17,7 +17,7 @@ class FShaderFormatGLSL : public IShaderFormat
 	enum
 	{
 		/** Version for shader format, this becomes part of the DDC key. */
-		UE_SHADER_GLSL_VER = 103,
+		UE_SHADER_GLSL_VER = 104,
 	};
 
 	void CheckFormat(FName Format) const
diff --git a/Engine/Source/Developer/ShaderFormatOpenGL/Public/ShaderFormatOpenGL.h b/Engine/Source/Developer/ShaderFormatOpenGL/Public/ShaderFormatOpenGL.h
--- a/Engine/Source/Developer/ShaderFormatOpenGL/Public/ShaderFormatOpenGL.h
+++ b/Engine/Source/Developer/ShaderFormatOpenGL/Public/ShaderFormatOpenGL.h
@@ -43,9 +43,6 @@ public:
 	void CompileShader(const struct FShaderCompilerInput& Input, struct FShaderCompilerOutput& Output, const class FString& WorkingDirectory, GLSLVersion Version);
 
 protected:
-	// does the given version support SSO? Allows subclasses to device based on it's own format
-	virtual bool SupportsSeparateShaderObjects(GLSLVersion Version);
-
 	// if true, the shader output map will contain true names (i.e. ColorModifier) instead of helper names for runtime binding (i.e. pb_5)
 	virtual bool OutputTrueParameterNames()
 	{
diff --git a/Engine/Source/Runtime/OpenGLDrv/Private/OpenGL3.cpp b/Engine/Source/Runtime/OpenGLDrv/Private/OpenGL3.cpp
--- a/Engine/Source/Runtime/OpenGLDrv/Private/OpenGL3.cpp
+++ b/Engine/Source/Runtime/OpenGLDrv/Private/OpenGL3.cpp
@@ -11,7 +11,6 @@
 
 #if OPENGL_GL3
 
-bool FOpenGL3::bSupportsSeparateShaderObjects = false;
 bool FOpenGL3::bAndroidGLESCompatibilityMode = false;
 
 GLsizei FOpenGL3::NextTextureName = OPENGL_NAME_CACHE_SIZE;
@@ -77,11 +76,6 @@ void FOpenGL3::ProcessExtensions( const FString& ExtensionsString )
 		glDeleteTextures(1, &VolumeTexture);
 		glDeleteFramebuffers(1, &FrameBuffer);
 	}
-	
-	static const auto CVar = IConsoleManager::Get().FindTConsoleVariableDataInt(TEXT("OpenGL.UseSeparateShaderObjects"));
-	bool const bUseSeparateShaderObjects = (CVar ? (CVar->GetValueOnRenderThread() == 1) : false) && OpenGLShaderPlatformSeparable(GetShaderPlatform());
-	
-	bSupportsSeparateShaderObjects = bUseSeparateShaderObjects && (ExtensionsString.Contains(TEXT("GL_ARB_separate_shader_objects")) || (MajorVersion == 4 && MinorVersion >= 4));
 
 	bAndroidGLESCompatibilityMode = GetFeatureLevel() == ERHIFeatureLevel::ES3_1 && ExtensionsString.Contains(TEXT("GL_ARB_ES3_1_compatibility")) && FParse::Param(FCommandLine::Get(), TEXT("GLESCompat"));
 }
diff --git a/Engine/Source/Runtime/OpenGLDrv/Private/OpenGLCommands.cpp b/Engine/Source/Runtime/OpenGLDrv/Private/OpenGLCommands.cpp
--- a/Engine/Source/Runtime/OpenGLDrv/Private/OpenGLCommands.cpp
+++ b/Engine/Source/Runtime/OpenGLDrv/Private/OpenGLCommands.cpp
@@ -97,12 +97,6 @@ namespace OpenGLConsoleVariables
 		ECVF_ReadOnly
 	);
 
-	static TAutoConsoleVariable<int32> CVarUseSeparateShaderObjects(
-		TEXT("OpenGL.UseSeparateShaderObjects"),
-		0,
-		TEXT("If set to 1, use OpenGL's separate shader objects to eliminate expensive program linking"),
-		ECVF_ReadOnly|ECVF_RenderThreadSafe);
-
 	int32 GOpenGLForceBilinear = 0;
 	static FAutoConsoleVariableRef CVarOpenGLForceBilinearSampling(
 		TEXT("r.OpenGL.ForceBilinear"),
@@ -2240,44 +2234,6 @@ void FOpenGLDynamicRHI::CommitComputeResourceTables(FOpenGLComputeShader* Comput
 	PendingState.DirtyUniformBuffers[SF_Compute] = 0;
 }
 
-#if DEBUG_GL_SHADERS
-static void VerifyProgramPipeline()
-{
-	if (FOpenGL::SupportsSeparateShaderObjects())
-	{
-		VERIFY_GL_SCOPE();
-		GLint ProgramPipeline = 0;
-		glGetIntegerv(GL_PROGRAM_PIPELINE_BINDING, &ProgramPipeline);
-		if(ProgramPipeline)
-		{
-			FOpenGL::ValidateProgramPipeline(ProgramPipeline);
-			GLint LinkStatus = GL_FALSE;
-			FOpenGL::GetProgramPipelineiv(ProgramPipeline, GL_VALIDATE_STATUS, &LinkStatus);
-			if(LinkStatus == GL_FALSE)
-			{
-				GLint LogLength = 0;
-				FOpenGL::GetProgramPipelineiv(ProgramPipeline, GL_INFO_LOG_LENGTH, &LogLength);
-				ANSICHAR DefaultLog[] = "No log";
-				ANSICHAR *CompileLog = DefaultLog;
-				if (LogLength > 1)
-				{
-					CompileLog = (ANSICHAR *)FMemory::Malloc(LogLength);
-					FOpenGL::GetProgramPipelineInfoLog(ProgramPipeline, LogLength, NULL, CompileLog);
-				}
-				
-				UE_LOG(LogRHI,Error,TEXT("Failed to validate pipeline %d. Compile log:\n%s"), ProgramPipeline,
-					   ANSI_TO_TCHAR(CompileLog));
-				
-				if (LogLength > 1)
-				{
-					FMemory::Free(CompileLog);
-				}
-			}
-		}
-	}
-}
-#endif
-
 void FOpenGLDynamicRHI::RHIDrawPrimitive(uint32 BaseVertexIndex,uint32 NumPrimitives,uint32 NumInstances)
 {
 	SCOPE_CYCLE_COUNTER_DETAILED(STAT_OpenGLDrawPrimitiveTime);
@@ -2305,10 +2261,6 @@ void FOpenGLDynamicRHI::RHIDrawPrimitive(uint32 BaseVertexIndex,uint32 NumPrimit
 	GLint PatchSize = 0;
 	FindPrimitiveType(PrimitiveType, NumPrimitives, DrawMode, NumElements, PatchSize);
 
-#if DEBUG_GL_SHADERS
-	VerifyProgramPipeline();
-#endif
-
 	GPUProfilingData.RegisterGPUWork(NumPrimitives * NumInstances, VertexCount * NumInstances);
 	if (NumInstances == 1)
 	{
@@ -2499,10 +2451,6 @@ void FOpenGLDynamicRHI::RHIDrawIndexedPrimitive(FRHIBuffer* IndexBufferRHI, int3
 	GLenum IndexType = IndexBuffer->GetStride() == sizeof(uint32) ? GL_UNSIGNED_INT : GL_UNSIGNED_SHORT;
 	StartIndex *= IndexBuffer->GetStride() == sizeof(uint32) ? sizeof(uint32) : sizeof(uint16);
 
-#if DEBUG_GL_SHADERS
-	VerifyProgramPipeline();
-#endif
-
 	GPUProfilingData.RegisterGPUWork(NumPrimitives * NumInstances, NumElements * NumInstances);
 	if (NumInstances > 1)
 	{
diff --git a/Engine/Source/Runtime/OpenGLDrv/Private/OpenGLShaders.cpp b/Engine/Source/Runtime/OpenGLDrv/Private/OpenGLShaders.cpp
--- a/Engine/Source/Runtime/OpenGLDrv/Private/OpenGLShaders.cpp
+++ b/Engine/Source/Runtime/OpenGLDrv/Private/OpenGLShaders.cpp
@@ -229,6 +229,7 @@ FORCEINLINE void FOpenGLShaderParameterCache::FRange::MarkDirtyRange(uint32 NewS
 static bool VerifyLinkedProgram(GLuint Program)
 {
 	SCOPE_CYCLE_COUNTER(STAT_OpenGLShaderLinkVerifyTime);
+	VERIFY_GL_SCOPE();
 
 	GLint LinkStatus = 0;
 	glGetProgramiv(Program, GL_LINK_STATUS, &LinkStatus);
@@ -267,27 +268,6 @@ static bool VerifyLinkedProgram(GLuint Program)
 	return true;
 }
 
-
-// Verify a program has created successfully, the non-SSO case will log errors and return back success status.
-// TODO: SupportsSeparateShaderObjects case.
-static bool VerifyProgramPipeline(GLuint Program)
-{
-	VERIFY_GL_SCOPE();
-	bool bOK = true;
-	// Don't try and validate SSOs here - the draw state matters to SSOs and it definitely can't be guaranteed to be valid at this stage
-	if ( FOpenGL::SupportsSeparateShaderObjects() )
-	{
-#if DEBUG_GL_SHADERS
-		bOK = FOpenGL::IsProgramPipeline(Program);
-#endif
-	}
-	else
-	{
-		bOK = VerifyLinkedProgram(Program);
-	}
-	return bOK;
-}
-
 // ============================================================================================================================
 
 class FOpenGLCompiledShaderValue
@@ -580,19 +560,6 @@ static bool CompileCurrentShader(const GLuint Resource, const FAnsiCharArray& Gl
 
 	// Verify that an OpenGL shader has compiled successfully.
 	SCOPE_CYCLE_COUNTER(STAT_OpenGLShaderCompileVerifyTime);
-	
-	if (FOpenGL::SupportsSeparateShaderObjects() && glIsProgram(Resource))
-	{
-		bool const bCompiledOK = VerifyLinkedProgram(Resource);
-#if DEBUG_GL_SHADERS
-		if (!bCompiledOK && GlslCodeString)
-		{
-			UE_LOG(LogRHI,Error,TEXT("Shader:\n%s"), ANSI_TO_TCHAR(GlslCodeString));
-		}
-#endif
-		return bCompiledOK;
-	}
-	else
 	{
 		GLint CompileStatus;
 		glGetShaderiv(Resource, GL_COMPILE_STATUS, &CompileStatus);
@@ -638,21 +605,6 @@ static bool CompileCurrentShader(const GLuint Resource, const FAnsiCharArray& Gl
 	return true;
 }
 
-
-// Set the shader hash for FRHIShaders only.
-template<typename TRHIType>
-static typename TEnableIf<TPointerIsConvertibleFromTo<TRHIType, const FRHIShader>::Value, void>::Type SetShaderHash(const FSHAHash& Hash, TRHIType* Shader)
-{
-	Shader->SetHash(Hash);
-}
-
-template<typename TRHIType>
-static typename TEnableIf<!TPointerIsConvertibleFromTo<TRHIType, const FRHIShader>::Value, void>::Type SetShaderHash(const FSHAHash& Hash, TRHIType* Shader)
-{
-	// Cannot set the shader hash on a non-FRHIShader type.
-	checkNoEntry();
-}
-
 static const FOpenGLShaderDeviceCapabilities& GetOpenGLShaderDeviceCapabilities()
 {
 	static bool bInitialized = false;
@@ -694,20 +646,6 @@ ShaderType* CompileOpenGLShader(const FOpenGLCodeHeader& Header, const FOpenGLCo
 			Shader->Bindings = Val->Header->Bindings;
 			Shader->UniformBuffersCopyInfo = Val->Header->UniformBuffersCopyInfo;
 			Shader->StaticSlots = Val->StaticSlots;
-			if (FOpenGL::SupportsSeparateShaderObjects())
-			{
-				FSHAHash Hash;
-				// Just use the CRC - if it isn't being cached & logged we'll be dependent on the CRC alone anyway
-				FMemory::Memcpy(Hash.Hash, &Val->ShaderCrc, sizeof(uint32));
-				if (RHIShader)
-				{
-					SetShaderHash(Hash, RHIShader);
-				}
-				else
-				{
-					SetShaderHash(Hash, Shader);
-				}
-			}
 #if DEBUG_GL_SHADERS
 			Shader->GlslCode = Val->GlslCode;
 			Shader->GlslCodeString = (ANSICHAR*)Shader->GlslCode.GetData();
@@ -728,39 +666,18 @@ ShaderType* CompileOpenGLShader(const FOpenGLCodeHeader& Header, const FOpenGLCo
 		const FOpenGLShaderDeviceCapabilities& Capabilities = GetOpenGLShaderDeviceCapabilities();
 
 		// Save the code and defer compilation if our device supports program binaries and we're not checking for shader compatibility.
-		const bool bDeferredCompilation = FOpenGLProgramBinaryCache::IsEnabled() && !Capabilities.bSupportsSeparateShaderObjects;
-		// deferred compilation is not supported for SeparateShaderObjects
-		check(!bDeferredCompilation || !Capabilities.bSupportsSeparateShaderObjects);
-
+		const bool bDeferredCompilation = FOpenGLProgramBinaryCache::IsEnabled();
 		if (bDeferredCompilation)
 		{
 			// do nothing...
 		}
 		else
 		{
-			const FAnsiCharArray glslPlatformCode = FoundShader->GetUncompressedShader();
-
-			const bool bSuccessfullyCompiled = CompileCurrentShader(Resource, glslPlatformCode);
-
-			if (Capabilities.bSupportsSeparateShaderObjects && bSuccessfullyCompiled)
-			{
-				ANSICHAR Buf[32] = { 0 };
-				// Create separate shader program
-				GLuint SeparateResource = FOpenGL::CreateProgram();
-				FOpenGL::ProgramParameter(SeparateResource, GL_PROGRAM_SEPARABLE, GL_TRUE);
-				glAttachShader(SeparateResource, Resource);
-
-				glLinkProgram(SeparateResource);
-				VerifyLinkedProgram(SeparateResource);
-
-#if ENABLE_UNIFORM_BUFFER_LAYOUT_VERIFICATION
-				void VerifyUniformBufferLayouts(GLuint Program);
-				VerifyUniformBufferLayouts(SeparateResource);
-#endif // #if ENABLE_UNIFORM_BUFFER_LAYOUT_VERIFICATION
-
-				Resource = SeparateResource;
-			}
+			const FAnsiCharArray GlslPlatformCode = FoundShader->GetUncompressedShader();
+			const bool bSuccessfullyCompiled = CompileCurrentShader(Resource, GlslPlatformCode);
+			ensure(bSuccessfullyCompiled);
 		}
+
 		FoundShader->Resource = Resource;
 	}
 
@@ -785,22 +702,6 @@ ShaderType* CompileOpenGLShader(const FOpenGLCodeHeader& Header, const FOpenGLCo
 	checkf(Shader->StaticSlots.Num() == Shader->Bindings.ShaderResourceTable.ResourceTableLayoutHashes.Num(), TEXT("StaticSlots %d, Bindings %d"),
 		Shader->StaticSlots.Num(), Shader->Bindings.ShaderResourceTable.ResourceTableLayoutHashes.Num());
 
-	if (FOpenGL::SupportsSeparateShaderObjects())
-	{
-		FSHAHash Hash;
-		// Just use the CRC - if it isn't being cached & logged we'll be dependent on the CRC alone anyway
-		uint32 Crc = SourceKey.GetCodeCRC();
-		FMemory::Memcpy(Hash.Hash, &Crc, sizeof(uint32));
-		if (RHIShader)
-		{
-			SetShaderHash(Hash, RHIShader);
-		}
-		else
-		{
-			SetShaderHash(Hash, Shader);
-		}
-	}
-
 #if DEBUG_GL_SHADERS
 	{
 		const FAnsiCharArray glslPlatformCode = FoundShader->GetUncompressedShader();
@@ -852,7 +753,6 @@ void OPENGLDRV_API GetCurrentOpenGLShaderDeviceCapabilities(FOpenGLShaderDeviceC
 	FOpenGL::PE_GetCurrentOpenGLShaderDeviceCapabilities(Capabilities); // platform extension
 #endif
 	Capabilities.MaxRHIShaderPlatform = GMaxRHIShaderPlatform;
-	Capabilities.bSupportsSeparateShaderObjects = FOpenGL::SupportsSeparateShaderObjects();
 }
 
 void OPENGLDRV_API GLSLToDeviceCompatibleGLSL(FAnsiCharArray& GlslCodeOriginal, const FString& ShaderName, GLenum TypeEnum, const FOpenGLShaderDeviceCapabilities& Capabilities, FAnsiCharArray& GlslCode)
@@ -970,7 +870,6 @@ void OPENGLDRV_API GLSLToDeviceCompatibleGLSL(FAnsiCharArray& GlslCodeOriginal,
 	// OpenGL SM5 shader platforms require location declarations for the layout, but don't necessarily use SSOs
 	if (Capabilities.TargetPlatform == EOpenGLShaderTargetPlatform::OGLSTP_Desktop)
 	{
-		AppendCString(GlslCode, "#extension GL_ARB_separate_shader_objects : enable\n");
 		AppendCString(GlslCode, "#define INTERFACE_BLOCK(Pos, Interp, Modifiers, Semantic, PreType, PostType) layout(location=Pos) Interp Modifiers struct { PreType PostType; }\n");
 	}
 	else
@@ -1433,11 +1332,9 @@ public:
 		SetDeletedProgramStats(Program);
 		FOpenGL::DeleteProgramPipelines(1, &Program);
 
-		if (!FOpenGL::SupportsSeparateShaderObjects())
-		{
-			GetOpenGLUniformBlockLocations().Remove(Program);
-			GetOpenGLUniformBlockBindings().Remove(Program);
-		}
+		GetOpenGLUniformBlockLocations().Remove(Program);
+		GetOpenGLUniformBlockBindings().Remove(Program);
+
 		Program = 0;
 
 		for (int Stage = 0; Stage < CrossCompiler::NUM_SHADER_STAGES; Stage++)
@@ -1448,9 +1345,6 @@ public:
 		}
 	}
 
-	// Rebind the uniform blocks when changing the separable shader pipeline as different stages will have different uniform block arrangements. Does nothing for non-separable GLs.
-	void VerifyUniformBlockBindings( int Stage, uint32 FirstUniformBuffer );
-
 	void ConfigureShaderStage( int Stage, uint32 FirstUniformBuffer );
 
 	// Make sure GlobalArrays (created from shader reflection) matches our info (from the cross compiler)
@@ -1744,7 +1638,7 @@ static int32 GetProgramBinarySize(GLuint Program)
 
 void ConfigureGLProgramStageStates(FOpenGLLinkedProgram* LinkedProgram)
 {
-	ensure(VerifyProgramPipeline(LinkedProgram->Program));
+	ensure(VerifyLinkedProgram(LinkedProgram->Program));
 	FOpenGL::BindProgramPipeline(LinkedProgram->Program);
 	ConfigureStageStates(LinkedProgram);
 }
@@ -2068,7 +1962,7 @@ public:
 			UE_LOG(LogRHI, Warning, TEXT("Requesting OpenGL program LRU cache, but program binary is not supported by driver. Falling back to non-lru cache."));
 		}
 
-		bUseLRUCache = CVarEnableLRU.GetValueOnAnyThread() == 1 && FOpenGL::SupportsProgramBinary() && !FOpenGL::SupportsSeparateShaderObjects();
+		bUseLRUCache = CVarEnableLRU.GetValueOnAnyThread() == 1 && FOpenGL::SupportsProgramBinary();
 		UE_LOG(LogRHI, Log, TEXT("Using OpenGL program LRU cache: %d"), bUseLRUCache ? 1 : 0);
 	}
 
@@ -2287,29 +2181,6 @@ static int32 CountSetBits(const TBitArray<>& Array)
 	return Result;
 }
 
-FORCEINLINE_DEBUGGABLE void FOpenGLLinkedProgram::VerifyUniformBlockBindings( int Stage, uint32 FirstUniformBuffer )
-{
-	if ( FOpenGL::SupportsSeparateShaderObjects() && FOpenGL::SupportsUniformBuffers() )
-	{
-		VERIFY_GL_SCOPE();
-		FOpenGLUniformName Name;
-		Name.Buffer[0] = CrossCompiler::ShaderStageIndexToTypeName(Stage);
-		Name.Buffer[1] = 'b';
-		
-		GLuint StageProgram = Config.Shaders[Stage].Resource;
-
-		for (int32 BufferIndex = 0; BufferIndex < Config.Shaders[Stage].Bindings.NumUniformBuffers; ++BufferIndex)
-		{
-			SetIndex(Name.Buffer, 2, BufferIndex);
-			GLint Location = GetOpenGLProgramUniformBlockIndex(StageProgram, Name);
-			if (Location >= 0)
-			{
-				GetOpenGLProgramUniformBlockBinding(StageProgram, Location, FirstUniformBuffer + BufferIndex);
-			}
-		}
-	}
-}
-
 void FOpenGLLinkedProgram::ConfigureShaderStage( int Stage, uint32 FirstUniformBuffer )
 {
 	static const GLint FirstTextureUnit[CrossCompiler::NUM_SHADER_STAGES] =
@@ -2340,7 +2211,7 @@ void FOpenGLLinkedProgram::ConfigureShaderStage( int Stage, uint32 FirstUniformB
 	FOpenGLUniformName Name;
 	Name.Buffer[0] = CrossCompiler::ShaderStageIndexToTypeName(Stage);
 
-	GLuint StageProgram = FOpenGL::SupportsSeparateShaderObjects() ? Config.Shaders[Stage].Resource : Program;
+	GLuint StageProgram = Program;
 	
 	// Bind Global uniform arrays (vu_h, pu_i, etc)
 	{
@@ -2830,21 +2701,18 @@ static FOpenGLLinkedProgram* LinkProgram( const FOpenGLLinkedProgramConfiguratio
 			FOpenGL::UseProgramStages(Program, GL_COMPUTE_SHADER_BIT, Config.Shaders[CrossCompiler::SHADER_STAGE_COMPUTE].Resource);
 		}
 	
-		if( !FOpenGL::SupportsSeparateShaderObjects() )
+		if(FOpenGLProgramBinaryCache::IsEnabled() || GetOpenGLProgramsCache().IsUsingLRU())
 		{
-			if(FOpenGLProgramBinaryCache::IsEnabled() || GetOpenGLProgramsCache().IsUsingLRU())
-			{
-				FOpenGL::ProgramParameter(Program, PROGRAM_BINARY_RETRIEVABLE_HINT, GL_TRUE);
-			}
-
-			// Link.
-			glLinkProgram(Program);
+			FOpenGL::ProgramParameter(Program, PROGRAM_BINARY_RETRIEVABLE_HINT, GL_TRUE);
 		}
+
+		// Link.
+		glLinkProgram(Program);
 	}
 
-	if (VerifyProgramPipeline(Program))
+	if (VerifyLinkedProgram(Program))
 	{
-		if(bShouldLinkProgram && !FOpenGL::SupportsSeparateShaderObjects())
+		if (bShouldLinkProgram)
 		{
 			SetNewProgramStats(Program);
 
@@ -3000,12 +2868,6 @@ struct FOpenGLShaderVaryingMapping
 
 typedef TMap<FOpenGLLinkedProgramConfiguration,FOpenGLLinkedProgramConfiguration::ShaderInfo> FOpenGLSeparateShaderObjectCache;
 
-static FOpenGLSeparateShaderObjectCache& GetOpenGLSeparateShaderObjectCache()
-{
-	static FOpenGLSeparateShaderObjectCache SeparateShaderObjectCache;
-	return SeparateShaderObjectCache;
-}
-
 template<class TOpenGLStage0RHI, class TOpenGLStage1RHI>
 static void BindShaderStage(FOpenGLLinkedProgramConfiguration& Config, CrossCompiler::EShaderStage NextStage, TOpenGLStage0RHI* NextStageShaderIn, CrossCompiler::EShaderStage PrevStage, TOpenGLStage1RHI* PrevStageShaderIn)
 {
@@ -3023,202 +2885,6 @@ static void BindShaderStage(FOpenGLLinkedProgramConfiguration& Config, CrossComp
 	GLuint NextStageResource = NextStageShader->Resource;
 	FOpenGLShaderBindings NextStageBindings = NextStageShader->Bindings;
 	
-	if ( FOpenGL::SupportsSeparateShaderObjects() )
-	{
-		FOpenGLLinkedProgramConfiguration SeparableConfig;
-		SeparableConfig.Shaders[0] = PrevInfo;
-		SeparableConfig.ProgramKey.ShaderHashes[0] = Config.ProgramKey.ShaderHashes[PrevStage];
-		SeparableConfig.Shaders[1] = ShaderInfo;
-		SeparableConfig.ProgramKey.ShaderHashes[1] = Config.ProgramKey.ShaderHashes[NextStage];
-
-		FOpenGLLinkedProgramConfiguration::ShaderInfo* PrevResource = GetOpenGLSeparateShaderObjectCache().Find(SeparableConfig);
-		if(PrevResource)
-		{
-			PrevInfo.Bindings = PrevResource->Bindings;
-			PrevInfo.Resource = PrevResource->Resource;
-		}
-		else
-		{
-			FOpenGLShaderBindings& PrevStageBindings = PrevStageShader->Bindings;
-			TMap<FAnsiCharArray, int32> PrevStageVaryings;
-			for (int32 i = 0; i < PrevStageBindings.OutputVaryings.Num(); i++)
-			{
-				FAnsiCharArray Name = PrevStageBindings.OutputVaryings[i].Varying;
-				if ( Name.Num() >= 4 && (FCStringAnsi::Strncmp(Name.GetData(), "out_", 4) == 0 || FCStringAnsi::Strncmp(Name.GetData(), "var_", 4) == 0) )
-				{
-					Name.RemoveAt(0, 4);
-				}
-				PrevStageVaryings.Add(Name, PrevStageBindings.OutputVaryings[i].Location);
-			}
-			
-			bool bInterpolatorMatches = true;
-			
-			TMap<FAnsiCharArray, int32> NextStageVaryings;
-			TArray<FString> InputErrors;
-			TArray<FOpenGLShaderVaryingMapping> VaryingMapping;
-			for (int32 i = 0; i < NextStageBindings.InputVaryings.Num(); i++)
-			{
-				FAnsiCharArray Name = NextStageBindings.InputVaryings[i].Varying;
-				if ( Name.Num() >= 3 && FCStringAnsi::Strncmp(Name.GetData(), "in_", 3) == 0 )
-				{
-					Name.RemoveAt(0, 3);
-				}
-				if ( Name.Num() >= 4 && FCStringAnsi::Strncmp(Name.GetData(), "var_", 4) == 0 )
-				{
-					Name.RemoveAt(0, 4);
-				}
-				NextStageVaryings.Add(Name, NextStageBindings.InputVaryings[i].Location);
-				if( PrevStageVaryings.Contains(Name) )
-				{
-					int32& PrevLocation = PrevStageVaryings.FindChecked(Name);
-					if(PrevLocation != NextStageBindings.InputVaryings[i].Location)
-					{
-						if(PrevLocation >= 0 && NextStageBindings.InputVaryings[i].Location >= 0)
-						{
-							FOpenGLShaderVaryingMapping Pair;
-							Pair.Name = Name;
-							Pair.WriteLoc = PrevLocation;
-							Pair.ReadLoc = NextStageBindings.InputVaryings[i].Location;
-							VaryingMapping.Add(Pair);
-							UE_LOG(LogRHI,Warning,TEXT("Separate Shader Object Binding Warning: Input %s @ %d of stage 0x%x written by stage 0x%x at wrong location %d"), ANSI_TO_TCHAR(NextStageBindings.InputVaryings[i].Varying.GetData()), NextStageBindings.InputVaryings[i].Location, TOpenGLStage0::TypeEnum, TOpenGLStage1::TypeEnum, PrevLocation);
-						}
-						else if(NextStageBindings.InputVaryings[i].Location == -1)
-						{
-							InputErrors.Add(FString::Printf(TEXT("Separate Shader Object Binding Error: Input %s of stage 0x%x written by stage 0x%x at location %d, can't be rewritten."), ANSI_TO_TCHAR(NextStageBindings.InputVaryings[i].Varying.GetData()), TOpenGLStage0::TypeEnum, TOpenGLStage1::TypeEnum, PrevLocation));
-						}
-						else
-						{
-							InputErrors.Add(FString::Printf(TEXT("Separate Shader Object Binding Error: Input %s @ %d of stage 0x%x written by stage 0x%x without location, can't be rewritten."), ANSI_TO_TCHAR(NextStageBindings.InputVaryings[i].Varying.GetData()), NextStageBindings.InputVaryings[i].Location, TOpenGLStage0::TypeEnum, TOpenGLStage1::TypeEnum));
-						}
-						bInterpolatorMatches = false;
-					}
-				}
-				else
-				{
-					InputErrors.Add(FString::Printf(TEXT("Separate Shader Object Binding Error: Input %s @ %d of stage 0x%x not written by stage 0x%x"), ANSI_TO_TCHAR(NextStageBindings.InputVaryings[i].Varying.GetData()), NextStageBindings.InputVaryings[i].Location, TOpenGLStage0::TypeEnum, TOpenGLStage1::TypeEnum));
-					bInterpolatorMatches = false;
-				}
-			}
-			
-			TArray<FOpenGLShaderVarying> OutputElimination;
-			for (int32 i = 0; i < PrevStageBindings.OutputVaryings.Num(); i++)
-			{
-				if ( PrevStageBindings.OutputVaryings[i].Location == -1 )
-				{
-					FAnsiCharArray Name = PrevStageBindings.OutputVaryings[i].Varying;
-					if ( Name.Num() >= 4 && (FCStringAnsi::Strncmp(Name.GetData(), "out_", 4) == 0 || FCStringAnsi::Strncmp(Name.GetData(), "var_", 4) == 0) )
-					{
-						Name.RemoveAt(0, 4);
-					}
-					if( !NextStageVaryings.Contains(Name) )
-					{
-						OutputElimination.Add(PrevStageBindings.OutputVaryings[i]);
-						UE_LOG(LogRHI,Warning,TEXT("Separate Shader Object Binding Warning: Named output %s of stage 0x%x not read by stage 0x%x"), ANSI_TO_TCHAR(PrevStageBindings.OutputVaryings[i].Varying.GetData()), TOpenGLStage1::TypeEnum, TOpenGLStage0::TypeEnum);
-						bInterpolatorMatches = false;
-					}
-				}
-			}
-		
-			if(!bInterpolatorMatches)
-			{
-				if(InputErrors.Num() == 0)
-				{
-					FOpenGLCodeHeader Header;
-					Header.GlslMarker = 0x474c534c;
-					CA_SUPPRESS(6326);
-					switch ((int32)TOpenGLStage1::StaticFrequency)
-					{
-						case SF_Vertex:
-							Header.FrequencyMarker = 0x5653;
-							break;
-						case SF_Pixel:
-							Header.FrequencyMarker = 0x5053;
-							break;
-						case SF_Geometry:
-							Header.FrequencyMarker = 0x4753;
-							break;
-						case SF_Compute:
-							Header.FrequencyMarker = 0x4353;
-							break;
-						default:
-							UE_LOG(LogRHI, Fatal, TEXT("Invalid shader frequency: %d"), (int32)TOpenGLStage1::StaticFrequency);
-					}
-					Header.Bindings = PrevStageShader->Bindings;
-					Header.UniformBuffersCopyInfo = PrevStageShader->UniformBuffersCopyInfo;
-					
-					TArray<FString> PrevLines;
-					FString PrevSource = GetShaderStageSource<TOpenGLStage1>(PrevStageShader);
-					PrevSource.ParseIntoArrayLines(PrevLines);
-					bool const bOutputElimination = OutputElimination.Num() > 0;
-					for(FOpenGLShaderVarying Output : OutputElimination)
-					{
-						for(int32 i = 0; i < PrevLines.Num(); i++)
-						{
-							if(PrevLines[i].Contains(Output.Varying.GetData()))
-							{
-								PrevLines[i].Empty();
-							}
-						}
-						for(int32 i = 0; i < Header.Bindings.OutputVaryings.Num(); i++)
-						{
-							if(Output == Header.Bindings.OutputVaryings[i])
-							{
-								Header.Bindings.OutputVaryings.RemoveAt(i);
-								break;
-							}
-						}
-					}
-					OutputElimination.Empty();
-					
-					bool const bVaryingRemapping = VaryingMapping.Num() > 0;
-					
-					if (OutputElimination.Num() == 0 && VaryingMapping.Num() == 0 && (bOutputElimination || bVaryingRemapping))
-					{
-						FString NewPrevSource;
-						for(FString Line : PrevLines)
-						{
-							if(!Line.IsEmpty())
-							{
-								NewPrevSource += Line + TEXT("\n");
-							}
-						}
-						
-						TArray<uint8> Bytes;
-						FMemoryWriter Ar(Bytes);
-						Ar << Header;
-						TArray<ANSICHAR> Chars;
-						int32 Len = FCStringAnsi::Strlen(TCHAR_TO_ANSI(*NewPrevSource)) + 1;
-						Chars.Append(TCHAR_TO_ANSI(*NewPrevSource), Len);
-						Ar.Serialize(Chars.GetData(), Chars.Num());
-						
-						TRefCountPtr<TOpenGLStage1RHI> NewPrev = CreateProxyShader<TOpenGLStage1RHI>(Bytes, FSHAHash());
-						PrevInfo.Bindings = Header.Bindings;
-						PrevInfo.Resource = FOpenGLDynamicRHI::ResourceCast(NewPrev.GetReference())->Resource;
-					}
-					
-					bInterpolatorMatches = (OutputElimination.Num() == 0 && VaryingMapping.Num() == 0);
-				}
-				else
-				{
-					for(int32 i = 0; i < InputErrors.Num(); i++)
-					{
-						UE_LOG(LogRHI, Error, TEXT("%s"), *InputErrors[i]);
-					}
-				}
-				
-				if(!bInterpolatorMatches)
-				{
-					FString PrevShaderStageSource = GetShaderStageSource<TOpenGLStage1>(PrevStageShader);
-					FString NextShaderStageSource = GetShaderStageSource<TOpenGLStage0>(NextStageShader);
-					UE_LOG(LogRHI, Error, TEXT("Separate Shader Object Stage 0x%x:\n%s"), TOpenGLStage1::TypeEnum, *PrevShaderStageSource);
-					UE_LOG(LogRHI, Error, TEXT("Separate Shader Object Stage 0x%x:\n%s"), TOpenGLStage0::TypeEnum, *NextShaderStageSource);
-				}
-			}
-			
-			GetOpenGLSeparateShaderObjectCache().Add(SeparableConfig, PrevInfo);
-		}
-	}
-	
 	ShaderInfo.Bindings = NextStageBindings;
 	ShaderInfo.Resource = NextStageResource;
 }
@@ -3635,9 +3301,6 @@ void FOpenGLDynamicRHI::BindPendingShaderState( FOpenGLContextState& ContextStat
 		ContextState.Program = PendingProgram;
 		MarkShaderParameterCachesDirty(PendingState.ShaderParameters, false);
 		PendingState.LinkedProgramAndDirtyFlag = nullptr;
-
-		//Disable the forced rebinding to reduce driver overhead - required by SSOs
-		ForceUniformBindingUpdate = FOpenGL::SupportsSeparateShaderObjects();
 	}
 
 	if (PendingState.bAnyDirtyRealUniformBuffers[SF_Vertex] || 
@@ -3655,7 +3318,6 @@ void FOpenGLDynamicRHI::BindPendingShaderState( FOpenGLContextState& ContextStat
 
 		if (PendingState.bAnyDirtyRealUniformBuffers[SF_Vertex])
 		{
-			PendingState.BoundShaderState->LinkedProgram->VerifyUniformBlockBindings(CrossCompiler::SHADER_STAGE_VERTEX, NextUniformBufferIndex);
 			BindUniformBufferBase(
 				ContextState,
 				NumUniformBuffers[SF_Vertex],
@@ -3667,7 +3329,6 @@ void FOpenGLDynamicRHI::BindPendingShaderState( FOpenGLContextState& ContextStat
 
 		if (PendingState.bAnyDirtyRealUniformBuffers[SF_Pixel])
 		{
-			PendingState.BoundShaderState->LinkedProgram->VerifyUniformBlockBindings(CrossCompiler::SHADER_STAGE_PIXEL, NextUniformBufferIndex);
 			BindUniformBufferBase(
 				ContextState,
 				NumUniformBuffers[SF_Pixel],
@@ -3679,7 +3340,6 @@ void FOpenGLDynamicRHI::BindPendingShaderState( FOpenGLContextState& ContextStat
 
 		if (NumUniformBuffers[SF_Geometry] >= 0 && PendingState.bAnyDirtyRealUniformBuffers[SF_Geometry])
 		{
-			PendingState.BoundShaderState->LinkedProgram->VerifyUniformBlockBindings(CrossCompiler::SHADER_STAGE_GEOMETRY, NextUniformBufferIndex);
 			BindUniformBufferBase(
 				ContextState,
 				NumUniformBuffers[SF_Geometry],
@@ -3859,7 +3519,6 @@ void FOpenGLDynamicRHI::BindPendingComputeShaderState(FOpenGLContextState& Conte
 
 	if (PendingState.bAnyDirtyRealUniformBuffers[SF_Compute])
 	{
-		ComputeShader->LinkedProgram->VerifyUniformBlockBindings(CrossCompiler::SHADER_STAGE_COMPUTE, OGL_FIRST_UNIFORM_BUFFER);
 		BindUniformBufferBase(
 			ContextState,
 			ComputeShader->Bindings.NumUniformBuffers,
@@ -4226,12 +3885,6 @@ void FOpenGLProgramBinaryCache::Initialize()
 		return;
 	}
 
-	if (FOpenGL::SupportsSeparateShaderObjects())
-	{
-		UE_LOG(LogRHI, Warning, TEXT("FOpenGLProgramBinaryCache disabled as RHI supports separate shader objects"));
-		return;
-	}
-
 	if (!FOpenGL::SupportsProgramBinary())
 	{
 		UE_LOG(LogRHI, Warning, TEXT("FOpenGLProgramBinaryCache disabled as devices does not support program binaries"));
diff --git a/Engine/Source/Runtime/OpenGLDrv/Public/OpenGL.h b/Engine/Source/Runtime/OpenGLDrv/Public/OpenGL.h
--- a/Engine/Source/Runtime/OpenGLDrv/Public/OpenGL.h
+++ b/Engine/Source/Runtime/OpenGLDrv/Public/OpenGL.h
@@ -132,7 +132,6 @@ public:
 	static FORCEINLINE bool SupportsBindlessTexture()					{ return false; }
 	static FORCEINLINE bool HasHardwareHiddenSurfaceRemoval()			{ return false; }
 	static FORCEINLINE bool AmdWorkaround()								{ return false; }
-	static FORCEINLINE bool SupportsSeparateShaderObjects()				{ return false; }
 	static FORCEINLINE bool SupportsProgramBinary()						{ return false; }
 	static FORCEINLINE bool SupportsDepthClamp()						{ return true; }
 	
diff --git a/Engine/Source/Runtime/OpenGLDrv/Public/OpenGL3.h b/Engine/Source/Runtime/OpenGLDrv/Public/OpenGL3.h
--- a/Engine/Source/Runtime/OpenGLDrv/Public/OpenGL3.h
+++ b/Engine/Source/Runtime/OpenGLDrv/Public/OpenGL3.h
@@ -41,7 +41,6 @@ struct FOpenGL3 : public FOpenGLBase
 	static FORCEINLINE bool SupportsVolumeTextureRendering()		{ return bSupportsVolumeTextureRendering; }
 	static FORCEINLINE bool SupportsGenerateMipmap()				{ return true; }
 	static FORCEINLINE bool AmdWorkaround()							{ return bAmdWorkaround; }
-	static FORCEINLINE bool SupportsSeparateShaderObjects()			{ return bSupportsSeparateShaderObjects; }
 	static FORCEINLINE bool SupportsBufferStorage()					{ return true; }
 
 	// Optional
@@ -578,156 +577,24 @@ struct FOpenGL3 : public FOpenGLBase
 	}
 	static FORCEINLINE void ProgramParameter(GLuint Program, GLenum PName, GLint Value)
 	{
-		check(PName != GL_PROGRAM_SEPARABLE || FOpenGL3::SupportsSeparateShaderObjects());
+		check(PName != GL_PROGRAM_SEPARABLE);
 		glProgramParameteri(Program, PName, Value);
 	}
-	static FORCEINLINE void UseProgramStages(GLuint Pipeline, GLbitfield Stages, GLuint Program)
-	{
-		if(FOpenGL3::SupportsSeparateShaderObjects())
-		{
-			glUseProgramStages(Pipeline, Stages, Program);
-		}
-		else
-		{
-			glAttachShader(Pipeline, Program);
-		}
-	}
-	static FORCEINLINE void BindProgramPipeline(GLuint Pipeline)
-	{
-		if(FOpenGL3::SupportsSeparateShaderObjects())
-		{
-			glBindProgramPipeline(Pipeline);
-		}
-		else
-		{
-			glUseProgram(Pipeline);
-		}
-	}
-	static FORCEINLINE void DeleteShader(GLuint Program)
-	{
-		if(FOpenGL3::SupportsSeparateShaderObjects())
-		{
-			GLsizei NumShaders = 0;
-			glGetProgramiv(Program, GL_ATTACHED_SHADERS, (GLint*)&NumShaders);
-			if(NumShaders > 0)
-			{
-				GLuint* Shaders = (GLuint*)FMemory::Malloc(sizeof(GLuint) * NumShaders);
-				glGetAttachedShaders(Program, NumShaders, &NumShaders, Shaders);
-				for(int32 i = 0; i < NumShaders; i++)
-				{
-					glDetachShader(Program, Shaders[i]);
-					glDeleteShader(Shaders[i]);
-				}
-				FMemory::Free(Shaders);
-			}
-			
-			glDeleteProgram(Program);
-		}
-		else
-		{
-			glDeleteShader(Program);
-		}
-	}
-	static FORCEINLINE void DeleteProgramPipelines(GLsizei Number, const GLuint *Pipelines)
-	{
-		if(FOpenGL3::SupportsSeparateShaderObjects())
-		{
-			glDeleteProgramPipelines(Number, Pipelines);
-		}
-		else
-		{
-			for(GLsizei i = 0; i < Number; i++)
-			{
-				glDeleteProgram(Pipelines[i]);
-			}
-		}
-	}
-	static FORCEINLINE void GenProgramPipelines(GLsizei Number, GLuint *Pipelines)
-	{
-		if(FOpenGL3::SupportsSeparateShaderObjects())
-		{
-#if USE_OPENGL_NAME_CACHE
-			if ( Number < OPENGL_NAME_CACHE_SIZE - NextPipelineName)
-			{
-				FMemory::Memcpy( Pipelines, &PipelineNamesCache[NextPipelineName], sizeof(GLuint)*Number);
-				NextPipelineName += Number;
-			}
-			else
-			{
-				if ( Number >= OPENGL_NAME_CACHE_SIZE)
-				{
-					glGenProgramPipelines(Number, Pipelines);
-				}
-				else
-				{
-					GLsizei Leftover = OPENGL_NAME_CACHE_SIZE - NextPipelineName;
-
-					FMemory::Memcpy( Pipelines, &PipelineNamesCache[NextPipelineName], sizeof(GLuint)*Leftover);
-
-					glGenProgramPipelines( OPENGL_NAME_CACHE_SIZE, PipelineNamesCache);
-
-					Number -= Leftover;
-					Pipelines += Leftover;
-
-					FMemory::Memcpy( Pipelines, PipelineNamesCache, sizeof(GLuint)*Number);
-					NextPipelineName = Number;
-				}
-			}
-#else
-			glGenProgramPipelines(Number, Pipelines);
-#endif
-		}
-		else
-		{
-			for(GLsizei i = 0; i < Number; i++)
-			{
-				Pipelines[i] = FOpenGL3::CreateProgram();
-			}
-		}
-	}
 	static FORCEINLINE void ProgramUniform1i(GLuint Program, GLint Location, GLint V0)
 	{
-		if (FOpenGL3::SupportsSeparateShaderObjects())
-		{
-			glProgramUniform1i(Program, Location, V0);
-		}
-		else
-		{
-			glUniform1i(Location, V0);
-		}
+		glUniform1i(Location, V0);
 	}
 	static FORCEINLINE void ProgramUniform4iv(GLuint Program, GLint Location, GLsizei Count, const GLint *Value)
 	{
-		if (FOpenGL3::SupportsSeparateShaderObjects())
-		{
-			glProgramUniform4iv(Program, Location, Count, Value);
-		}
-		else
-		{
-			glUniform4iv(Location, Count, Value);
-		}
+		glUniform4iv(Location, Count, Value);
 	}
 	static FORCEINLINE void ProgramUniform4fv(GLuint Program, GLint Location, GLsizei Count, const GLfloat *Value)
 	{
-		if (FOpenGL3::SupportsSeparateShaderObjects())
-		{
-			glProgramUniform4fv(Program, Location, Count, Value);
-		}
-		else
-		{
-			glUniform4fv(Location, Count, Value);
-		}
+		glUniform4fv(Location, Count, Value);
 	}
 	static FORCEINLINE void ProgramUniform4uiv(GLuint Program, GLint Location, GLsizei Count, const GLuint *Value)
 	{
-		if (FOpenGL3::SupportsSeparateShaderObjects())
-		{
-			glProgramUniform4uiv(Program, Location, Count, Value);
-		}
-		else
-		{
-			glUniform4uiv(Location, Count, Value);
-		}
+		glUniform4uiv(Location, Count, Value);
 	}
 	static FORCEINLINE void GetProgramPipelineiv(GLuint Pipeline, GLenum Pname, GLint* Params)
 	{
@@ -783,6 +650,5 @@ protected:
 	static GLint TimestampQueryBits;
 	
 	static bool bDebugContext;
-	static bool bSupportsSeparateShaderObjects;
 	static bool bAndroidGLESCompatibilityMode;
 };
diff --git a/Engine/Source/Runtime/OpenGLDrv/Public/OpenGLShaders.h b/Engine/Source/Runtime/OpenGLDrv/Public/OpenGLShaders.h
--- a/Engine/Source/Runtime/OpenGLDrv/Public/OpenGLShaders.h
+++ b/Engine/Source/Runtime/OpenGLDrv/Public/OpenGLShaders.h
@@ -24,7 +24,6 @@ struct FOpenGLShaderDeviceCapabilities
 {
 	EOpenGLShaderTargetPlatform TargetPlatform;
 	EShaderPlatform MaxRHIShaderPlatform;
-	bool bSupportsSeparateShaderObjects;
 	bool bSupportsShaderFramebufferFetch;
 	bool bRequiresARMShaderFramebufferFetchDepthStencilUndef;
 	GLint MaxVaryingVectors;
