Changed PLS to require ARM framebuffer depth/stencil fetch on OpenGL

diff --git a/Engine/Source/Runtime/OpenGLDrv/Private/OpenGLRenderTarget.cpp b/Engine/Source/Runtime/OpenGLDrv/Private/OpenGLRenderTarget.cpp
--- a/Engine/Source/Runtime/OpenGLDrv/Private/OpenGLRenderTarget.cpp
+++ b/Engine/Source/Runtime/OpenGLDrv/Private/OpenGLRenderTarget.cpp
@@ -1141,7 +1141,7 @@ void FOpenGLDynamicRHI::RHIBeginRenderPass(const FRHIRenderPassInfo& InInfo, con
 
 #if PLATFORM_ANDROID
 	if (RenderPassInfo.SubpassHint == ESubpassHint::DeferredShadingSubpass &&
-		FOpenGL::SupportsPixelLocalStorage())
+		 FOpenGL::SupportsPixelLocalStorage() && FOpenGL::SupportsShaderDepthStencilFetch())
 	{
 		glEnable(GL_SHADER_PIXEL_LOCAL_STORAGE_EXT);
 	}
@@ -1199,7 +1199,7 @@ void FOpenGLDynamicRHI::RHIEndRenderPass()
 
 #if PLATFORM_ANDROID
 	if (RenderPassInfo.SubpassHint == ESubpassHint::DeferredShadingSubpass &&
-		FOpenGL::SupportsPixelLocalStorage())
+		FOpenGL::SupportsPixelLocalStorage() && FOpenGL::SupportsShaderDepthStencilFetch())
 	{
 		glDisable(GL_SHADER_PIXEL_LOCAL_STORAGE_EXT);
 	}
diff --git a/Engine/Source/Runtime/OpenGLDrv/Private/OpenGLShaders.cpp b/Engine/Source/Runtime/OpenGLDrv/Private/OpenGLShaders.cpp
--- a/Engine/Source/Runtime/OpenGLDrv/Private/OpenGLShaders.cpp
+++ b/Engine/Source/Runtime/OpenGLDrv/Private/OpenGLShaders.cpp
@@ -797,7 +797,7 @@ void OPENGLDRV_API GLSLToDeviceCompatibleGLSL(FAnsiCharArray& GlslCodeOriginal,
 	// This is the place to insert such engine preprocessor defines, immediately after the glsl version declaration.
 	if (TypeEnum == GL_FRAGMENT_SHADER)
 	{
-		if (FOpenGL::SupportsPixelLocalStorage())
+		if (FOpenGL::SupportsPixelLocalStorage() && FOpenGL::SupportsShaderDepthStencilFetch())
 		{
 			AppendCString(GlslCode, "#define UE_MRT_PLS 1\n");
 		}
diff --git a/Engine/Source/Runtime/RenderCore/Private/RenderUtils.cpp b/Engine/Source/Runtime/RenderCore/Private/RenderUtils.cpp
--- a/Engine/Source/Runtime/RenderCore/Private/RenderUtils.cpp
+++ b/Engine/Source/Runtime/RenderCore/Private/RenderUtils.cpp
@@ -1134,7 +1134,7 @@ RENDERCORE_API bool MobileRequiresSceneDepthAux(const FStaticShaderPlatform Plat
 	{
 		return true;
 	}
-	else if (IsMobileDeferredShadingEnabled(Platform) && IsAndroidOpenGLESPlatform(Platform))
+	else if (IsMobileDeferredShadingEnabled(Platform) && IsAndroidOpenGLESPlatform(Platform) && !GSupportsShaderDepthStencilFetch)
 	{
 		return true;
 	}
diff --git a/Engine/Source/Runtime/Renderer/Private/MobileShadingRenderer.cpp b/Engine/Source/Runtime/Renderer/Private/MobileShadingRenderer.cpp
--- a/Engine/Source/Runtime/Renderer/Private/MobileShadingRenderer.cpp
+++ b/Engine/Source/Runtime/Renderer/Private/MobileShadingRenderer.cpp
@@ -1316,7 +1316,7 @@ void FMobileSceneRenderer::RenderDeferred(FRDGBuilder& GraphBuilder, const FSort
 	TArray<FRDGTextureRef, TInlineAllocator<6>> ColorTargets;
 
 	// If we are using GL and don't have FBF support, use PLS
-	bool bUsingPixelLocalStorage = IsAndroidOpenGLESPlatform(ShaderPlatform) && GSupportsPixelLocalStorage;
+	bool bUsingPixelLocalStorage = IsAndroidOpenGLESPlatform(ShaderPlatform) && GSupportsPixelLocalStorage && GSupportsShaderDepthStencilFetch;
 
 	if (bUsingPixelLocalStorage)
 	{
