Fix incorrect handling of bArray argument in FOpenGLDynamicRHI::RHICreateTextureCubeFromResource

diff --git a/Engine/Source/Runtime/OpenGLDrv/Private/OpenGLTexture.cpp b/Engine/Source/Runtime/OpenGLDrv/Private/OpenGLTexture.cpp
--- a/Engine/Source/Runtime/OpenGLDrv/Private/OpenGLTexture.cpp
+++ b/Engine/Source/Runtime/OpenGLDrv/Private/OpenGLTexture.cpp
@@ -2256,25 +2256,25 @@ FTextureCubeRHIRef FOpenGLDynamicRHI::RHICreateTextureCubeFromResource(EPixelFor
 {
 	if (bArray)
 	{
-		return new FOpenGLTexture(FRHITextureCreateDesc::CreateCube(
+		return new FOpenGLTexture(FRHITextureCreateDesc::CreateCubeArray(
 			TEXT("RHICreateTextureCubeFromResource"),
 			Size,
 			Format,
 			ClearValueBinding,
 			TexCreateFlags,
+			ArraySize,
 			NumMips,
 			NumSamples
 		), Resource);
 	}
 	else
 	{
-		return new FOpenGLTexture(FRHITextureCreateDesc::CreateCubeArray(
+		return new FOpenGLTexture(FRHITextureCreateDesc::CreateCube(
 			TEXT("RHICreateTextureCubeFromResource"),
 			Size,
 			Format,
 			ClearValueBinding,
 			TexCreateFlags,
-			ArraySize,
 			NumMips,
 			NumSamples
 		), Resource);
