Add logging for GL buffer lock fail

diff --git a/Engine/Source/Runtime/OpenGLDrv/Public/OpenGLResources.h b/Engine/Source/Runtime/OpenGLDrv/Public/OpenGLResources.h
--- a/Engine/Source/Runtime/OpenGLDrv/Public/OpenGLResources.h
+++ b/Engine/Source/Runtime/OpenGLDrv/Public/OpenGLResources.h
@@ -643,7 +643,17 @@ public:
 			bLockBufferWasAllocated = true;
 		}
 
-		check(Data != NULL);
+		if (Data == nullptr)
+		{
+			UE_LOG(LogRHI, Fatal, TEXT("Failed to lock buffer: Resource %u, Size %u, Offset %u, bReadOnly %d, bUseMapBuffer %d, glError (0x%x)"), 
+				(uint32)Resource, 
+				InSize, 
+				InOffset, 
+				bReadOnly ? 1:0, 
+				bUseMapBuffer ? 1:0, 
+				glGetError());
+		}
+		
 		return Data;
 	}
 
@@ -710,7 +720,16 @@ public:
 			bLockBufferWasAllocated = true;
 		}
 
-		check(Data != NULL);
+		if (Data == nullptr)
+		{
+			UE_LOG(LogRHI, Fatal, TEXT("Failed to lock buffer (write only): Resource %u, Size %u, Offset %u, bUseMapBuffer %d, glError (0x%x)"), 
+				(uint32)Resource, 
+				InSize, 
+				InOffset, 
+				bUseMapBuffer ? 1:0, 
+				glGetError());
+		}
+
 		return Data;
 	}
 
