Fix a dangling ptr bug with FOpenGLErrorScope.
caught with asan.

diff --git a/Engine/Source/Runtime/OpenGLDrv/Public/OpenGLUtil.h b/Engine/Source/Runtime/OpenGLDrv/Public/OpenGLUtil.h
--- a/Engine/Source/Runtime/OpenGLDrv/Public/OpenGLUtil.h
+++ b/Engine/Source/Runtime/OpenGLDrv/Public/OpenGLUtil.h
@@ -68,12 +68,12 @@ extern bool PlatformOpenGLContextValid();
 
 	struct FOpenGLErrorScope
 	{
-		const TCHAR* FunctionName;
+		const char* FunctionName;
 		const TCHAR* Filename;
 		const uint32 Line;
 
 		FOpenGLErrorScope(
-			const TCHAR* InFunctionName,
+			const char* InFunctionName,
 			const TCHAR* InFilename,
 			const uint32 InLine)
 			: FunctionName(InFunctionName)
@@ -81,7 +81,7 @@ extern bool PlatformOpenGLContextValid();
 			, Line(InLine)
 		{
 #if ENABLE_VERIFY_GL_TRACE
-			UE_LOG(LogRHI, Log, TEXT("log before %s(%d): %s"), InFilename, InLine, InFunctionName);
+			UE_LOG(LogRHI, Log, TEXT("log before %s(%d): %s"), InFilename, InLine, ANSI_TO_TCHAR(InFunctionName));
 #endif
 			CheckForErrors(0);
 		}
@@ -89,7 +89,7 @@ extern bool PlatformOpenGLContextValid();
 		~FOpenGLErrorScope()
 		{
 #if ENABLE_VERIFY_GL_TRACE
-			UE_LOG(LogRHI, Log, TEXT("log after  %s(%d): %s"), Filename, Line, FunctionName);
+			UE_LOG(LogRHI, Log, TEXT("log after  %s(%d): %s"), Filename, Line, ANSI_TO_TCHAR(FunctionName));
 
 #endif
 
@@ -104,15 +104,15 @@ extern bool PlatformOpenGLContextValid();
 			if (ErrorCode != GL_NO_ERROR)
 			{
 				const TCHAR* PrefixStrings[] = { TEXT("Before "), TEXT("During ") };
-				VerifyOpenGLResult(ErrorCode,PrefixStrings[BeforeOrAfter],FunctionName,Filename,Line);
+				VerifyOpenGLResult(ErrorCode,PrefixStrings[BeforeOrAfter], ANSI_TO_TCHAR(FunctionName),Filename,Line);
 			}
 		}
 	};
 	#define MACRO_TOKENIZER(IdentifierName, Msg, FileName, LineNumber) FOpenGLErrorScope IdentifierName_ ## LineNumber (Msg, FileName, LineNumber)
 	#define MACRO_TOKENIZER2(IdentifierName, Msg, FileName, LineNumber) MACRO_TOKENIZER(IdentiferName, Msg, FileName, LineNumber)
 	#define VERIFY_GL_SCOPE_WITH_MSG_STR(MsgStr) CHECK_EXPECTED_GL_THREAD(); MACRO_TOKENIZER2(ErrorScope_, MsgStr, TEXT(__FILE__), __LINE__)
-	#define VERIFY_GL_SCOPE() VERIFY_GL_SCOPE_WITH_MSG_STR(ANSI_TO_TCHAR(__FUNCTION__))
-	#define VERIFY_GL_FUNC(Func, ...) { VERIFY_GL_SCOPE_WITH_MSG_STR(TEXT(#Func)); Func(__VA_ARGS__); }
+	#define VERIFY_GL_SCOPE() VERIFY_GL_SCOPE_WITH_MSG_STR(__FUNCTION__)
+	#define VERIFY_GL_FUNC(Func, ...) { VERIFY_GL_SCOPE_WITH_MSG_STR((#Func)); Func(__VA_ARGS__); }
 
 	/**
 	 * Some important GL calls are trapped individually.
