Silence SA false-positive due to differences between FOpenGLBase and FOpenGLES
#rnx
UE-147674

diff --git a/Engine/Source/Runtime/OpenGLDrv/Public/OpenGLResources.h b/Engine/Source/Runtime/OpenGLDrv/Public/OpenGLResources.h
--- a/Engine/Source/Runtime/OpenGLDrv/Public/OpenGLResources.h
+++ b/Engine/Source/Runtime/OpenGLDrv/Public/OpenGLResources.h
@@ -1334,7 +1334,7 @@ public:
 	// Returns true if this texture will perform automatic MSAA resolve on tile-based GPUs.
 	bool IsTiledMSAA() const
     {
-		return NumSamplesInternal > 1 && NumSamplesInternal <= FOpenGL::GetMaxMSAASamplesTileMem();
+		return NumSamplesInternal > 1 && NumSamplesInternal <= FOpenGL::GetMaxMSAASamplesTileMem(); //-V547
     }
 
 	// Returns the number of samples used during rendering. These samples may be kept in on-chip GPU tiled memory.
