OpenGL: Cache NULLPS in to avoid shadermap lookup on each depth drawcall
OpenGL: Do not invalidate Stencil target if Depth attachment does not use Stencil

diff --git a/Engine/Source/Runtime/OpenGLDrv/Private/OpenGLDevice.cpp b/Engine/Source/Runtime/OpenGLDrv/Private/OpenGLDevice.cpp
--- a/Engine/Source/Runtime/OpenGLDrv/Private/OpenGLDevice.cpp
+++ b/Engine/Source/Runtime/OpenGLDrv/Private/OpenGLDevice.cpp
@@ -227,12 +227,16 @@ void FOpenGLDynamicRHI::RHIBeginScene()
 	}
 
 	BeginSceneContextType = (int32)PlatformOpenGLCurrentContext(PlatformDevice);
+
+	// recache NULL shader as it can change with ODSC
+	NULLPixelShaderRHI = GetNULLPixelShader();
 }
 
 void FOpenGLDynamicRHI::RHIEndScene()
 {
 	ResourceTableFrameCounter = INDEX_NONE;
 	BeginSceneContextType = CONTEXT_Other;
+	NULLPixelShaderRHI = nullptr;
 }
 
 #if PLATFORM_ANDROID
@@ -1046,6 +1050,7 @@ FOpenGLDynamicRHI::FOpenGLDynamicRHI()
 ,   BeginSceneContextType(CONTEXT_Other)
 ,	PlatformDevice(NULL)
 ,	GPUProfilingData(this)
+,	NULLPixelShaderRHI(nullptr)
 {
 	check(Singleton == nullptr);
 	Singleton = this;
diff --git a/Engine/Source/Runtime/OpenGLDrv/Private/OpenGLRenderTarget.cpp b/Engine/Source/Runtime/OpenGLDrv/Private/OpenGLRenderTarget.cpp
--- a/Engine/Source/Runtime/OpenGLDrv/Private/OpenGLRenderTarget.cpp
+++ b/Engine/Source/Runtime/OpenGLDrv/Private/OpenGLRenderTarget.cpp
@@ -1169,8 +1169,10 @@ void FOpenGLDynamicRHI::RHIEndRenderPass()
 	{
 		ERenderTargetActions DepthActions = GetDepthActions(RenderPassInfo.DepthStencilRenderTarget.Action);
 		ERenderTargetActions StencilActions = GetStencilActions(RenderPassInfo.DepthStencilRenderTarget.Action);
+		const bool bIsUsingStencil = RenderPassInfo.DepthStencilRenderTarget.ExclusiveDepthStencil.IsUsingStencil();
+
 		bool bDiscardDepth = GetStoreAction(DepthActions) == ERenderTargetStoreAction::ENoAction;
-		bool bDiscardStencil = GetStoreAction(StencilActions) == ERenderTargetStoreAction::ENoAction;
+		bool bDiscardStencil = bIsUsingStencil && GetStoreAction(StencilActions) == ERenderTargetStoreAction::ENoAction;
 		if (bDiscardDepth || bDiscardStencil)
 		{
 			RHIDiscardRenderTargets(bDiscardDepth, bDiscardStencil, 0);
diff --git a/Engine/Source/Runtime/OpenGLDrv/Private/OpenGLShaders.cpp b/Engine/Source/Runtime/OpenGLDrv/Private/OpenGLShaders.cpp
--- a/Engine/Source/Runtime/OpenGLDrv/Private/OpenGLShaders.cpp
+++ b/Engine/Source/Runtime/OpenGLDrv/Private/OpenGLShaders.cpp
@@ -194,10 +194,20 @@ static void SetDeletedProgramStats(GLuint Program)
 // Create any resources that are required by internal ogl rhi functions.
 void FOpenGLDynamicRHI::SetupRecursiveResources()
 {
-	auto ShaderMap = GetGlobalShaderMap(GMaxRHIFeatureLevel);
+	NULLPixelShaderRHI = GetNULLPixelShader();
+}
+
+FRHIPixelShader* FOpenGLDynamicRHI::GetNULLPixelShader() const
+{
+	if (NULLPixelShaderRHI)
+	{
+		return NULLPixelShaderRHI;
+	}
+	else
 	{
+		auto ShaderMap = GetGlobalShaderMap(GMaxRHIFeatureLevel);
 		TShaderMapRef<FNULLPS> PixelShader(ShaderMap);
-		PixelShader.GetPixelShader();
+		return PixelShader.GetPixelShader();
 	}
 }
 
@@ -3037,7 +3047,7 @@ FBoundShaderStateRHIRef FOpenGLDynamicRHI::RHICreateBoundShaderState_OnThisThrea
 	if (!PixelShaderRHI)
 	{
 		// use special null pixel shader when PixelShader was set to NULL
-		PixelShaderRHI = TShaderMapRef<FNULLPS>(GetGlobalShaderMap(GMaxRHIFeatureLevel)).GetPixelShader();
+		PixelShaderRHI = GetNULLPixelShader();
 	}
 
 
diff --git a/Engine/Source/Runtime/OpenGLDrv/Public/OpenGLDrv.h b/Engine/Source/Runtime/OpenGLDrv/Public/OpenGLDrv.h
--- a/Engine/Source/Runtime/OpenGLDrv/Public/OpenGLDrv.h
+++ b/Engine/Source/Runtime/OpenGLDrv/Public/OpenGLDrv.h
@@ -913,6 +913,8 @@ public:
 	}
 
 	GLuint GetOpenGLFramebuffer(uint32 NumSimultaneousRenderTargets, FOpenGLTexture** RenderTargets, const uint32* ArrayIndices, const uint32* MipmapLevels, FOpenGLTexture* DepthStencilTarget);
+
+	FRHIPixelShader* GetNULLPixelShader() const;
 	
 private:
 
@@ -989,6 +991,9 @@ private:
 	FCriticalSection CustomPresentSection;
 	TRefCountPtr<class FRHICustomPresent> CustomPresent;
 
+	// Cached NULLPS
+	FRHIPixelShader* NULLPixelShaderRHI;
+
 	void InitializeStateResources();
 
 	void SetupVertexArrays(FOpenGLContextState& ContextCache, uint32 BaseVertexIndex, FOpenGLStream* Streams, uint32 NumStreams, uint32 MaxVertices);
