Separating out OpenGL's Binary file cache into its own file.

diff --git a/Engine/Source/Runtime/OpenGLDrv/Private/OpenGLDevice.cpp b/Engine/Source/Runtime/OpenGLDrv/Private/OpenGLDevice.cpp
--- a/Engine/Source/Runtime/OpenGLDrv/Private/OpenGLDevice.cpp
+++ b/Engine/Source/Runtime/OpenGLDrv/Private/OpenGLDevice.cpp
@@ -23,6 +23,7 @@
 #include "SceneUtils.h"
 
 #include "HardwareInfo.h"
+#include "OpenGLProgramBinaryFileCache.h"
 
 #if PLATFORM_ANDROID
 #include <jni.h>
diff --git a/Engine/Source/Runtime/OpenGLDrv/Private/OpenGLProgramBinaryFileCache.cpp b/Engine/Source/Runtime/OpenGLDrv/Private/OpenGLProgramBinaryFileCache.cpp
new file mode 100644
--- /dev/null
+++ b/Engine/Source/Runtime/OpenGLDrv/Private/OpenGLProgramBinaryFileCache.cpp
@@ -0,0 +1,1038 @@
+// Copyright Epic Games, Inc. All Rights Reserved.
+
+/*=============================================================================
+	OpenGLProgramBinaryFileCache.cpp: OpenGL program binary file cache stores/loads a set of binary ogl programs.
+=============================================================================*/
+
+#include "OpenGLShaders.h"
+#include "OpenGLProgramBinaryFileCache.h"
+#include "HAL/PlatformFileManager.h"
+#include "HAL/FileManager.h"
+#include "Misc/Paths.h"
+#include "Serialization/MemoryWriter.h"
+#include "Serialization/MemoryReader.h"
+#include "OpenGLDrvPrivate.h"
+#include "Shader.h"
+#include "GlobalShader.h"
+#include "SceneUtils.h"
+#include "PsoLruCache.h"
+
+static const uint32 GBinaryProgramFileVersion = 4;
+
+TAutoConsoleVariable<int32> FOpenGLProgramBinaryCache::CVarPBCEnable(
+	TEXT("r.ProgramBinaryCache.Enable"),
+#if PLATFORM_ANDROID
+	1,	// Enabled by default on Android.
+#else
+	0,
+#endif
+	TEXT("If true, enables binary program cache. Enabled by default only on Android"),
+	ECVF_ReadOnly | ECVF_RenderThreadSafe
+);
+
+TAutoConsoleVariable<int32> FOpenGLProgramBinaryCache::CVarRestartAndroidAfterPrecompile(
+	TEXT("r.ProgramBinaryCache.RestartAndroidAfterPrecompile"),
+	1,	// Enabled by default on Android.
+	TEXT("If true, Android apps will restart after precompiling the binary program cache. Enabled by default only on Android"),
+	ECVF_ReadOnly | ECVF_RenderThreadSafe
+);
+
+static TAutoConsoleVariable<int32> CVarUseExistingBinaryFileCache(
+	TEXT("r.OpenGL.UseExistingBinaryFileCache"),
+	1,
+	TEXT("When generating a new binary cache (such as when Shader Pipeline Cache Version Guid changes) use the existing binary file cache to speed up generation of the new cache.\n")
+	TEXT("0: Always rebuild binary file cache when Pipeline Cache Version Guid changes.\n")
+	TEXT("1: When Pipeline Cache Version Guid changes re-use programs from the existing binary cache where possible (default)."),
+	ECVF_RenderThreadSafe);
+
+static int32 GMaxShaderLibProcessingTimeMS = 10;
+static FAutoConsoleVariableRef CVarMaxShaderLibProcessingTime(
+	TEXT("r.OpenGL.MaxShaderLibProcessingTime"),
+	GMaxShaderLibProcessingTimeMS,
+	TEXT("The maximum time per frame to process shader library requests in milliseconds.\n")
+	TEXT("default 10ms. Note: Driver compile time for a single program may exceed this limit."),
+	ECVF_RenderThreadSafe
+);
+
+namespace UE
+{
+	namespace OpenGL
+	{
+		bool AreBinaryProgramsCompressed()
+		{
+			static const auto CVarStoreCompressedBinaries = IConsoleManager::Get().FindConsoleVariable(TEXT("r.OpenGL.StoreCompressedProgramBinaries"));
+			return CVarStoreCompressedBinaries->GetInt() != 0;
+		}		
+	}
+}
+
+static FCriticalSection GProgramBinaryFileCacheCS;
+
+FOpenGLProgramBinaryCache::FPreviousGLProgramBinaryCacheInfo::FPreviousGLProgramBinaryCacheInfo() : NumberOfOldEntriesReused(0) {}
+FOpenGLProgramBinaryCache::FPreviousGLProgramBinaryCacheInfo::FPreviousGLProgramBinaryCacheInfo(FOpenGLProgramBinaryCache::FPreviousGLProgramBinaryCacheInfo&&) = default;
+FOpenGLProgramBinaryCache::FPreviousGLProgramBinaryCacheInfo& FOpenGLProgramBinaryCache::FPreviousGLProgramBinaryCacheInfo::operator = (FOpenGLProgramBinaryCache::FPreviousGLProgramBinaryCacheInfo&&) = default;
+FOpenGLProgramBinaryCache::FPreviousGLProgramBinaryCacheInfo::~FPreviousGLProgramBinaryCacheInfo() = default;
+
+FOpenGLProgramBinaryCache* FOpenGLProgramBinaryCache::CachePtr = nullptr;
+
+FOpenGLProgramBinaryCache::FOpenGLProgramBinaryCache(const FString& InCachePath)
+	: CachePath(InCachePath)
+	, BinaryCacheAsyncReadFileHandle(nullptr)
+	, BinaryCacheWriteFileHandle(nullptr)
+	, BinaryFileState(EBinaryFileState::Uninitialized)
+{
+	ANSICHAR* GLVersion = (ANSICHAR*)glGetString(GL_VERSION);
+	ANSICHAR* GLRenderer = (ANSICHAR*)glGetString(GL_RENDERER);
+	FString HashString;
+	HashString.Append(GLVersion);
+	HashString.Append(GLRenderer);
+	FSHAHash VersionHash;
+	FSHA1::HashBuffer(TCHAR_TO_ANSI(*HashString), HashString.Len(), VersionHash.Hash);
+
+	CacheFilename = LegacyShaderPlatformToShaderFormat(GMaxRHIShaderPlatform).ToString() + TEXT("_") + VersionHash.ToString();
+}
+
+FOpenGLProgramBinaryCache::~FOpenGLProgramBinaryCache()
+{
+	if (BinaryCacheAsyncReadFileHandle)
+	{
+		delete BinaryCacheAsyncReadFileHandle;
+	}
+	if (BinaryCacheWriteFileHandle)
+	{
+		delete BinaryCacheWriteFileHandle;
+	}
+
+	if (OnShaderPipelineCacheOpenedDelegate.IsValid())
+	{
+		FShaderPipelineCache::GetCacheOpenedDelegate().Remove(OnShaderPipelineCacheOpenedDelegate);
+	}
+
+	if (OnShaderPipelineCachePrecompilationCompleteDelegate.IsValid())
+	{
+		FShaderPipelineCache::GetPrecompilationCompleteDelegate().Remove(OnShaderPipelineCachePrecompilationCompleteDelegate);
+	}
+};
+
+bool FOpenGLProgramBinaryCache::IsEnabled()
+{
+	return CachePtr != nullptr;
+}
+
+bool FOpenGLProgramBinaryCache::IsBuildingCache()
+{
+	if (CachePtr != nullptr)
+	{
+		return CachePtr->IsBuildingCache_internal();
+	}
+	return false;
+}
+
+void FOpenGLProgramBinaryCache::Initialize()
+{
+	check(CachePtr == nullptr);
+
+	if (CVarPBCEnable.GetValueOnAnyThread() == 0)
+	{
+		UE_LOG(LogRHI, Log, TEXT("FOpenGLProgramBinaryCache disabled by r.ProgramBinaryCache.Enable=0"));
+		return;
+	}
+
+	if (!FOpenGL::SupportsProgramBinary())
+	{
+		UE_LOG(LogRHI, Warning, TEXT("FOpenGLProgramBinaryCache disabled as devices does not support program binaries"));
+		return;
+	}
+
+#if PLATFORM_ANDROID
+	if (FOpenGL::HasBinaryProgramRetrievalFailed())
+	{
+		if (FOpenGL::SupportsProgramBinary())
+		{
+			UE_LOG(LogRHI, Warning, TEXT("FOpenGLProgramBinaryCache: Device has failed to emit program binary despite SupportsProgramBinary == true. Disabling binary cache."));
+			return;
+		}
+	}
+#endif
+
+
+	FString CacheFolderPath;
+#if PLATFORM_ANDROID && USE_ANDROID_FILE
+	// @todo Lumin: Use that GetPathForExternalWrite or something?
+	extern FString GExternalFilePath;
+	CacheFolderPath = GExternalFilePath / TEXT("ProgramBinaryCache");
+
+#else
+	CacheFolderPath = FPaths::ProjectSavedDir() / TEXT("ProgramBinaryCache");
+#endif
+
+	// Remove entire ProgramBinaryCache folder if -ClearOpenGLBinaryProgramCache is specified on command line
+	if (FParse::Param(FCommandLine::Get(), TEXT("ClearOpenGLBinaryProgramCache")))
+	{
+		UE_LOG(LogRHI, Log, TEXT("ClearOpenGLBinaryProgramCache specified, deleting binary program cache folder: %s"), *CacheFolderPath);
+		FPlatformFileManager::Get().GetPlatformFile().DeleteDirectoryRecursively(*CacheFolderPath);
+	}
+
+	CachePtr = new FOpenGLProgramBinaryCache(CacheFolderPath);
+	UE_LOG(LogRHI, Log, TEXT("Enabling program binary cache as %s"), *CachePtr->GetProgramBinaryCacheFilePath());
+
+	// Add delegates for the ShaderPipelineCache precompile.
+	UE_LOG(LogRHI, Log, TEXT("FOpenGLProgramBinaryCache will be initialized when ShaderPipelineCache opens its file"));
+	CachePtr->OnShaderPipelineCacheOpenedDelegate = FShaderPipelineCache::GetCacheOpenedDelegate().AddRaw(CachePtr, &FOpenGLProgramBinaryCache::OnShaderPipelineCacheOpened);
+	CachePtr->OnShaderPipelineCachePrecompilationCompleteDelegate = FShaderPipelineCache::GetPrecompilationCompleteDelegate().AddRaw(CachePtr, &FOpenGLProgramBinaryCache::OnShaderPipelineCachePrecompilationComplete);
+}
+
+#if PLATFORM_ANDROID
+static int32 GNumRemoteProgramCompileServices = 4;
+static FAutoConsoleVariableRef CVarNumRemoteProgramCompileServices(
+	TEXT("Android.OpenGL.NumRemoteProgramCompileServices"),
+	GNumRemoteProgramCompileServices,
+	TEXT("The number of separate processes to make available to compile opengl programs.\n")
+	TEXT("0 to disable use of separate processes to precompile PSOs\n")
+	TEXT("valid range is 1-8 (4 default).")
+	,
+	ECVF_RenderThreadSafe | ECVF_ReadOnly
+);
+#endif
+
+void FOpenGLProgramBinaryCache::OnShaderPipelineCacheOpened(FString const& Name, EShaderPlatform Platform, uint32 Count, const FGuid& VersionGuid, FShaderPipelineCache::FShaderCachePrecompileContext& ShaderCachePrecompileContext)
+{
+	UE_LOG(LogRHI, Log, TEXT("Scanning Binary program cache, using Shader Pipeline Cache version %s"), *VersionGuid.ToString());
+	ScanProgramCacheFile(VersionGuid);
+	if (IsBuildingCache_internal())
+	{
+#if PLATFORM_ANDROID
+		if (GNumRemoteProgramCompileServices)
+		{
+			FAndroidOpenGL::StartAndWaitForRemoteCompileServices(GNumRemoteProgramCompileServices);
+		}
+#endif
+		ShaderCachePrecompileContext.SetPrecompilationIsSlowTask();
+	}
+}
+
+void FOpenGLProgramBinaryCache::OnShaderPipelineCachePrecompilationComplete(uint32 Count, double Seconds, const FShaderPipelineCache::FShaderCachePrecompileContext& ShaderCachePrecompileContext)
+{
+#if PLATFORM_ANDROID
+	FAndroidOpenGL::StopRemoteCompileServices();
+#endif
+
+	UE_LOG(LogRHI, Log, TEXT("OnShaderPipelineCachePrecompilationComplete: %d shaders"), Count);
+
+	// Want to ignore any subsequent Shader Pipeline Cache opening/closing, eg when loading modules
+	FShaderPipelineCache::GetCacheOpenedDelegate().Remove(OnShaderPipelineCacheOpenedDelegate);
+	FShaderPipelineCache::GetPrecompilationCompleteDelegate().Remove(OnShaderPipelineCachePrecompilationCompleteDelegate);
+	OnShaderPipelineCacheOpenedDelegate.Reset();
+	OnShaderPipelineCachePrecompilationCompleteDelegate.Reset();
+
+	check(IsBuildingCache_internal() || BinaryFileState == EBinaryFileState::ValidCacheFile);
+
+	if (IsBuildingCache_internal())
+	{
+		CloseWriteHandle();
+
+#if PLATFORM_ANDROID
+		FAndroidMisc::bNeedsRestartAfterPSOPrecompile = true;
+		if (CVarRestartAndroidAfterPrecompile.GetValueOnAnyThread() == 1)
+		{
+#if USE_ANDROID_JNI
+			extern void AndroidThunkCpp_RestartApplication(const FString & IntentString);
+			AndroidThunkCpp_RestartApplication(TEXT(""));
+#endif
+		}
+#endif
+		OpenAsyncReadHandle();
+		BinaryFileState = EBinaryFileState::ValidCacheFile;
+	}
+}
+
+// contains runtime and file information for a single program entry in the cache file.
+struct FGLProgramBinaryFileCacheEntry
+{
+	struct FGLProgramBinaryFileCacheFileInfo
+	{
+		// contains location info for a program in the binary cache file
+		FOpenGLProgramKey ShaderHasheSet;
+		uint32 ProgramOffset;
+		uint32 ProgramSize;
+
+		FGLProgramBinaryFileCacheFileInfo() : ProgramOffset(0), ProgramSize(0) { }
+
+		friend bool operator == (const FGLProgramBinaryFileCacheFileInfo& A, const FGLProgramBinaryFileCacheFileInfo& B)
+		{
+			return A.ShaderHasheSet == B.ShaderHasheSet && A.ProgramOffset == B.ProgramOffset && A.ProgramSize == B.ProgramSize;
+		}
+	} FileInfo;
+
+	// program read request.
+	TWeakPtr<IAsyncReadRequest, ESPMode::ThreadSafe> ReadRequest;
+	// read data
+	TArray<uint8> ProgramBinaryData;
+	// debugging use only, index of program as encountered during scan, -1 if new.
+	int32 ProgramIndex;
+
+	enum class EGLProgramState : uint8
+	{
+		Unset,
+		ProgramStored, // program exists in binary cache but has not yet been loaded
+		ProgramLoading,	// program has started async loading.
+		ProgramLoaded,	// program has loaded and is ready for GL object creation.
+		ProgramAvailable, // program has loaded from binary cache and is available for use with GL.
+		ProgramComplete // program has been either added by rhi or handed over to rhi after being made available to it.
+	};
+	EGLProgramState GLProgramState;
+
+	// if != 0 then prepared runtime GL program name:
+	GLuint GLProgramId;
+
+	FGLProgramBinaryFileCacheEntry() : ProgramIndex(-1), GLProgramState(EGLProgramState::Unset), GLProgramId(0) { }
+
+	friend uint32 GetTypeHash(const FGLProgramBinaryFileCacheEntry& Key)
+	{
+		return FCrc::MemCrc32(&Key, sizeof(Key));
+	}
+};
+
+static FCriticalSection GPendingGLProgramCreateRequestsCS;
+
+// Scan the binary cache file and build a record of all programs.
+void FOpenGLProgramBinaryCache::ScanProgramCacheFile(const FGuid& ShaderPipelineCacheVersionGuid)
+{
+	UE_LOG(LogRHI, Log, TEXT("OnShaderScanProgramCacheFile"));
+	FScopeLock Lock(&GProgramBinaryFileCacheCS);
+	FString ProgramCacheFilename = GetProgramBinaryCacheFilePath();
+	FString ProgramCacheFilenameTemp = GetProgramBinaryCacheFilePath() + TEXT(".scan");
+
+	IPlatformFile& PlatformFile = FPlatformFileManager::Get().GetPlatformFile();
+
+	check(BinaryFileState == EBinaryFileState::Uninitialized);
+
+	bool bBinaryFileIsValid = false;
+	bool bBinaryFileIsValidAndGuidMatch = false;
+
+	// Try to move the file to a temporary filename before the scan, so we won't try to read it again if it's corrupted
+	PlatformFile.DeleteFile(*ProgramCacheFilenameTemp);
+	PlatformFile.MoveFile(*ProgramCacheFilenameTemp, *ProgramCacheFilename);
+
+	TUniquePtr<FArchive> FileReader(IFileManager::Get().CreateFileReader(*ProgramCacheFilenameTemp));
+	if (FileReader)
+	{
+		UE_LOG(LogRHI, Log, TEXT("OnShaderScanProgramCacheFile : Opened %s"), *ProgramCacheFilenameTemp);
+		FArchive& Ar = *FileReader;
+		uint32 Version = 0;
+		Ar << Version;
+		if (Version == GBinaryProgramFileVersion)
+		{
+			FGuid BinaryCacheGuid;
+			Ar << BinaryCacheGuid;
+			bool bCacheUsesCompressedBinaries;
+			Ar << bCacheUsesCompressedBinaries;
+
+			const bool bUseCompressedProgramBinaries = UE::OpenGL::AreBinaryProgramsCompressed();
+			bBinaryFileIsValid = (bUseCompressedProgramBinaries == bCacheUsesCompressedBinaries);
+			bBinaryFileIsValidAndGuidMatch = bBinaryFileIsValid && (!ShaderPipelineCacheVersionGuid.IsValid() || ShaderPipelineCacheVersionGuid == BinaryCacheGuid);
+
+			if (CVarUseExistingBinaryFileCache.GetValueOnAnyThread() == 0 && bBinaryFileIsValidAndGuidMatch == false)
+			{
+				// If we dont want to use the existing binary cache and the guids have changed then rebuild the binary file.
+				bBinaryFileIsValid = false;
+			}
+		}
+
+		if (bBinaryFileIsValid)
+		{
+			const uint32 ProgramBinaryStart = Ar.Tell();
+
+			// Search the file for the end record.
+			bool bFoundEndRecord = false;
+			int32 ProgramIndex = 0;
+			while (!Ar.AtEnd())
+			{
+				check(bFoundEndRecord == false); // There should be no additional data after the eof record.
+
+				FOpenGLProgramKey ProgramKey;
+				uint32 ProgramBinarySize = 0;
+				Ar << ProgramKey;
+				Ar << ProgramBinarySize;
+				uint32 ProgramBinaryOffset = Ar.Tell();
+				if (ProgramBinarySize == 0)
+				{
+					if (ProgramKey == FOpenGLProgramKey())
+					{
+						bFoundEndRecord = true;
+					}
+					else
+					{
+						// Note: This should not happen with new code. We can no longer write out records with 0 program size. see AppendProgramBinaryFile.
+						UE_LOG(LogRHI, Warning, TEXT("FOpenGLProgramBinaryCache::ScanProgramCacheFile : encountered 0 sized program during binary program cache scan"));
+					}
+				}
+				Ar.Seek(ProgramBinaryOffset + ProgramBinarySize);
+			}
+
+			if (bFoundEndRecord)
+			{
+				Ar.Seek(ProgramBinaryStart);
+				while (!Ar.AtEnd())
+				{
+					FOpenGLProgramKey ProgramKey;
+					uint32 ProgramBinarySize = 0;
+					Ar << ProgramKey;
+					Ar << ProgramBinarySize;
+
+					if (ProgramBinarySize > 0)
+					{
+						FGLProgramBinaryFileCacheEntry* NewEntry = new FGLProgramBinaryFileCacheEntry();
+						NewEntry->FileInfo.ShaderHasheSet = ProgramKey;
+						NewEntry->ProgramIndex = ProgramIndex++;
+
+						uint32 ProgramBinaryOffset = Ar.Tell();
+						NewEntry->FileInfo.ProgramSize = ProgramBinarySize;
+						NewEntry->FileInfo.ProgramOffset = ProgramBinaryOffset;
+
+						if (bBinaryFileIsValidAndGuidMatch)
+						{
+							ProgramEntryContainer.Emplace(TUniquePtr<FGLProgramBinaryFileCacheEntry>(NewEntry));
+
+							// check to see if any of the shaders are already loaded and so we should serialize the binary
+							bool bAllShadersLoaded = true;
+							for (int32 i = 0; i < CrossCompiler::NUM_SHADER_STAGES && bAllShadersLoaded; i++)
+							{
+								bAllShadersLoaded = ProgramKey.ShaderHashes[i] == FSHAHash() || ShaderIsLoaded(ProgramKey.ShaderHashes[i]);
+							}
+							if (bAllShadersLoaded)
+							{
+								FPlatformMisc::LowLevelOutputDebugStringf(TEXT("*** All shaders for %s already loaded\n"), *ProgramKey.ToString());
+								NewEntry->ProgramBinaryData.AddUninitialized(ProgramBinarySize);
+								Ar.Serialize(NewEntry->ProgramBinaryData.GetData(), ProgramBinarySize);
+								NewEntry->GLProgramState = FGLProgramBinaryFileCacheEntry::EGLProgramState::ProgramLoaded;
+								CompleteLoadedGLProgramRequest_internal(NewEntry);
+							}
+							else
+							{
+								NewEntry->GLProgramState = FGLProgramBinaryFileCacheEntry::EGLProgramState::ProgramStored;
+							}
+							AddProgramFileEntryToMap(NewEntry);
+						}
+						else
+						{
+							check(!PreviousBinaryCacheInfo.ProgramToOldBinaryCacheMap.Contains(ProgramKey));
+							PreviousBinaryCacheInfo.ProgramToOldBinaryCacheMap.Emplace(ProgramKey, TUniquePtr<FGLProgramBinaryFileCacheEntry>(NewEntry));
+						}
+						Ar.Seek(ProgramBinaryOffset + ProgramBinarySize);
+					}
+				}
+
+				if (bBinaryFileIsValidAndGuidMatch)
+				{
+					UE_LOG(LogRHI, Log, TEXT("Program Binary cache: Found %d cached programs, end record found: %d"), ProgramIndex, (uint32)bFoundEndRecord);
+					FileReader->Close();
+					// Rename the file back after a successful scan.
+					PlatformFile.MoveFile(*ProgramCacheFilename, *ProgramCacheFilenameTemp);
+				}
+				else
+				{
+					UE_LOG(LogRHI, Log, TEXT("Program Binary cache: ShaderPipelineCache changed, regenerating for new pipeline cache. Existing cache contains %d programs, using it to populate."), PreviousBinaryCacheInfo.ProgramToOldBinaryCacheMap.Num());
+					// Not closing the scan source file, we're using it to move shaders from the old cache.
+					PreviousBinaryCacheInfo.OldCacheArchive = MoveTemp(FileReader);
+					PreviousBinaryCacheInfo.OldCacheFilename = ProgramCacheFilenameTemp;
+				}
+			}
+			else
+			{
+				// failed to find sentinel record, the file was not finalized.
+				UE_LOG(LogRHI, Warning, TEXT("ScanProgramCacheFile - incomplete binary cache file encountered. Rebuilding binary program cache."));
+				FileReader->Close();
+				bBinaryFileIsValid = false;
+				bBinaryFileIsValidAndGuidMatch = false;
+			}
+		}
+
+		if (!bBinaryFileIsValid)
+		{
+			UE_LOG(LogRHI, Log, TEXT("OnShaderScanProgramCacheFile : binary file version invalid"));
+		}
+
+		if (bBinaryFileIsValidAndGuidMatch)
+		{
+			OpenAsyncReadHandle();
+			BinaryFileState = EBinaryFileState::ValidCacheFile;
+		}
+	}
+	else
+	{
+		UE_LOG(LogRHI, Log, TEXT("OnShaderScanProgramCacheFile : Failed to open %s"), *ProgramCacheFilename);
+	}
+
+	if (!bBinaryFileIsValid)
+	{
+		// Attempt to remove any existing binary cache or temp files (eg for different driver version)
+		UE_LOG(LogRHI, Log, TEXT("Deleting binary program cache folder: %s"), *CachePath);
+		PlatformFile.DeleteDirectoryRecursively(*CachePath);
+
+		// Create
+		if (!PlatformFile.CreateDirectoryTree(*CachePath))
+		{
+			UE_LOG(LogRHI, Warning, TEXT("Failed to create directory for a program binary cache. Cache will be disabled: %s"), *CachePath);
+			return;
+		}
+	}
+
+	if (!bBinaryFileIsValid || !bBinaryFileIsValidAndGuidMatch)
+	{
+		if (OpenWriteHandle())
+		{
+			BinaryFileState = bBinaryFileIsValid && !bBinaryFileIsValidAndGuidMatch ? EBinaryFileState::BuildingCacheFileWithMove : EBinaryFileState::BuildingCacheFile;
+
+			// save header
+			FArchive& Ar = *BinaryCacheWriteFileHandle;
+			uint32 Version = GBinaryProgramFileVersion;
+			Ar << Version;
+			FGuid BinaryCacheGuid = ShaderPipelineCacheVersionGuid;
+			Ar << BinaryCacheGuid;
+			bool bWritingCompressedBinaries = (UE::OpenGL::AreBinaryProgramsCompressed());
+			Ar << bWritingCompressedBinaries;
+		}
+		else
+		{
+			// Binary cache file cannot be used, failed to open output file.
+			BinaryFileState = EBinaryFileState::Uninitialized;
+			RHIGetPanicDelegate().ExecuteIfBound(FName("FailedBinaryProgramArchiveOpen"));
+			UE_LOG(LogRHI, Fatal, TEXT("ScanProgramCacheFile - Failed to open binary cache."));
+		}
+	}
+}
+
+// add the GLProgramBinaryFileCacheEntry into the runtime lookup containers.
+void FOpenGLProgramBinaryCache::AddProgramFileEntryToMap(FGLProgramBinaryFileCacheEntry* NewEntry)
+{
+	const FOpenGLProgramKey& ProgramKey = NewEntry->FileInfo.ShaderHasheSet;
+	check(!ProgramToBinaryMap.Contains(ProgramKey));
+	ProgramToBinaryMap.Add(ProgramKey, NewEntry);
+
+	UE_LOG(LogRHI, Verbose, TEXT("AddProgramFileEntryToMap : Adding program: %s"), *ProgramKey.ToString());
+
+	for (int i = 0; i < CrossCompiler::NUM_NON_COMPUTE_SHADER_STAGES; ++i)
+	{
+		const FSHAHash& ShaderHash = ProgramKey.ShaderHashes[i];
+
+		if (ShaderHash != FSHAHash())
+		{
+			if (ShaderToProgramsMap.Contains(ShaderHash))
+			{
+				ShaderToProgramsMap[ShaderHash].Add(NewEntry);
+			}
+			else
+			{
+				ShaderToProgramsMap.Add(ShaderHash, NewEntry);
+			}
+		}
+	}
+}
+
+bool FOpenGLProgramBinaryCache::OpenWriteHandle()
+{
+	check(BinaryCacheWriteFileHandle == nullptr);
+	check(BinaryCacheAsyncReadFileHandle == nullptr);
+
+	// perform file writing to a file temporary filename so we don't attempt to use the file later if the write session is interrupted
+	FString ProgramCacheFilename = GetProgramBinaryCacheFilePath();
+	FString ProgramCacheFilenameWrite = ProgramCacheFilename + TEXT(".write");
+
+	BinaryCacheWriteFileHandle = IFileManager::Get().CreateFileWriter(*ProgramCacheFilenameWrite, EFileWrite::FILEWRITE_None);
+
+	UE_CLOG(BinaryCacheWriteFileHandle == nullptr, LogRHI, Warning, TEXT("Failed to open OGL binary cache output file."));
+
+	return BinaryCacheWriteFileHandle != nullptr;
+}
+
+void FOpenGLProgramBinaryCache::CloseWriteHandle()
+{
+	if (BinaryFileState == EBinaryFileState::BuildingCacheFileWithMove)
+	{
+		UE_LOG(LogRHI, Log, TEXT("FOpenGLProgramBinaryCache: Deleting previous binary program cache (%s), reused %d programs from a total of %d."), *PreviousBinaryCacheInfo.OldCacheFilename, PreviousBinaryCacheInfo.NumberOfOldEntriesReused, ProgramToBinaryMap.Num());
+
+		// clean up references to old cache.
+		PreviousBinaryCacheInfo.OldCacheArchive->Close();
+		PreviousBinaryCacheInfo.OldCacheArchive = nullptr;
+		IPlatformFile& PlatformFile = FPlatformFileManager::Get().GetPlatformFile();
+		PlatformFile.DeleteFile(*PreviousBinaryCacheInfo.OldCacheFilename);
+		PreviousBinaryCacheInfo.OldCacheFilename.Empty();
+		PreviousBinaryCacheInfo.ProgramToOldBinaryCacheMap.Empty();
+	}
+
+	check(BinaryCacheWriteFileHandle != nullptr);
+
+	AppendProgramBinaryFileEofEntry(*BinaryCacheWriteFileHandle);
+	bool bArchiveFailed = BinaryCacheWriteFileHandle->IsError() || BinaryCacheWriteFileHandle->IsCriticalError();
+
+	BinaryCacheWriteFileHandle->Close();
+	delete BinaryCacheWriteFileHandle;
+	BinaryCacheWriteFileHandle = nullptr;
+
+	if (bArchiveFailed)
+	{
+		RHIGetPanicDelegate().ExecuteIfBound(FName("FailedBinaryProgramArchiveWrite"));
+		UE_LOG(LogRHI, Fatal, TEXT("CloseWriteHandle - FArchive error bit set, failed to write binary cache."));
+	}
+
+	// rename the temp filename back to the final filename
+	FString ProgramCacheFilename = GetProgramBinaryCacheFilePath();
+	FString ProgramCacheFilenameWrite = ProgramCacheFilename + TEXT(".write");
+	IPlatformFile& PlatformFile = FPlatformFileManager::Get().GetPlatformFile();
+	PlatformFile.DeleteFile(*ProgramCacheFilename); // file should never exist, but for safety
+	PlatformFile.MoveFile(*ProgramCacheFilename, *ProgramCacheFilenameWrite);
+}
+
+void FOpenGLProgramBinaryCache::OpenAsyncReadHandle()
+{
+	check(BinaryCacheAsyncReadFileHandle == nullptr);
+
+	FString ProgramCacheFilename = GetProgramBinaryCacheFilePath();
+	BinaryCacheAsyncReadFileHandle = FPlatformFileManager::Get().GetPlatformFile().OpenAsyncRead(*ProgramCacheFilename);
+	checkf(BinaryCacheAsyncReadFileHandle, TEXT("Could not opan an async file")); // this generally cannot fail because it is async
+}
+
+/* dead code, needs removal
+void FOpenGLProgramBinaryCache::CloseAsyncReadHandle()
+{
+	// wait for any pending reads.
+	{
+		FScopeLock Lock(&GPendingGLProgramCreateRequestsCS);
+		for (FGLProgramBinaryFileCacheEntry* CreateRequest : PendingGLProgramCreateRequests)
+		{
+			TSharedPtr<IAsyncReadRequest, ESPMode::ThreadSafe> ReadRequest = CreateRequest->ReadRequest.Pin();
+			if (ReadRequest.IsValid())
+			{
+				ReadRequest->WaitCompletion();
+				CreateRequest->ReadRequest = nullptr;
+			}
+		}
+	}
+
+	delete BinaryCacheAsyncReadFileHandle;
+	BinaryCacheAsyncReadFileHandle = nullptr;
+}*/
+
+// Called when a new program has been created by OGL RHI, creates the binary cache if it's invalid and then appends the new program details to the file and runtime containers.
+void FOpenGLProgramBinaryCache::AppendGLProgramToBinaryCache(const FOpenGLProgramKey& ProgramKey, GLuint Program, TArray<uint8>& CachedProgramBinaryOUT)
+{
+	if (IsBuildingCache_internal() == false)
+	{
+		return;
+	}
+
+	FScopeLock Lock(&GProgramBinaryFileCacheCS);
+
+	AddUniqueGLProgramToBinaryCache(ProgramKey, Program, CachedProgramBinaryOUT);
+}
+
+// Add the program to the binary cache if it does not already exist.
+void FOpenGLProgramBinaryCache::AddUniqueGLProgramToBinaryCache(const FOpenGLProgramKey& ProgramKey, GLuint Program, TArray<uint8>& CachedProgramBinaryOUT)
+{
+	// Add to runtime and disk.
+	const FOpenGLProgramKey& ProgramHash = ProgramKey;
+
+	// Check we dont already have this: Something could be in the cache but still reach this point if OnSharedShaderCodeRequest(s) have not occurred.
+	if (!ProgramToBinaryMap.Contains(ProgramHash))
+	{
+		uint32 ProgramBinaryOffset = 0, ProgramBinarySize = 0;
+
+		FOpenGLProgramKey SerializedProgramKey = ProgramKey;
+		if (ensure(UE::OpenGL::GetProgramBinaryFromGLProgram(Program, CachedProgramBinaryOUT)))
+		{
+			AddProgramBinaryDataToBinaryCache(CachedProgramBinaryOUT, ProgramKey);
+		}
+		else
+		{
+			// we've encountered a problem with this program and there's nothing to write.
+			// This likely means the device will never be able to use this program.
+			// Panic!
+			RHIGetPanicDelegate().ExecuteIfBound(FName("FailedBinaryProgramWrite"));
+			UE_LOG(LogRHI, Fatal, TEXT("AppendProgramBinaryFile Binary program returned 0 bytes!"));
+			// Panic!
+		}
+	}
+}
+
+void FOpenGLProgramBinaryCache::CacheProgramBinary(const FOpenGLProgramKey& ProgramKey, TArray<uint8>& CachedProgramBinary)
+{
+	if (CachePtr)
+	{
+		if (CachePtr->IsBuildingCache_internal() == false)
+		{
+			return;
+		}
+
+		FScopeLock Lock(&GProgramBinaryFileCacheCS);
+
+		if (!CachePtr->ProgramToBinaryMap.Contains(ProgramKey))
+		{
+			CachePtr->AddProgramBinaryDataToBinaryCache(CachedProgramBinary, ProgramKey);
+		}
+	}
+}
+
+// Serialize out the program binary data and add to runtime structures.
+void FOpenGLProgramBinaryCache::AddProgramBinaryDataToBinaryCache(TArray<uint8>& BinaryProgramData, const FOpenGLProgramKey& ProgramKey)
+{
+	FScopeLock Lock(&GProgramBinaryFileCacheCS);
+	FArchive& Ar = *BinaryCacheWriteFileHandle;
+	// Serialize to output file:
+	FOpenGLProgramKey SerializedProgramKey = ProgramKey;
+	uint32 ProgramBinarySize = (uint32)BinaryProgramData.Num();
+	Ar << SerializedProgramKey;
+	uint32 ProgramBinaryOffset = Ar.Tell();
+	Ar << ProgramBinarySize;
+	Ar.Serialize(BinaryProgramData.GetData(), ProgramBinarySize);
+	if (UE::OpenGL::AreBinaryProgramsCompressed())
+	{
+		static uint32 TotalUncompressed = 0;
+		static uint32 TotalCompressed = 0;
+
+		UE::OpenGL::FCompressedProgramBinaryHeader* Header = (UE::OpenGL::FCompressedProgramBinaryHeader*)BinaryProgramData.GetData();
+		TotalUncompressed += Header->UncompressedSize;
+		TotalCompressed += BinaryProgramData.Num();
+
+		UE_LOG(LogRHI, Verbose, TEXT("AppendProgramBinaryFile: total Uncompressed: %d, total Compressed %d, Total saved so far: %d"), TotalUncompressed, TotalCompressed, TotalUncompressed - TotalCompressed);
+	}
+
+	FGLProgramBinaryFileCacheEntry* NewIndexEntry = new FGLProgramBinaryFileCacheEntry();
+	ProgramEntryContainer.Emplace(TUniquePtr<FGLProgramBinaryFileCacheEntry>(NewIndexEntry));
+
+	// Store the program file descriptor in the runtime program/shader container:
+	NewIndexEntry->GLProgramState = FGLProgramBinaryFileCacheEntry::EGLProgramState::ProgramStored;
+	NewIndexEntry->FileInfo.ProgramOffset = ProgramBinaryOffset;
+	NewIndexEntry->FileInfo.ProgramSize = ProgramBinarySize;
+	NewIndexEntry->ProgramIndex = ProgramToBinaryMap.Num();
+	NewIndexEntry->FileInfo.ShaderHasheSet = ProgramKey;
+	AddProgramFileEntryToMap(NewIndexEntry);
+}
+
+void FOpenGLProgramBinaryCache::AppendProgramBinaryFileEofEntry(FArchive& Ar)
+{
+	// write out an all zero record that signifies eof.
+	FOpenGLProgramKey SerializedProgramKey;
+	Ar << SerializedProgramKey;
+	uint32 ProgramBinarySize = 0;
+	Ar << ProgramBinarySize;
+}
+
+void FOpenGLProgramBinaryCache::Shutdown()
+{
+	if (CachePtr)
+	{
+		delete CachePtr;
+		CachePtr = nullptr;
+	}
+}
+
+void FOpenGLProgramBinaryCache::CacheProgram(GLuint Program, const FOpenGLProgramKey& ProgramKey, TArray<uint8>& CachedProgramBinaryOUT)
+{
+	if (CachePtr)
+	{
+		CachePtr->AppendGLProgramToBinaryCache(ProgramKey, Program, CachedProgramBinaryOUT);
+	}
+}
+
+bool FOpenGLProgramBinaryCache::UseCachedProgram(GLuint& ProgramOUT, const FOpenGLProgramKey& ProgramKey, TArray<uint8>& CachedProgramBinaryOUT)
+{
+	if (CachePtr)
+	{
+		return CachePtr->UseCachedProgram_internal(ProgramOUT, ProgramKey, CachedProgramBinaryOUT);
+	}
+	return false;
+}
+
+
+namespace UE
+{
+	namespace OpenGL
+	{
+		extern bool CreateGLProgramFromBinary(GLuint& GLProgramIDOUT, const TArray<uint8>& ProgramBinaryData);
+		extern void OnGLProgramLoadedFromBinaryCache(const FOpenGLProgramKey& ProgramKey, TArray<uint8>&& ProgramBinaryData);
+	}
+}
+
+bool FOpenGLProgramBinaryCache::UseCachedProgram_internal(GLuint& ProgramOUT, const FOpenGLProgramKey& ProgramKey, TArray<uint8>& CachedProgramBinaryOUT)
+{
+	SCOPE_CYCLE_COUNTER(STAT_OpenGLUseCachedProgramTime);
+
+	FGLProgramBinaryFileCacheEntry** ProgramBinRefPtr = nullptr;
+
+	FScopeLock Lock(&GProgramBinaryFileCacheCS);
+
+	ProgramBinRefPtr = ProgramToBinaryMap.Find(ProgramKey);
+
+	if (ProgramBinRefPtr)
+	{
+		FGLProgramBinaryFileCacheEntry* FoundProgram = *ProgramBinRefPtr;
+		check(FoundProgram->FileInfo.ShaderHasheSet == ProgramKey);
+
+		TSharedPtr<IAsyncReadRequest, ESPMode::ThreadSafe> LocalReadRequest = FoundProgram->ReadRequest.Pin();
+		bool bHasReadRequest = LocalReadRequest.IsValid();
+		check(!bHasReadRequest);
+
+		// by this point the program must be either available or no attempt to load from shader library has occurred.
+		checkf(FoundProgram->GLProgramState == FGLProgramBinaryFileCacheEntry::EGLProgramState::ProgramStored
+			|| FoundProgram->GLProgramState == FGLProgramBinaryFileCacheEntry::EGLProgramState::ProgramAvailable,
+			TEXT("Unexpected program state:  (%s) == %d"), *ProgramKey.ToString(), (int32)FoundProgram->GLProgramState);
+
+		if (FoundProgram->GLProgramState == FGLProgramBinaryFileCacheEntry::EGLProgramState::ProgramAvailable)
+		{
+			UE_LOG(LogRHI, Log, TEXT("UseCachedProgram : Program (%s) GLid = %x is ready!"), *ProgramKey.ToString(), FoundProgram->GLProgramId);
+			ProgramOUT = FoundProgram->GLProgramId;
+
+			// GLProgram has been handed over.
+			FoundProgram->GLProgramId = 0;
+			FoundProgram->GLProgramState = FGLProgramBinaryFileCacheEntry::EGLProgramState::ProgramComplete;
+			return true;
+		}
+		else
+		{
+			UE_LOG(LogRHI, Log, TEXT("UseCachedProgram : %s was not ready when needed!! (state %d)"), *ProgramKey.ToString(), (uint32)FoundProgram->GLProgramState);
+		}
+	}
+	else if (BinaryFileState == EBinaryFileState::BuildingCacheFileWithMove)
+	{
+		// We're building the new cache using the original cache to warm:
+		TUniquePtr<FGLProgramBinaryFileCacheEntry>* FoundExistingBinary = PreviousBinaryCacheInfo.ProgramToOldBinaryCacheMap.Find(ProgramKey);
+		if (FoundExistingBinary)
+		{
+			TUniquePtr<FGLProgramBinaryFileCacheEntry>& ExistingBinary = *FoundExistingBinary;
+			// read old binary:
+			CachedProgramBinaryOUT.SetNumUninitialized(ExistingBinary->FileInfo.ProgramSize);
+			PreviousBinaryCacheInfo.OldCacheArchive->Seek(ExistingBinary->FileInfo.ProgramOffset);
+			PreviousBinaryCacheInfo.OldCacheArchive->Serialize(CachedProgramBinaryOUT.GetData(), ExistingBinary->FileInfo.ProgramSize);
+			bool bSuccess = UE::OpenGL::CreateGLProgramFromBinary(ProgramOUT, CachedProgramBinaryOUT);
+			if (!bSuccess)
+			{
+				UE_LOG(LogRHI, Log, TEXT("[%s, %d, %d]"), *ProgramKey.ToString(), ProgramOUT, CachedProgramBinaryOUT.Num());
+				RHIGetPanicDelegate().ExecuteIfBound(FName("FailedBinaryProgramCreateFromOldCache"));
+				UE_LOG(LogRHI, Fatal, TEXT("UseCachedProgram : Failed to create GL program from binary data while BuildingCacheFileWithMove! [%s]"), *ProgramKey.ToString());
+			}
+			// Now write to new cache, we're returning true here so no attempt will be made to add it back to the cache later.
+			AddProgramBinaryDataToBinaryCache(CachedProgramBinaryOUT, ProgramKey);
+
+			PreviousBinaryCacheInfo.NumberOfOldEntriesReused++;
+			return true;
+		}
+	}
+	return false;
+}
+
+// void FOpenGLProgramBinaryCache::CompilePendingShaders(const FOpenGLLinkedProgramConfiguration& Config)
+// {
+// 	VERIFY_GL_SCOPE();
+// 
+// 	// Find the existing compiled shader in the cache.
+// 	FScopeLock Lock(&GCompiledShaderCacheCS);
+// 	for (int32 StageIdx = 0; StageIdx < UE_ARRAY_COUNT(Config.Shaders); ++StageIdx)
+// 	{
+// 		TSharedPtr<FOpenGLCompiledShaderValue> FoundShader = GetOpenGLCompiledShaderCache().FindRef(Config.Shaders[StageIdx].ShaderKey);
+// 		if (FoundShader && FoundShader->bHasCompiled == false)
+// 		{
+// 			TArray<ANSICHAR> GlslCode = FoundShader->GetUncompressedShader();
+// 			CompileCurrentShader(Config.Shaders[StageIdx].Resource, GlslCode);
+// 			FoundShader->bHasCompiled = true;
+// 		}
+// 	}
+// }
+
+FString FOpenGLProgramBinaryCache::GetProgramBinaryCacheFilePath() const
+{
+	FString ProgramFilename = CachePath + TEXT("/") + CacheFilename;
+	return ProgramFilename;
+}
+
+void FOpenGLProgramBinaryCache::CheckPendingGLProgramCreateRequests()
+{
+	if (CachePtr)
+	{
+		QUICK_SCOPE_CYCLE_COUNTER(STAT_OpenGLShaderCreateShaderLibRequests);
+		CachePtr->CheckPendingGLProgramCreateRequests_internal();
+	}
+}
+
+void FOpenGLProgramBinaryCache::CheckPendingGLProgramCreateRequests_internal()
+{
+	check(IsInRenderingThread() || IsInRHIThread());
+	FScopeLock Lock(&GPendingGLProgramCreateRequestsCS);
+	//UE_LOG(LogRHI, Log, TEXT("CheckPendingGLProgramCreateRequests : PendingGLProgramCreateRequests = %d"), PendingGLProgramCreateRequests.Num());
+
+	if (PendingGLProgramCreateRequests.Num() > 0)
+	{
+		float TimeRemainingS = (float)GMaxShaderLibProcessingTimeMS / 1000.0f;
+		double StartTime = FPlatformTime::Seconds();
+		int32 Count = 0;
+		while (PendingGLProgramCreateRequests.Num() && TimeRemainingS > 0.0f)
+		{
+			CompleteLoadedGLProgramRequest_internal(PendingGLProgramCreateRequests.Pop());
+			TimeRemainingS -= (float)(FPlatformTime::Seconds() - StartTime);
+			StartTime = FPlatformTime::Seconds();
+			Count++;
+		}
+		float TimeTaken = (float)GMaxShaderLibProcessingTimeMS - (TimeRemainingS * 1000.0f);
+		UE_CLOG(TimeTaken > 2.0f, LogRHI, Log, TEXT("CheckPendingGLProgramCreateRequests : iter count = %d, time taken = %f ms (remaining %d)"), Count, TimeTaken, PendingGLProgramCreateRequests.Num());
+	}
+}
+
+
+
+void FOpenGLProgramBinaryCache::CompleteLoadedGLProgramRequest_internal(FGLProgramBinaryFileCacheEntry* PendingGLCreate)
+{
+	check(PendingGLCreate->GLProgramState == FGLProgramBinaryFileCacheEntry::EGLProgramState::ProgramLoaded);
+	FOpenGLProgramKey& ProgramKey = PendingGLCreate->FileInfo.ShaderHasheSet;
+	UE::OpenGL::OnGLProgramLoadedFromBinaryCache(ProgramKey, MoveTemp(PendingGLCreate->ProgramBinaryData));
+	// Ownership transfered to OpenGLProgramsCache.
+	PendingGLCreate->GLProgramState = FGLProgramBinaryFileCacheEntry::EGLProgramState::ProgramComplete;
+}
+
+bool FOpenGLProgramBinaryCache::CheckSinglePendingGLProgramCreateRequest(const FOpenGLProgramKey& ProgramKey)
+{
+	if (CachePtr)
+	{
+		return CachePtr->CheckSinglePendingGLProgramCreateRequest_internal(ProgramKey);
+	}
+	return false;
+}
+
+// Any pending program must complete in this case.
+bool FOpenGLProgramBinaryCache::CheckSinglePendingGLProgramCreateRequest_internal(const FOpenGLProgramKey& ProgramKey)
+{
+	FGLProgramBinaryFileCacheEntry** ProgramBinRefPtr = nullptr;
+	FScopeLock ProgramBinaryCacheLock(&GProgramBinaryFileCacheCS);
+	ProgramBinRefPtr = CachePtr->ProgramToBinaryMap.Find(ProgramKey);
+	if (ProgramBinRefPtr)
+	{
+		FGLProgramBinaryFileCacheEntry* ProgramEntry = *ProgramBinRefPtr;
+		TSharedPtr<IAsyncReadRequest, ESPMode::ThreadSafe> LocalReadRequest = ProgramEntry->ReadRequest.Pin();
+		if (LocalReadRequest.IsValid())
+		{
+			ensure(ProgramEntry->GLProgramState == FGLProgramBinaryFileCacheEntry::EGLProgramState::ProgramLoading);
+			LocalReadRequest->WaitCompletion();
+			ProgramEntry->ReadRequest = nullptr;
+			ProgramEntry->GLProgramState = FGLProgramBinaryFileCacheEntry::EGLProgramState::ProgramLoaded;
+			CompleteLoadedGLProgramRequest_internal(ProgramEntry);
+		}
+		else
+		{
+			FScopeLock Lock(&GPendingGLProgramCreateRequestsCS);
+			if (ProgramEntry->GLProgramState == FGLProgramBinaryFileCacheEntry::EGLProgramState::ProgramLoaded)
+			{
+				int32 PendingRequestIndex = -1;
+				if (ensure(PendingGLProgramCreateRequests.Find(ProgramEntry, PendingRequestIndex)))
+				{
+					CompleteLoadedGLProgramRequest_internal(ProgramEntry);
+					PendingGLProgramCreateRequests.RemoveAtSwap(PendingRequestIndex);
+				}
+			}
+		}
+		return true;
+	}
+	return false;
+}
+
+bool OnExternalReadCallback(const TSharedPtr<IAsyncReadRequest, ESPMode::ThreadSafe>& AsyncReadRequest, FGLProgramBinaryFileCacheEntry* ProgramBinEntry, TArray<FGLProgramBinaryFileCacheEntry*>& PendingGLProgramCreateRequests, double RemainingTime)
+{
+	if (!AsyncReadRequest->WaitCompletion(RemainingTime))
+	{
+		return false;
+	}
+
+	FScopeLock ProgramBinaryCacheLock(&GProgramBinaryFileCacheCS);
+
+	if (ProgramBinEntry->GLProgramState == FGLProgramBinaryFileCacheEntry::EGLProgramState::ProgramLoading)
+	{
+		// Async load complete.
+		ProgramBinEntry->GLProgramState = FGLProgramBinaryFileCacheEntry::EGLProgramState::ProgramLoaded;
+		FOpenGLProgramKey& ProgramKey = ProgramBinEntry->FileInfo.ShaderHasheSet;
+
+		{
+			// Add this program to the create queue.
+			FScopeLock Lock(&GPendingGLProgramCreateRequestsCS);
+			PendingGLProgramCreateRequests.Add(ProgramBinEntry);
+		}
+	}
+
+	return true;
+}
+
+void FOpenGLProgramBinaryCache::BeginProgramReadRequest(FGLProgramBinaryFileCacheEntry* ProgramBinEntry, FArchive* Ar)
+{
+	check(ProgramBinEntry);
+
+	TSharedPtr<IAsyncReadRequest, ESPMode::ThreadSafe> LocalReadRequest = ProgramBinEntry->ReadRequest.Pin();
+	bool bHasReadRequest = LocalReadRequest.IsValid();
+
+	if (ensure(!bHasReadRequest))
+	{
+		check(ProgramBinEntry->ProgramBinaryData.Num() == 0);
+		check(ProgramBinEntry->GLProgramState == FGLProgramBinaryFileCacheEntry::EGLProgramState::ProgramStored);
+
+		int64 ReadSize = ProgramBinEntry->FileInfo.ProgramSize;
+		int64 ReadOffset = ProgramBinEntry->FileInfo.ProgramOffset;
+
+		if (ensure(ReadSize > 0))
+		{
+			ProgramBinEntry->ProgramBinaryData.SetNumUninitialized(ReadSize);
+			ProgramBinEntry->GLProgramState = FGLProgramBinaryFileCacheEntry::EGLProgramState::ProgramLoading;
+			LocalReadRequest = MakeShareable(BinaryCacheAsyncReadFileHandle->ReadRequest(ReadOffset, ReadSize, AIOP_Normal, nullptr, ProgramBinEntry->ProgramBinaryData.GetData()));
+			ProgramBinEntry->ReadRequest = LocalReadRequest;
+			bHasReadRequest = true;
+
+			FExternalReadCallback ExternalReadCallback = [ProgramBinEntry, LocalReadRequest, this](double ReaminingTime)
+			{
+				return OnExternalReadCallback(LocalReadRequest, ProgramBinEntry, PendingGLProgramCreateRequests, ReaminingTime);
+			};
+
+			if (!Ar || !Ar->AttachExternalReadDependency(ExternalReadCallback))
+			{
+				// Archive does not support async loading
+				// do a blocking load
+				ExternalReadCallback(0.0);
+			}
+		}
+	}
+}
+
+void FOpenGLProgramBinaryCache::OnShaderLibraryRequestShaderCode(const FSHAHash& Hash, FArchive* Ar)
+{
+	if (CachePtr)
+	{
+		CachePtr->OnShaderLibraryRequestShaderCode_internal(Hash, Ar);
+	}
+}
+
+void FOpenGLProgramBinaryCache::OnShaderLibraryRequestShaderCode_internal(const FSHAHash& Hash, FArchive* Ar)
+{
+	FScopeLock Lock(&GProgramBinaryFileCacheCS);
+	FGLShaderToPrograms& FoundShaderToBinary = ShaderToProgramsMap.FindOrAdd(Hash);
+	if (!FoundShaderToBinary.bLoaded)
+	{
+		FoundShaderToBinary.bLoaded = true;
+
+		// if the binary cache is valid, look to see if we now have any complete programs to stream in.
+		// otherwise, we'll do this check bLoaded shaders when the binary cache loads.
+		if (BinaryFileState == EBinaryFileState::ValidCacheFile)
+		{
+			for (struct FGLProgramBinaryFileCacheEntry* ProgramBinEntry : FoundShaderToBinary.AssociatedPrograms)
+			{
+				const FOpenGLProgramKey& ProgramKey = ProgramBinEntry->FileInfo.ShaderHasheSet;
+				if (ProgramBinEntry->GLProgramState == FGLProgramBinaryFileCacheEntry::EGLProgramState::ProgramStored)
+				{
+					bool bAllShadersLoaded = true;
+					for (int32 i = 0; i < CrossCompiler::NUM_NON_COMPUTE_SHADER_STAGES && bAllShadersLoaded; i++)
+					{
+						bAllShadersLoaded = ProgramKey.ShaderHashes[i] == FSHAHash() || ShaderIsLoaded(ProgramKey.ShaderHashes[i]);
+					}
+
+					if (bAllShadersLoaded)
+					{
+						FOpenGLProgramBinaryCache::BeginProgramReadRequest(ProgramBinEntry, Ar);
+					}
+				}
+			}
+		}
+	}
+}
diff --git a/Engine/Source/Runtime/OpenGLDrv/Private/OpenGLProgramBinaryFileCache.h b/Engine/Source/Runtime/OpenGLDrv/Private/OpenGLProgramBinaryFileCache.h
new file mode 100644
--- /dev/null
+++ b/Engine/Source/Runtime/OpenGLDrv/Private/OpenGLProgramBinaryFileCache.h
@@ -0,0 +1,204 @@
+// Copyright Epic Games, Inc. All Rights Reserved.
+
+/*=============================================================================
+	OpenGLProgramBinaryFileCache.h: OpenGL program binary file cache stores/loads a set of binary ogl programs.
+=============================================================================*/
+
+#pragma once
+
+#include "HAL/UnrealMemory.h"
+#include "Templates/UnrealTemplate.h"
+#include "Containers/Array.h"
+#include "Misc/Crc.h"
+#include "Containers/UnrealString.h"
+#include "Containers/Map.h"
+#include "Misc/SecureHash.h"
+#include "HAL/IConsoleManager.h"
+#include "RHI.h"
+#include "ShaderCore.h"
+#include "CrossCompilerCommon.h"
+#include "ShaderCodeLibrary.h"
+#include "Async/AsyncFileHandle.h"
+#include "ShaderPipelineCache.h"
+
+class FOpenGLLinkedProgram;
+
+
+struct FGLProgramBinaryFileCacheEntry;
+class FOpenGLProgramBinaryCache
+{
+public:
+	static void Initialize();
+	static void Shutdown();
+
+	static bool IsEnabled();
+
+	/** Compile required shaders for a program, only in case binary program was not found in the cache   */
+// 	static void CompilePendingShaders(const FOpenGLLinkedProgramConfiguration& Config);
+
+	/** Try to find and load program binary from cache */
+	static bool UseCachedProgram(GLuint& ProgramOUT, const FOpenGLProgramKey& ProgramKey, TArray<uint8>& CachedProgramBinaryOUT);
+
+	/** Extract the program binary from GL and store on disk. Only when ProgramBinaryCache is enabled */
+	static void CacheProgram(GLuint Program, const FOpenGLProgramKey& ProgramKey, TArray<uint8>& CachedProgramBinaryOUT);
+
+	/** Take an existing program binary and store on disk. Only when ProgramBinaryCache is enabled
+	* CachedProgramBinary should be const, cant be due to serialize api
+	*/
+	static void CacheProgramBinary(const FOpenGLProgramKey& ProgramKey, /*const*/ TArray<uint8>& CachedProgramBinary);
+
+	static void OnShaderLibraryRequestShaderCode(const FSHAHash& Hash, FArchive* Ar);
+
+	/** Create any pending GL programs that have come from shader library requests */
+	static void CheckPendingGLProgramCreateRequests();
+
+	/** Create any single GL program that have come from shader library requests */
+	static bool CheckSinglePendingGLProgramCreateRequest(const FOpenGLProgramKey& ProgramKey);
+
+	/** true if the program binary cache is currently in cache build mode */
+	static bool IsBuildingCache();
+
+private:
+	FOpenGLProgramBinaryCache(const FString& InCachePath);
+	~FOpenGLProgramBinaryCache();
+
+	FString GetProgramBinaryCacheFilePath() const;
+
+	struct FPendingShaderCode
+	{
+		TArray<ANSICHAR> GlslCode;
+		int32 UncompressedSize;
+		bool bCompressed;
+	};
+
+private:
+	static TAutoConsoleVariable<int32> CVarPBCEnable;
+	static TAutoConsoleVariable<int32> CVarRestartAndroidAfterPrecompile;
+	static FOpenGLProgramBinaryCache* CachePtr;
+
+	/*  Path to directory where binary programs will be stored excluding the cache filename */
+	FString CachePath;
+	/* Just the cache filename, without the path */
+	FString CacheFilename;
+
+
+	/**
+	* Shaders that were requested for compilation
+	* They will be compiled just before linking a program only in case when there is no saved binary program
+	*/
+	TMap<GLuint, FPendingShaderCode> ShadersPendingCompilation;
+
+	bool AppendProgramBinaryFile(FArchive& Ar, const FOpenGLProgramKey& ProgramKey, GLuint Program, uint32& ProgramBinaryOffsetOUT, uint32& ProgramBinarySizeOUT);
+	void AppendProgramBinaryFileEofEntry(FArchive& Ar);
+
+	void ScanProgramCacheFile(const FGuid& ShaderPipelineCacheVersionGuid = FGuid());
+
+	/* Add a program */
+	void AddProgramFileEntryToMap(FGLProgramBinaryFileCacheEntry* IndexEntry);
+
+	void OpenAsyncReadHandle();
+	void CloseAsyncReadHandle();
+
+	bool OpenWriteHandle();
+	void CloseWriteHandle();
+
+	void AppendGLProgramToBinaryCache(const FOpenGLProgramKey& ProgramKey, GLuint Program, TArray<uint8>& CachedProgramBinaryOUT);
+	void AddUniqueGLProgramToBinaryCache(const FOpenGLProgramKey& ProgramKey, GLuint Program, TArray<uint8>& CachedProgramBinaryOUT);
+
+	void AddProgramBinaryDataToBinaryCache(TArray<uint8>& BinaryProgramData, const FOpenGLProgramKey& ProgramKey);
+
+	void ReleaseGLProgram_internal(FOpenGLLinkedProgramConfiguration& Config, GLuint Program);
+
+	FORCEINLINE_DEBUGGABLE bool ShaderIsLoaded(const FSHAHash& Hash)
+	{
+		const FGLShaderToPrograms* FoundShaderToBinary = CachePtr->ShaderToProgramsMap.Find(Hash);
+		return FoundShaderToBinary && FoundShaderToBinary->bLoaded;
+	}
+
+	bool UseCachedProgram_internal(GLuint& ProgramOUT, const FOpenGLProgramKey& ProgramKey, TArray<uint8>& CachedProgramBinaryOUT);
+
+	void OnShaderLibraryRequestShaderCode_internal(const FSHAHash& Hash, FArchive* Ar);
+
+	void BeginProgramReadRequest(FGLProgramBinaryFileCacheEntry* IndexEntry, FArchive* Ar);
+
+	void CheckPendingGLProgramCreateRequests_internal();
+	bool CheckSinglePendingGLProgramCreateRequest_internal(const FOpenGLProgramKey& ProgramKey);
+
+	void CompleteLoadedGLProgramRequest_internal(FGLProgramBinaryFileCacheEntry* PendingGLCreate);
+
+	/** Delegate handlers to track the ShaderPipelineCache precompile. */
+	void OnShaderPipelineCacheOpened(FString const& Name, EShaderPlatform Platform, uint32 Count, const FGuid& VersionGuid, FShaderPipelineCache::FShaderCachePrecompileContext& ShaderCachePrecompileContext);
+	void OnShaderPipelineCachePrecompilationComplete(uint32 Count, double Seconds, const FShaderPipelineCache::FShaderCachePrecompileContext& ShaderCachePrecompileContext);
+
+private:
+
+	FDelegateHandle OnShaderPipelineCacheOpenedDelegate;
+	FDelegateHandle OnShaderPipelineCachePrecompilationCompleteDelegate;
+
+	TArray<TUniquePtr<FGLProgramBinaryFileCacheEntry>> ProgramEntryContainer; // this is the owner of all FGLProgramBinaryFileCacheEntry ptrs
+
+	TMap<FOpenGLProgramKey, FGLProgramBinaryFileCacheEntry*> ProgramToBinaryMap; // program key to program entry.
+
+	struct FGLShaderToPrograms
+	{
+		FGLShaderToPrograms() : bLoaded(false)
+		{
+		}
+
+		FGLShaderToPrograms(FGLProgramBinaryFileCacheEntry* ProgramEntry) : bLoaded(false)
+		{
+			AssociatedPrograms.Add(ProgramEntry);
+		}
+
+		void Add(FGLProgramBinaryFileCacheEntry* ProgramEntry)
+		{
+			checkSlow(!AssociatedPrograms.Contains(ProgramEntry));
+			AssociatedPrograms.Add(ProgramEntry);
+		}
+
+		bool bLoaded;
+		TArray<FGLProgramBinaryFileCacheEntry*> AssociatedPrograms;
+	};
+
+	// Map of shader hash to a list of programs which reference it.
+	TMap<FSHAHash, FGLShaderToPrograms> ShaderToProgramsMap;
+
+	// programs loaded via async and now ready for creation on GL context owning thread.
+	TArray<FGLProgramBinaryFileCacheEntry*> PendingGLProgramCreateRequests;
+
+	IAsyncReadFileHandle* BinaryCacheAsyncReadFileHandle;
+	FArchive* BinaryCacheWriteFileHandle;
+	bool bShownLoadingScreen;
+
+	enum class EBinaryFileState : uint8
+	{
+		Uninitialized,					// No binary file is yet established and we should not read or write to it.
+		BuildingCacheFile,				// We are precompiling shaders from the PSO and storing them in a new binary cache. Do not attempt to read.
+		BuildingCacheFileWithMove,		// We are precompiling shaders from the PSO and storing them in a new binary cache, shaders matching from the existing cache are moved to the new file. Do not attempt to read.
+		ValidCacheFile,					// We have a valid cache file we can use for reading. Do not attempt to write.
+	};
+
+	bool IsBuildingCache_internal() const
+	{
+		return BinaryFileState == EBinaryFileState::BuildingCacheFile || BinaryFileState == EBinaryFileState::BuildingCacheFileWithMove;
+	}
+
+	struct FPreviousGLProgramBinaryCacheInfo
+	{
+		FPreviousGLProgramBinaryCacheInfo();
+		FPreviousGLProgramBinaryCacheInfo(FPreviousGLProgramBinaryCacheInfo&&);
+		FPreviousGLProgramBinaryCacheInfo& operator=(FPreviousGLProgramBinaryCacheInfo&&);
+		~FPreviousGLProgramBinaryCacheInfo();
+
+
+		FString OldCacheFilename;
+		TUniquePtr<FArchive> OldCacheArchive;
+		TMap<FOpenGLProgramKey, TUniquePtr<FGLProgramBinaryFileCacheEntry> > ProgramToOldBinaryCacheMap; // program key to program entry for old cache used when generating new cache.
+
+		// for logging
+		uint32 NumberOfOldEntriesReused;
+	};
+	FPreviousGLProgramBinaryCacheInfo PreviousBinaryCacheInfo;
+
+	EBinaryFileState BinaryFileState;
+};
diff --git a/Engine/Source/Runtime/OpenGLDrv/Private/OpenGLShaders.cpp b/Engine/Source/Runtime/OpenGLDrv/Private/OpenGLShaders.cpp
--- a/Engine/Source/Runtime/OpenGLDrv/Private/OpenGLShaders.cpp
+++ b/Engine/Source/Runtime/OpenGLDrv/Private/OpenGLShaders.cpp
@@ -16,6 +16,7 @@
 #include "SceneUtils.h"
 #include "PsoLruCache.h"
 #include "RHICoreShader.h"
+#include "OpenGLProgramBinaryFileCache.h"
 
 #define CHECK_FOR_GL_SHADERS_TO_REPLACE 0
 
@@ -51,13 +52,6 @@ static TAutoConsoleVariable<int32> CVarLRUMaxProgramBinarySize(
 	TEXT("Non-Zero: Maximum number of bytes active shader programs may use. If reached, least recently used shader programs will deleted."),
 	ECVF_RenderThreadSafe);
 
-TAutoConsoleVariable<int32> CVarStoreCompressedBinaries(
-	TEXT("r.OpenGL.StoreCompressedProgramBinaries"),
-	0,
-	TEXT(""),
-	ECVF_ReadOnly | ECVF_RenderThreadSafe
-);
-
 static TAutoConsoleVariable<int32> CVarLRUKeepProgramBinaryResident(
 	TEXT("r.OpenGL.ProgramLRUKeepBinaryResident"),
 	0,
@@ -83,23 +77,6 @@ static TAutoConsoleVariable<int32> CVarIgnoreShaderCompileFailure(
 	TEXT("1: Ignore Shader compile failures."),
 	ECVF_RenderThreadSafe);
 
-static TAutoConsoleVariable<int32> CVarUseExistingBinaryFileCache(
-	TEXT("r.OpenGL.UseExistingBinaryFileCache"),
-	1,
-	TEXT("When generating a new binary cache (such as when Shader Pipeline Cache Version Guid changes) use the existing binary file cache to speed up generation of the new cache.\n")
-	TEXT("0: Always rebuild binary file cache when Pipeline Cache Version Guid changes.\n")
-	TEXT("1: When Pipeline Cache Version Guid changes re-use programs from the existing binary cache where possible (default)."),
-	ECVF_RenderThreadSafe);
-
-static int32 GMaxShaderLibProcessingTimeMS = 10;
-static FAutoConsoleVariableRef CVarMaxShaderLibProcessingTime(
-	TEXT("r.OpenGL.MaxShaderLibProcessingTime"),
-	GMaxShaderLibProcessingTimeMS,
-	TEXT("The maximum time per frame to process shader library requests in milliseconds.\n")
-	TEXT("default 10ms. Note: Driver compile time for a single program may exceed this limit."),
-	ECVF_RenderThreadSafe
-);
-
 #if PLATFORM_ANDROID
 bool GOpenGLShaderHackLastCompileSuccess = false;
 #endif
@@ -546,13 +523,6 @@ void PE_ReplaceCString(TArray<ANSICHAR> & Dest, const ANSICHAR * Source, const A
 	ReplaceCString(Dest, Source, Replacement);
 }
 
-
-FOpenGLProgramBinaryCache::FPreviousGLProgramBinaryCacheInfo::FPreviousGLProgramBinaryCacheInfo() : NumberOfOldEntriesReused(0) {}
-FOpenGLProgramBinaryCache::FPreviousGLProgramBinaryCacheInfo::FPreviousGLProgramBinaryCacheInfo(FOpenGLProgramBinaryCache::FPreviousGLProgramBinaryCacheInfo&&) = default;
-FOpenGLProgramBinaryCache::FPreviousGLProgramBinaryCacheInfo& FOpenGLProgramBinaryCache::FPreviousGLProgramBinaryCacheInfo::operator = (FOpenGLProgramBinaryCache::FPreviousGLProgramBinaryCacheInfo&&) = default;
-FOpenGLProgramBinaryCache::FPreviousGLProgramBinaryCacheInfo::~FPreviousGLProgramBinaryCacheInfo() = default;
-
-
 inline uint32 GetTypeHash(FAnsiCharArray const& CharArray)
 {
 	return FCrc::MemCrc32(CharArray.GetData(), CharArray.Num() * sizeof(ANSICHAR));
@@ -1470,158 +1440,88 @@ static void ConfigureStageStates(FOpenGLLinkedProgram* LinkedProgram)
 	}
 }
 
-static bool CreateGLProgramFromUncompressedBinary(GLuint& ProgramOUT, const TArray<uint8>& ProgramBinary)
-{
-	VERIFY_GL_SCOPE();
-	GLuint GLProgramName = 0;
-	FOpenGL::GenProgramPipelines(1, &GLProgramName);
-	int32 BinarySize = ProgramBinary.Num();
-	//UE_LOG(LogRHI, Log, TEXT("CreateGLProgramFromBinary : gen program %x, size: %d"), GLProgramName, BinarySize);
-
-	check(BinarySize);
-
-	const uint8* ProgramBinaryPtr = ProgramBinary.GetData();
-
-	// BinaryFormat is stored at the start of ProgramBinary array
-	FOpenGL::ProgramBinary(GLProgramName, ((GLenum*)ProgramBinaryPtr)[0], ProgramBinaryPtr + sizeof(GLenum), BinarySize - sizeof(GLenum));
-	//	UE_LOG(LogRHI, Warning, TEXT("LRU: CreateFromBinary %d, binary format: %x, BinSize: %d"), GLProgramName, ((GLenum*)ProgramBinaryPtr)[0], BinarySize - sizeof(GLenum));
-
-	ProgramOUT = GLProgramName;
-	return VerifyLinkedProgram(GLProgramName);
-}
-
-struct FCompressedProgramBinaryHeader
-{
-	static const uint32 NotCompressed = 0xFFFFFFFF;
-	uint32 UncompressedSize;
-};
-
-static bool UncompressCompressedBinaryProgram(const TArray<uint8>& CompressedProgramBinary, TArray<uint8>& UncompressedProgramBinaryOUT )
-{
-	if (ensure(CompressedProgramBinary.Num() > sizeof(FCompressedProgramBinaryHeader)))
-	{
-		FCompressedProgramBinaryHeader* Header = (FCompressedProgramBinaryHeader*)CompressedProgramBinary.GetData();
-
-		if (Header->UncompressedSize == FCompressedProgramBinaryHeader::NotCompressed)
-		{
-			const uint32 ProgramSize = CompressedProgramBinary.Num() - sizeof(FCompressedProgramBinaryHeader);
-			UncompressedProgramBinaryOUT.SetNumUninitialized(ProgramSize);
-			FMemory::Memcpy(UncompressedProgramBinaryOUT.GetData(), CompressedProgramBinary.GetData() + sizeof(FCompressedProgramBinaryHeader), ProgramSize);
-			return true;
-		}
-		else
-		{
-			UncompressedProgramBinaryOUT.AddUninitialized(Header->UncompressedSize);
-
-			if (Header->UncompressedSize > 0
-				&& FCompression::UncompressMemory(NAME_Zlib, UncompressedProgramBinaryOUT.GetData(), UncompressedProgramBinaryOUT.Num(), CompressedProgramBinary.GetData() + sizeof(FCompressedProgramBinaryHeader), CompressedProgramBinary.Num() - sizeof(FCompressedProgramBinaryHeader)))
+namespace UE
+{
+	namespace OpenGL
+	{
+        static bool CreateGLProgramFromUncompressedBinary(GLuint& ProgramOUT, const TArray<uint8>& ProgramBinary)
+        {
+	        VERIFY_GL_SCOPE();
+	        GLuint GLProgramName = 0;
+	        FOpenGL::GenProgramPipelines(1, &GLProgramName);
+	        int32 BinarySize = ProgramBinary.Num();
+	        //UE_LOG(LogRHI, Log, TEXT("CreateGLProgramFromBinary : gen program %x, size: %d"), GLProgramName, BinarySize);
+        
+	        check(BinarySize);
+        
+	        const uint8* ProgramBinaryPtr = ProgramBinary.GetData();
+        
+	        // BinaryFormat is stored at the start of ProgramBinary array
+	        FOpenGL::ProgramBinary(GLProgramName, ((GLenum*)ProgramBinaryPtr)[0], ProgramBinaryPtr + sizeof(GLenum), BinarySize - sizeof(GLenum));
+	        //	UE_LOG(LogRHI, Warning, TEXT("LRU: CreateFromBinary %d, binary format: %x, BinSize: %d"), GLProgramName, ((GLenum*)ProgramBinaryPtr)[0], BinarySize - sizeof(GLenum));
+        
+	        ProgramOUT = GLProgramName;
+	        return VerifyLinkedProgram(GLProgramName);
+        }
+
+        static bool CreateGLProgramFromCompressedBinary(GLuint& ProgramOUT, const TArray<uint8>& CompressedProgramBinary)
+        {
+	        TArray<uint8> UncompressedProgramBinary;
+        
+	        bool bDecompressSuccess;
+        
+	        {
+		        QUICK_SCOPE_CYCLE_COUNTER(STAT_DecompressProgramBinary);
+		        bDecompressSuccess = UE::OpenGL::UncompressCompressedBinaryProgram(CompressedProgramBinary, UncompressedProgramBinary);
+	        }
+        
+	        if(bDecompressSuccess)
+	        {
+		        QUICK_SCOPE_CYCLE_COUNTER(STAT_CreateProgramFromBinary);
+		        return CreateGLProgramFromUncompressedBinary(ProgramOUT, UncompressedProgramBinary);
+	        }
+	        return false;
+        }
+
+		bool CreateGLProgramFromBinary(GLuint& ProgramOUT, const TArray<uint8>& ProgramBinary)
+		{
+			SCOPE_CYCLE_COUNTER(STAT_OpenGLCreateProgramFromBinaryTime)
+			bool bSuccess;
+			if (UE::OpenGL::IsStoringCompressedBinaryPrograms())
+			{
+				bSuccess = CreateGLProgramFromCompressedBinary(ProgramOUT, ProgramBinary);
+			}
+			else
 			{
-				return true;
+				bSuccess = CreateGLProgramFromUncompressedBinary(ProgramOUT, ProgramBinary);
 			}
-		}
-	}
-	return false;
-}
-
-static bool CreateGLProgramFromCompressedBinary(GLuint& ProgramOUT, const TArray<uint8>& CompressedProgramBinary)
-{
-	TArray<uint8> UncompressedProgramBinary;
-
-	bool bDecompressSuccess;
-
-	{
-		QUICK_SCOPE_CYCLE_COUNTER(STAT_DecompressProgramBinary);
-		bDecompressSuccess = UncompressCompressedBinaryProgram(CompressedProgramBinary, UncompressedProgramBinary);
-	}
-
-	if(bDecompressSuccess)
-	{
-		QUICK_SCOPE_CYCLE_COUNTER(STAT_CreateProgramFromBinary);
-		return CreateGLProgramFromUncompressedBinary(ProgramOUT, UncompressedProgramBinary);
-	}
-	return false;
-}
-
-static bool GetUncompressedProgramBinaryFromGLProgram(GLuint Program, TArray<uint8>& ProgramBinaryOUT)
-{
-	VERIFY_GL_SCOPE();
-
-	// pull binary from linked program
-	GLint BinaryLength = -1;
-	glGetProgramiv(Program, GL_PROGRAM_BINARY_LENGTH, &BinaryLength);
-	if(BinaryLength > 0)
-	{
-		ProgramBinaryOUT.SetNumUninitialized(BinaryLength + sizeof(GLenum));
-		uint8* ProgramBinaryPtr = ProgramBinaryOUT.GetData();
-		// BinaryFormat is stored at the start of ProgramBinary array
-		FOpenGL::GetProgramBinary(Program, BinaryLength, &BinaryLength, (GLenum*)ProgramBinaryPtr, ProgramBinaryPtr + sizeof(GLenum));
-		return true;
-	}
-	return false;
-}
 
-static void CompressProgramBinary(const TArray<uint8>& UncompressedProgramBinary, TArray<uint8>& ProgramBinaryOUT)
-{
-	check(CVarStoreCompressedBinaries.GetValueOnAnyThread());		
-	check(ProgramBinaryOUT.IsEmpty());
-
-	int32 CompressedSize = FCompression::CompressMemoryBound(NAME_Zlib, UncompressedProgramBinary.Num());
-	uint32 CompressedHeaderSize = sizeof(FCompressedProgramBinaryHeader);
-	ProgramBinaryOUT.AddUninitialized(CompressedSize + CompressedHeaderSize);
-	bool bSuccess = FCompression::CompressMemory(NAME_Zlib, ProgramBinaryOUT.GetData() + CompressedHeaderSize, CompressedSize, UncompressedProgramBinary.GetData(), UncompressedProgramBinary.Num());
-	if (bSuccess)
-	{
-		ProgramBinaryOUT.SetNum(CompressedSize + CompressedHeaderSize);
-		ProgramBinaryOUT.Shrink();
-		FCompressedProgramBinaryHeader* Header = (FCompressedProgramBinaryHeader*)ProgramBinaryOUT.GetData();
-		Header->UncompressedSize = UncompressedProgramBinary.Num();
-	}
-	else
-	{
-		// failed, store the uncompressed version.
-		UE_LOG(LogRHI, Log, TEXT("Storing binary program uncompressed (%d, %d, %d)"), UncompressedProgramBinary.Num(), ProgramBinaryOUT.Num(), CompressedSize);
-		ProgramBinaryOUT.SetNumUninitialized(UncompressedProgramBinary.Num() + CompressedHeaderSize);
-		FCompressedProgramBinaryHeader* Header = (FCompressedProgramBinaryHeader*)ProgramBinaryOUT.GetData();
-		Header->UncompressedSize = FCompressedProgramBinaryHeader::NotCompressed;
-		FMemory::Memcpy(ProgramBinaryOUT.GetData() + sizeof(FCompressedProgramBinaryHeader), UncompressedProgramBinary.GetData(), UncompressedProgramBinary.Num());
-	}
-}
+			if( bSuccess )
+			{
+				SetNewProgramStats(ProgramOUT);
+			}
 
-static bool GetCompressedProgramBinaryFromGLProgram(GLuint Program, TArray<uint8>& ProgramBinaryOUT)
-{
-	// get uncompressed binary
-	TArray<uint8> UncompressedProgramBinary;
-	if (GetUncompressedProgramBinaryFromGLProgram(Program, UncompressedProgramBinary))
-	{
-		CompressProgramBinary(UncompressedProgramBinary, ProgramBinaryOUT);
-		return true;
-	}
-	return false;
-}
+			return bSuccess;
+		}
 
-static bool GetProgramBinaryFromGLProgram(GLuint Program, TArray<uint8>& ProgramBinaryOUT)
-{
-	if (CVarStoreCompressedBinaries.GetValueOnAnyThread())
-	{
-		return GetCompressedProgramBinaryFromGLProgram(Program, ProgramBinaryOUT);
-	}
-	else
-	{
-		return GetUncompressedProgramBinaryFromGLProgram(Program, ProgramBinaryOUT);
-	}
-}
+        static void CompilePendingShaders(const FOpenGLLinkedProgramConfiguration& Config)
+        {
+	        VERIFY_GL_SCOPE();
+        
+	        // Find the existing compiled shader in the cache.
+	        FScopeLock Lock(&GCompiledShaderCacheCS);
+	        for (int32 StageIdx = 0; StageIdx < UE_ARRAY_COUNT(Config.Shaders); ++StageIdx)
+	        {
+		        TSharedPtr<FOpenGLCompiledShaderValue> FoundShader = GetOpenGLCompiledShaderCache().FindRef(Config.Shaders[StageIdx].ShaderKey);
+		        if (FoundShader && FoundShader->bHasCompiled == false)
+		        {
+			        TArray<ANSICHAR> GlslCode = FoundShader->GetUncompressedShader();
+			        CompileCurrentShader(Config.Shaders[StageIdx].Resource, GlslCode);
+			        FoundShader->bHasCompiled = true;
+		        }
+	        }
+        }
 
-static bool CreateGLProgramFromBinary(GLuint& ProgramOUT, const TArray<uint8>& ProgramBinary)
-{
-	SCOPE_CYCLE_COUNTER(STAT_OpenGLCreateProgramFromBinaryTime)
-	
-	if (CVarStoreCompressedBinaries.GetValueOnAnyThread())
-	{
-		return CreateGLProgramFromCompressedBinary(ProgramOUT, ProgramBinary);
-	}
-	else
-	{
-		return CreateGLProgramFromUncompressedBinary(ProgramOUT, ProgramBinary);
 	}
 }
 
@@ -1671,7 +1571,7 @@ class FGLProgramCacheLRU
 			if( bCreateProgramBinary )
 			{
 			// build offline binary:
-				GetProgramBinaryFromGLProgram(LinkedProgram->Program, GetProgramBinary());
+				UE::OpenGL::GetProgramBinaryFromGLProgram(LinkedProgram->Program, GetProgramBinary());
 				INC_MEMORY_STAT_BY(STAT_OpenGLShaderLRUProgramMemory, GetProgramBinary().Num());
 			}
 
@@ -1690,7 +1590,7 @@ class FGLProgramCacheLRU
 		void RestoreGLProgramFromBinary()
 		{
 			check(LinkedProgram->Program == 0);
-			bool bSuccess = CreateGLProgramFromBinary(LinkedProgram->Program, GetProgramBinary());
+			bool bSuccess = UE::OpenGL::CreateGLProgramFromBinary(LinkedProgram->Program, GetProgramBinary());
 			if(bSuccess)
 			{
 				if(CVarLRUKeepProgramBinaryResident.GetValueOnAnyThread() == 0)
@@ -2669,7 +2569,7 @@ static FOpenGLLinkedProgram* LinkProgram( const FOpenGLLinkedProgramConfiguratio
 		if (bShouldLinkProgram)
 		{
 			// In case there is no saved binary in the cache, compile required shaders we have deferred before
-			FOpenGLProgramBinaryCache::CompilePendingShaders(Config);
+			UE::OpenGL::CompilePendingShaders(Config);
 		}
 	}
 
@@ -2993,11 +2893,11 @@ void FOpenGLDynamicRHI::PrepareGFXBoundShaderState(const FGraphicsPipelineStateI
 		if(FailLog.IsEmpty())
 		{
 			GLenum glFormat = *(GLenum*)CompiledProgramResult.GetData();		
-			if (CVarStoreCompressedBinaries.GetValueOnAnyThread())
+			if (UE::OpenGL::IsStoringCompressedBinaryPrograms())
 			{
 				
 				TArray<uint8> CompressedCompiledProgramResult;
-				CompressProgramBinary(CompiledProgramResult, CompressedCompiledProgramResult);
+				UE::OpenGL::CompressProgramBinary(CompiledProgramResult, CompressedCompiledProgramResult);
 				CompiledProgramResult = MoveTemp(CompressedCompiledProgramResult);
 			}
 
@@ -3023,6 +2923,16 @@ void FOpenGLDynamicRHI::PrepareGFXBoundShaderState(const FGraphicsPipelineStateI
 		else
 		{
 			UE_LOG(LogRHI, Error, TEXT("External compile of program %s failed: %s "), *ProgramKey.ToString(), *FailLog);
+#if DEBUG_GL_SHADERS
+			if (VSCode.Num())
+			{
+				UE_LOG(LogRHI, Error, TEXT("Vertex Shader:\n%s"), ANSI_TO_TCHAR(VSCode.GetData()));
+			}
+			if (PSCode.Num())
+			{
+				UE_LOG(LogRHI, Error, TEXT("Pixel Shader:\n%s"), ANSI_TO_TCHAR(PSCode.GetData()));
+			}
+#endif //DEBUG_GL_SHADERS
 		}
 	}
 #endif
@@ -3797,1021 +3707,60 @@ void FOpenGLShaderParameterCache::CommitPackedUniformBuffers(FOpenGLLinkedProgra
 }
 
 
-static const uint32 GBinaryProgramFileVersion = 4;
-
-TAutoConsoleVariable<int32> FOpenGLProgramBinaryCache::CVarPBCEnable(
-	TEXT("r.ProgramBinaryCache.Enable"),
-#if PLATFORM_ANDROID
-	1,	// Enabled by default on Android.
-#else
-	0,
-#endif
-	TEXT("If true, enables binary program cache. Enabled by default only on Android"),
-	ECVF_ReadOnly | ECVF_RenderThreadSafe
-	);
-
-TAutoConsoleVariable<int32> FOpenGLProgramBinaryCache::CVarRestartAndroidAfterPrecompile(
-	TEXT("r.ProgramBinaryCache.RestartAndroidAfterPrecompile"),
-	1,	// Enabled by default on Android.
-	TEXT("If true, Android apps will restart after precompiling the binary program cache. Enabled by default only on Android"),
-	ECVF_ReadOnly | ECVF_RenderThreadSafe
-	);
-
-FOpenGLProgramBinaryCache* FOpenGLProgramBinaryCache::CachePtr = nullptr;
-
-FOpenGLProgramBinaryCache::FOpenGLProgramBinaryCache(const FString& InCachePath)
-	: CachePath(InCachePath)
-	, BinaryCacheAsyncReadFileHandle(nullptr)
-	, BinaryCacheWriteFileHandle(nullptr)
-	, BinaryFileState(EBinaryFileState::Uninitialized)
-{
-	ANSICHAR* GLVersion = (ANSICHAR*)glGetString(GL_VERSION);
-	ANSICHAR* GLRenderer = (ANSICHAR*)glGetString(GL_RENDERER);
-	FString HashString;
-	HashString.Append(GLVersion);
-	HashString.Append(GLRenderer);
-	FSHAHash VersionHash;
-	FSHA1::HashBuffer(TCHAR_TO_ANSI(*HashString), HashString.Len(), VersionHash.Hash);
-
-	CacheFilename = LegacyShaderPlatformToShaderFormat(GMaxRHIShaderPlatform).ToString() + TEXT("_") + VersionHash.ToString();
-}
-
-FOpenGLProgramBinaryCache::~FOpenGLProgramBinaryCache()
-{
-	if(BinaryCacheAsyncReadFileHandle)
-	{
-		delete BinaryCacheAsyncReadFileHandle;
-	}
-	if (BinaryCacheWriteFileHandle)
-	{
-		delete BinaryCacheWriteFileHandle;
-	}
-
-	if (OnShaderPipelineCacheOpenedDelegate.IsValid())
-	{
-		FShaderPipelineCache::GetCacheOpenedDelegate().Remove(OnShaderPipelineCacheOpenedDelegate);
-	}
-
-	if (OnShaderPipelineCachePrecompilationCompleteDelegate.IsValid())
-	{
-		FShaderPipelineCache::GetPrecompilationCompleteDelegate().Remove(OnShaderPipelineCachePrecompilationCompleteDelegate);
-	}
-};
-
-bool FOpenGLProgramBinaryCache::IsEnabled()
-{
-	return CachePtr != nullptr;
-}
-
-bool FOpenGLProgramBinaryCache::IsBuildingCache()
-{
-	if(CachePtr != nullptr)
-	{
-		return CachePtr->IsBuildingCache_internal();
-	}
-	return false;
-}
-
-void FOpenGLProgramBinaryCache::Initialize()
-{
-	check(CachePtr == nullptr);
-
-	if (CVarPBCEnable.GetValueOnAnyThread() == 0)
-	{
-		UE_LOG(LogRHI, Log, TEXT("FOpenGLProgramBinaryCache disabled by r.ProgramBinaryCache.Enable=0"));
-		return;
-	}
-
-	if (!FOpenGL::SupportsProgramBinary())
-	{
-		UE_LOG(LogRHI, Warning, TEXT("FOpenGLProgramBinaryCache disabled as devices does not support program binaries"));
-		return;
-	}
-
-#if PLATFORM_ANDROID
-	if (FOpenGL::HasBinaryProgramRetrievalFailed())
-	{
-		if (FOpenGL::SupportsProgramBinary())
-		{
-			UE_LOG(LogRHI, Warning, TEXT("FOpenGLProgramBinaryCache: Device has failed to emit program binary despite SupportsProgramBinary == true. Disabling binary cache."));
-			return;
-		}
-	}
-#endif
-
-
-		FString CacheFolderPath;
-#if PLATFORM_ANDROID && USE_ANDROID_FILE
-		// @todo Lumin: Use that GetPathForExternalWrite or something?
-		extern FString GExternalFilePath;
-		CacheFolderPath = GExternalFilePath / TEXT("ProgramBinaryCache");
-			
-#else
-		CacheFolderPath = FPaths::ProjectSavedDir() / TEXT("ProgramBinaryCache");
-#endif
-
-	// Remove entire ProgramBinaryCache folder if -ClearOpenGLBinaryProgramCache is specified on command line
-	if (FParse::Param(FCommandLine::Get(), TEXT("ClearOpenGLBinaryProgramCache")))
-	{
-		UE_LOG(LogRHI, Log, TEXT("ClearOpenGLBinaryProgramCache specified, deleting binary program cache folder: %s"), *CacheFolderPath);
-		FPlatformFileManager::Get().GetPlatformFile().DeleteDirectoryRecursively(*CacheFolderPath);
-	}
-
-	CachePtr = new FOpenGLProgramBinaryCache(CacheFolderPath);
-	UE_LOG(LogRHI, Log, TEXT("Enabling program binary cache as %s"), *CachePtr->GetProgramBinaryCacheFilePath());
-
-	// Add delegates for the ShaderPipelineCache precompile.
-	UE_LOG(LogRHI, Log, TEXT("FOpenGLProgramBinaryCache will be initialized when ShaderPipelineCache opens its file"));
-	CachePtr->OnShaderPipelineCacheOpenedDelegate				  = FShaderPipelineCache::GetCacheOpenedDelegate().AddRaw(CachePtr, &FOpenGLProgramBinaryCache::OnShaderPipelineCacheOpened);
-	CachePtr->OnShaderPipelineCachePrecompilationCompleteDelegate = FShaderPipelineCache::GetPrecompilationCompleteDelegate().AddRaw(CachePtr, &FOpenGLProgramBinaryCache::OnShaderPipelineCachePrecompilationComplete);
-}
-
-#if PLATFORM_ANDROID
-static int32 GNumRemoteProgramCompileServices = 4;
-static FAutoConsoleVariableRef CVarNumRemoteProgramCompileServices(
-	TEXT("Android.OpenGL.NumRemoteProgramCompileServices"),
-	GNumRemoteProgramCompileServices,
-	TEXT("The number of separate processes to make available to compile opengl programs.\n")
-	TEXT("0 to disable use of separate processes to precompile PSOs\n")
-	TEXT("valid range is 1-8 (4 default).")
-	,
-	ECVF_RenderThreadSafe|ECVF_ReadOnly
-);
-#endif
-
-void FOpenGLProgramBinaryCache::OnShaderPipelineCacheOpened(FString const& Name , EShaderPlatform Platform, uint32 Count, const FGuid& VersionGuid, FShaderPipelineCache::FShaderCachePrecompileContext& ShaderCachePrecompileContext)
-{
-	UE_LOG(LogRHI, Log, TEXT("Scanning Binary program cache, using Shader Pipeline Cache version %s"), *VersionGuid.ToString());
-	ScanProgramCacheFile(VersionGuid);
-	if(IsBuildingCache_internal())
-	{
-#if PLATFORM_ANDROID
-		if(GNumRemoteProgramCompileServices)
-		{
-			FAndroidOpenGL::StartAndWaitForRemoteCompileServices(GNumRemoteProgramCompileServices);
-		}
-#endif
-		ShaderCachePrecompileContext.SetPrecompilationIsSlowTask();
-	}
-}
-
-void FOpenGLProgramBinaryCache::OnShaderPipelineCachePrecompilationComplete(uint32 Count, double Seconds, const FShaderPipelineCache::FShaderCachePrecompileContext& ShaderCachePrecompileContext)
-{
-#if PLATFORM_ANDROID
-	FAndroidOpenGL::StopRemoteCompileServices();
-#endif
-	
-	UE_LOG(LogRHI, Log, TEXT("OnShaderPipelineCachePrecompilationComplete: %d shaders"), Count);
-	
-	// Want to ignore any subsequent Shader Pipeline Cache opening/closing, eg when loading modules
-	FShaderPipelineCache::GetCacheOpenedDelegate().Remove(OnShaderPipelineCacheOpenedDelegate);
-	FShaderPipelineCache::GetPrecompilationCompleteDelegate().Remove(OnShaderPipelineCachePrecompilationCompleteDelegate);
-	OnShaderPipelineCacheOpenedDelegate.Reset();
-	OnShaderPipelineCachePrecompilationCompleteDelegate.Reset();
-
-	check(IsBuildingCache_internal() || BinaryFileState == EBinaryFileState::ValidCacheFile);
-
-	if (IsBuildingCache_internal())
-	{
-		CloseWriteHandle();
-
-#if PLATFORM_ANDROID
-		FAndroidMisc::bNeedsRestartAfterPSOPrecompile = true;
-		if (CVarRestartAndroidAfterPrecompile.GetValueOnAnyThread() == 1)
-		{
-#if USE_ANDROID_JNI
-			extern void AndroidThunkCpp_RestartApplication(const FString& IntentString);
-			AndroidThunkCpp_RestartApplication(TEXT(""));
-#endif
-		}
-#endif
-		OpenAsyncReadHandle();
-		BinaryFileState = EBinaryFileState::ValidCacheFile;
-	}
-}
-
-// contains runtime and file information for a single program entry in the cache file.
-struct FGLProgramBinaryFileCacheEntry
-{
-	struct FGLProgramBinaryFileCacheFileInfo
-	{
-		// contains location info for a program in the binary cache file
-		FOpenGLProgramKey ShaderHasheSet;
-		uint32 ProgramOffset;
-		uint32 ProgramSize;
-
-		FGLProgramBinaryFileCacheFileInfo() : ProgramOffset(0), ProgramSize(0) { }
-
-		friend bool operator == (const FGLProgramBinaryFileCacheFileInfo& A, const FGLProgramBinaryFileCacheFileInfo& B)
-		{
-			return A.ShaderHasheSet == B.ShaderHasheSet && A.ProgramOffset == B.ProgramOffset&& A.ProgramSize == B.ProgramSize;
-		}
-	} FileInfo;
-
-	// program read request.
-	TWeakPtr<IAsyncReadRequest, ESPMode::ThreadSafe> ReadRequest;
-	// read data
-	TArray<uint8> ProgramBinaryData;
-	// debugging use only, index of program as encountered during scan, -1 if new.
-	int32 ProgramIndex;
-
-	enum class EGLProgramState : uint8
-	{
-		Unset,
-		ProgramStored, // program exists in binary cache but has not yet been loaded
-		ProgramLoading,	// program has started async loading.
-		ProgramLoaded,	// program has loaded and is ready for GL object creation.
-		ProgramAvailable, // program has loaded from binary cache and is available for use with GL.
-		ProgramComplete // program has been either added by rhi or handed over to rhi after being made available to it.
-	};
-	EGLProgramState GLProgramState;
-
-	// if != 0 then prepared runtime GL program name:
-	GLuint GLProgramId;
-
-	FGLProgramBinaryFileCacheEntry() : ProgramIndex(-1), GLProgramState(EGLProgramState::Unset), GLProgramId(0) { }
-
-	friend uint32 GetTypeHash(const FGLProgramBinaryFileCacheEntry& Key)
-	{
-		return FCrc::MemCrc32(&Key, sizeof(Key));
-	}
-};
-
-static FCriticalSection GPendingGLProgramCreateRequestsCS;
-
-// Scan the binary cache file and build a record of all programs.
-void FOpenGLProgramBinaryCache::ScanProgramCacheFile(const FGuid& ShaderPipelineCacheVersionGuid)
+namespace UE
 {
-	UE_LOG(LogRHI, Log, TEXT("OnShaderScanProgramCacheFile"));
-	FScopeLock Lock(&GProgramBinaryCacheCS);
-	FString ProgramCacheFilename = GetProgramBinaryCacheFilePath();
-	FString ProgramCacheFilenameTemp = GetProgramBinaryCacheFilePath() + TEXT(".scan");
-
-	IPlatformFile& PlatformFile = FPlatformFileManager::Get().GetPlatformFile();
-
-	check(BinaryFileState == EBinaryFileState::Uninitialized);
-
-	bool bBinaryFileIsValid = false;
-	bool bBinaryFileIsValidAndGuidMatch = false;
-
-	// Try to move the file to a temporary filename before the scan, so we won't try to read it again if it's corrupted
-	PlatformFile.DeleteFile(*ProgramCacheFilenameTemp);
-	PlatformFile.MoveFile(*ProgramCacheFilenameTemp, *ProgramCacheFilename);
-
-	TUniquePtr<FArchive> FileReader(IFileManager::Get().CreateFileReader(*ProgramCacheFilenameTemp));
-	if (FileReader)
+	namespace OpenGL
 	{
-		UE_LOG(LogRHI, Log, TEXT("OnShaderScanProgramCacheFile : Opened %s"), *ProgramCacheFilenameTemp);
-		FArchive& Ar = *FileReader;
-		uint32 Version = 0;
-		Ar << Version;
-		if(Version == GBinaryProgramFileVersion)
-		{
-			FGuid BinaryCacheGuid;
-			Ar << BinaryCacheGuid;
-			bool bCacheUsesCompressedBinaries;
-			Ar << bCacheUsesCompressedBinaries;
-
-			const bool bUseCompressedProgramBinaries = CVarStoreCompressedBinaries.GetValueOnAnyThread() != 0;
-			bBinaryFileIsValid = (bUseCompressedProgramBinaries == bCacheUsesCompressedBinaries);
-			bBinaryFileIsValidAndGuidMatch = bBinaryFileIsValid && (!ShaderPipelineCacheVersionGuid.IsValid() || ShaderPipelineCacheVersionGuid == BinaryCacheGuid);
-			
-			if (CVarUseExistingBinaryFileCache.GetValueOnAnyThread() == 0 && bBinaryFileIsValidAndGuidMatch == false)
-			{
-				// If we dont want to use the existing binary cache and the guids have changed then rebuild the binary file.
-				bBinaryFileIsValid = false;
-			}
-		}
-		
-		if (bBinaryFileIsValid)
+		// Called from the binary file cache when program loads from disk.
+		void OnGLProgramLoadedFromBinaryCache(const FOpenGLProgramKey& ProgramKey, TArray<uint8>&& ProgramBinaryData)
 		{
-			const uint32 ProgramBinaryStart = Ar.Tell();
-
-			// Search the file for the end record.
-			bool bFoundEndRecord = false;
-			int32 ProgramIndex = 0;
-			while (!Ar.AtEnd())
-			{
-				check(bFoundEndRecord == false); // There should be no additional data after the eof record.
-
-				FOpenGLProgramKey ProgramKey;
-				uint32 ProgramBinarySize = 0;
-				Ar << ProgramKey;
-				Ar << ProgramBinarySize;
-				uint32 ProgramBinaryOffset = Ar.Tell();
-				if (ProgramBinarySize == 0)
-				{
-					if (ProgramKey == FOpenGLProgramKey())
-					{
-						bFoundEndRecord = true;
-					}
-					else
-					{
-						// Note: This should not happen with new code. We can no longer write out records with 0 program size. see AppendProgramBinaryFile.
-						UE_LOG(LogRHI, Warning, TEXT("FOpenGLProgramBinaryCache::ScanProgramCacheFile : encountered 0 sized program during binary program cache scan"));
-					}
-				}
-				Ar.Seek(ProgramBinaryOffset + ProgramBinarySize);
-			}
+			const bool bProgramExists = GetOpenGLProgramsCache().Find(ProgramKey, false) != nullptr;
 
-			if(bFoundEndRecord)
+			if (GetOpenGLProgramsCache().IsUsingLRU())
 			{
-				Ar.Seek(ProgramBinaryStart);
-				while (!Ar.AtEnd())
-				{
-					FOpenGLProgramKey ProgramKey;
-					uint32 ProgramBinarySize = 0;
-					Ar << ProgramKey;
-					Ar << ProgramBinarySize;
-
-					if (ProgramBinarySize > 0)
-					{
-						FGLProgramBinaryFileCacheEntry* NewEntry = new FGLProgramBinaryFileCacheEntry();
-						NewEntry->FileInfo.ShaderHasheSet = ProgramKey;
-						NewEntry->ProgramIndex = ProgramIndex++;
-
-						uint32 ProgramBinaryOffset = Ar.Tell();
-						NewEntry->FileInfo.ProgramSize = ProgramBinarySize;
-						NewEntry->FileInfo.ProgramOffset = ProgramBinaryOffset;
-
-						if (bBinaryFileIsValidAndGuidMatch)
-						{
-							ProgramEntryContainer.Emplace(TUniquePtr<FGLProgramBinaryFileCacheEntry>(NewEntry));
-
-							// check to see if any of the shaders are already loaded and so we should serialize the binary
-							bool bAllShadersLoaded = true;
-							for (int32 i = 0; i < CrossCompiler::NUM_SHADER_STAGES && bAllShadersLoaded; i++)
-							{
-								bAllShadersLoaded = ProgramKey.ShaderHashes[i] == FSHAHash() || ShaderIsLoaded(ProgramKey.ShaderHashes[i]);
-							}
-							if (bAllShadersLoaded)
-							{
-								FPlatformMisc::LowLevelOutputDebugStringf(TEXT("*** All shaders for %s already loaded\n"), *ProgramKey.ToString());
-								NewEntry->ProgramBinaryData.AddUninitialized(ProgramBinarySize);
-								Ar.Serialize(NewEntry->ProgramBinaryData.GetData(), ProgramBinarySize);
-								NewEntry->GLProgramState = FGLProgramBinaryFileCacheEntry::EGLProgramState::ProgramLoaded;
-								CompleteLoadedGLProgramRequest_internal(NewEntry);
-							}
-							else
-							{
-								NewEntry->GLProgramState = FGLProgramBinaryFileCacheEntry::EGLProgramState::ProgramStored;
-							}
-							AddProgramFileEntryToMap(NewEntry);
-						}
-						else
-						{
-							check(!PreviousBinaryCacheInfo.ProgramToOldBinaryCacheMap.Contains(ProgramKey));
-							PreviousBinaryCacheInfo.ProgramToOldBinaryCacheMap.Emplace(ProgramKey, TUniquePtr<FGLProgramBinaryFileCacheEntry>(NewEntry));
-						}
-						Ar.Seek(ProgramBinaryOffset + ProgramBinarySize);
-					}
-				}
-
-				if (bBinaryFileIsValidAndGuidMatch)
+				if (!bProgramExists)
 				{
-					UE_LOG(LogRHI, Log, TEXT("Program Binary cache: Found %d cached programs, end record found: %d"), ProgramIndex, (uint32)bFoundEndRecord);
-					FileReader->Close();
-					// Rename the file back after a successful scan.
-					PlatformFile.MoveFile(*ProgramCacheFilename, *ProgramCacheFilenameTemp);
+					// Always add programs as evicted, 1st use will create them as programs.
+					// This will reduce pressure on driver by ensuring only used programs
+					// are created.
+					// In this case do not create the GL program.
+					GetOpenGLProgramsCache().AddAsEvicted(ProgramKey, MoveTemp(ProgramBinaryData));
 				}
 				else
 				{
-					UE_LOG(LogRHI, Log, TEXT("Program Binary cache: ShaderPipelineCache changed, regenerating for new pipeline cache. Existing cache contains %d programs, using it to populate."), PreviousBinaryCacheInfo.ProgramToOldBinaryCacheMap.Num());
-					// Not closing the scan source file, we're using it to move shaders from the old cache.
-					PreviousBinaryCacheInfo.OldCacheArchive = MoveTemp(FileReader);
-					PreviousBinaryCacheInfo.OldCacheFilename = ProgramCacheFilenameTemp;
+					// The program is already in use, discard the binary data.
+					ProgramBinaryData.Empty();
 				}
 			}
 			else
 			{
-				// failed to find sentinel record, the file was not finalized.
-				UE_LOG(LogRHI, Warning, TEXT("ScanProgramCacheFile - incomplete binary cache file encountered. Rebuilding binary program cache."));
-				FileReader->Close();
-				bBinaryFileIsValid = false;
-				bBinaryFileIsValidAndGuidMatch = false;
-			}
-		}
-		
-		if(!bBinaryFileIsValid)
-		{
-			UE_LOG(LogRHI, Log, TEXT("OnShaderScanProgramCacheFile : binary file version invalid"));
-		}
-
-		if (bBinaryFileIsValidAndGuidMatch)
-		{
-			OpenAsyncReadHandle();
-			BinaryFileState = EBinaryFileState::ValidCacheFile;
-		}
-	}
-	else
-	{
-		UE_LOG(LogRHI, Log, TEXT("OnShaderScanProgramCacheFile : Failed to open %s"), *ProgramCacheFilename);
-	}
-
-	if (!bBinaryFileIsValid)
-	{
-		// Attempt to remove any existing binary cache or temp files (eg for different driver version)
-		UE_LOG(LogRHI, Log, TEXT("Deleting binary program cache folder: %s"), *CachePath);
-		PlatformFile.DeleteDirectoryRecursively(*CachePath);
-
-		// Create
-		if (!PlatformFile.CreateDirectoryTree(*CachePath))
-		{
-			UE_LOG(LogRHI, Warning, TEXT("Failed to create directory for a program binary cache. Cache will be disabled: %s"), *CachePath);
-			return;
-		}
-	}
-
-	if(!bBinaryFileIsValid || !bBinaryFileIsValidAndGuidMatch)
-	{
-		if (OpenWriteHandle())
-		{
-			BinaryFileState = bBinaryFileIsValid && !bBinaryFileIsValidAndGuidMatch ? EBinaryFileState::BuildingCacheFileWithMove : EBinaryFileState::BuildingCacheFile;
-
-			// save header
-			FArchive& Ar = *BinaryCacheWriteFileHandle;
-			uint32 Version = GBinaryProgramFileVersion;
-			Ar << Version;
-			FGuid BinaryCacheGuid = ShaderPipelineCacheVersionGuid;
-			Ar << BinaryCacheGuid;
-			bool bWritingCompressedBinaries = (CVarStoreCompressedBinaries.GetValueOnAnyThread() != 0);
-			Ar << bWritingCompressedBinaries;
-		}
-		else
-		{
-			// Binary cache file cannot be used, failed to open output file.
-			BinaryFileState = EBinaryFileState::Uninitialized;
-			RHIGetPanicDelegate().ExecuteIfBound(FName("FailedBinaryProgramArchiveOpen"));
-			UE_LOG(LogRHI, Fatal, TEXT("ScanProgramCacheFile - Failed to open binary cache."));
-		}
-	}
-}
-
-// add the GLProgramBinaryFileCacheEntry into the runtime lookup containers.
-void FOpenGLProgramBinaryCache::AddProgramFileEntryToMap(FGLProgramBinaryFileCacheEntry* NewEntry)
-{
-	const FOpenGLProgramKey& ProgramKey = NewEntry->FileInfo.ShaderHasheSet;
-	check(!ProgramToBinaryMap.Contains(ProgramKey));
-	ProgramToBinaryMap.Add(ProgramKey, NewEntry);
-
-	UE_LOG(LogRHI, Verbose, TEXT("AddProgramFileEntryToMap : Adding program: %s"), *ProgramKey.ToString());
-
-	for (int i = 0; i < CrossCompiler::NUM_NON_COMPUTE_SHADER_STAGES; ++i)
-	{
-		const FSHAHash& ShaderHash = ProgramKey.ShaderHashes[i];
-
-		if (ShaderHash != FSHAHash())
-		{
-			if (ShaderToProgramsMap.Contains(ShaderHash))
-			{
-				ShaderToProgramsMap[ShaderHash].Add(NewEntry);
-			}
-			else
-			{
-				ShaderToProgramsMap.Add(ShaderHash, NewEntry);
-			}
-		}
-	}
-}
-
-bool FOpenGLProgramBinaryCache::OpenWriteHandle()
-{
-	check(BinaryCacheWriteFileHandle == nullptr);
-	check(BinaryCacheAsyncReadFileHandle == nullptr);
-
-	// perform file writing to a file temporary filename so we don't attempt to use the file later if the write session is interrupted
-	FString ProgramCacheFilename = GetProgramBinaryCacheFilePath();
-	FString ProgramCacheFilenameWrite = ProgramCacheFilename + TEXT(".write");
-
-	BinaryCacheWriteFileHandle = IFileManager::Get().CreateFileWriter(*ProgramCacheFilenameWrite, EFileWrite::FILEWRITE_None);
-
-	UE_CLOG(BinaryCacheWriteFileHandle == nullptr, LogRHI, Warning, TEXT("Failed to open OGL binary cache output file."));
-
-	return BinaryCacheWriteFileHandle != nullptr;
-}
-
-void FOpenGLProgramBinaryCache::CloseWriteHandle()
-{
-	if(BinaryFileState == EBinaryFileState::BuildingCacheFileWithMove)
-	{
-		UE_LOG(LogRHI, Log, TEXT("FOpenGLProgramBinaryCache: Deleting previous binary program cache (%s), reused %d programs from a total of %d."), *PreviousBinaryCacheInfo.OldCacheFilename, PreviousBinaryCacheInfo.NumberOfOldEntriesReused, ProgramToBinaryMap.Num());
-
-		// clean up references to old cache.
-		PreviousBinaryCacheInfo.OldCacheArchive->Close();
-		PreviousBinaryCacheInfo.OldCacheArchive = nullptr;
-		IPlatformFile& PlatformFile = FPlatformFileManager::Get().GetPlatformFile();
-		PlatformFile.DeleteFile(*PreviousBinaryCacheInfo.OldCacheFilename);
-		PreviousBinaryCacheInfo.OldCacheFilename.Empty();
-		PreviousBinaryCacheInfo.ProgramToOldBinaryCacheMap.Empty();
-	}
-
-	check(BinaryCacheWriteFileHandle != nullptr);
-
-	AppendProgramBinaryFileEofEntry(*BinaryCacheWriteFileHandle);
-	bool bArchiveFailed = BinaryCacheWriteFileHandle->IsError() || BinaryCacheWriteFileHandle->IsCriticalError();
-
-	BinaryCacheWriteFileHandle->Close();
-	delete BinaryCacheWriteFileHandle;
-	BinaryCacheWriteFileHandle = nullptr;
-
-	if (bArchiveFailed)
-	{
-		RHIGetPanicDelegate().ExecuteIfBound(FName("FailedBinaryProgramArchiveWrite"));
-		UE_LOG(LogRHI, Fatal, TEXT("CloseWriteHandle - FArchive error bit set, failed to write binary cache."));
-	}
-
-	// rename the temp filename back to the final filename
-	FString ProgramCacheFilename = GetProgramBinaryCacheFilePath();
-	FString ProgramCacheFilenameWrite = ProgramCacheFilename + TEXT(".write");
-	IPlatformFile& PlatformFile = FPlatformFileManager::Get().GetPlatformFile();
-	PlatformFile.DeleteFile(*ProgramCacheFilename); // file should never exist, but for safety
-	PlatformFile.MoveFile(*ProgramCacheFilename, *ProgramCacheFilenameWrite);
-}
-
-void FOpenGLProgramBinaryCache::OpenAsyncReadHandle()
-{
-	check(BinaryCacheAsyncReadFileHandle == nullptr);
-
-	FString ProgramCacheFilename = GetProgramBinaryCacheFilePath();
-	BinaryCacheAsyncReadFileHandle = FPlatformFileManager::Get().GetPlatformFile().OpenAsyncRead(*ProgramCacheFilename);
-	checkf(BinaryCacheAsyncReadFileHandle, TEXT("Could not opan an async file")); // this generally cannot fail because it is async
-}
-
-/* dead code, needs removal
-void FOpenGLProgramBinaryCache::CloseAsyncReadHandle()
-{
-	// wait for any pending reads.
-	{
-		FScopeLock Lock(&GPendingGLProgramCreateRequestsCS);
-		for (FGLProgramBinaryFileCacheEntry* CreateRequest : PendingGLProgramCreateRequests)
-		{
-			TSharedPtr<IAsyncReadRequest, ESPMode::ThreadSafe> ReadRequest = CreateRequest->ReadRequest.Pin();
-			if (ReadRequest.IsValid())
-			{
-				ReadRequest->WaitCompletion();
-				CreateRequest->ReadRequest = nullptr;
-			}
-		}
-	}
-
-	delete BinaryCacheAsyncReadFileHandle;
-	BinaryCacheAsyncReadFileHandle = nullptr;
-}*/
-
-// Called when a new program has been created by OGL RHI, creates the binary cache if it's invalid and then appends the new program details to the file and runtime containers.
-void FOpenGLProgramBinaryCache::AppendGLProgramToBinaryCache(const FOpenGLProgramKey& ProgramKey, GLuint Program, TArray<uint8>& CachedProgramBinaryOUT)
-{
-	if (IsBuildingCache_internal() == false)
-	{
-		return;
-	}
-
-	FScopeLock Lock(&GProgramBinaryCacheCS);
-
-	AddUniqueGLProgramToBinaryCache(ProgramKey, Program, CachedProgramBinaryOUT);
-}
-
-// Add the program to the binary cache if it does not already exist.
-void FOpenGLProgramBinaryCache::AddUniqueGLProgramToBinaryCache(const FOpenGLProgramKey& ProgramKey, GLuint Program, TArray<uint8>& CachedProgramBinaryOUT)
-{
-	// Add to runtime and disk.
-	const FOpenGLProgramKey& ProgramHash = ProgramKey;
-
-	// Check we dont already have this: Something could be in the cache but still reach this point if OnSharedShaderCodeRequest(s) have not occurred.
-	if (!ProgramToBinaryMap.Contains(ProgramHash))
-	{
-		uint32 ProgramBinaryOffset = 0, ProgramBinarySize = 0;
-
-		FOpenGLProgramKey SerializedProgramKey = ProgramKey;
-		if (ensure(GetProgramBinaryFromGLProgram(Program, CachedProgramBinaryOUT)))
-		{
-			AddProgramBinaryDataToBinaryCache(CachedProgramBinaryOUT, ProgramKey);
-		}
-		else
-		{
-			// we've encountered a problem with this program and there's nothing to write.
-			// This likely means the device will never be able to use this program.
-			// Panic!
-			RHIGetPanicDelegate().ExecuteIfBound(FName("FailedBinaryProgramWrite"));
-			UE_LOG(LogRHI, Fatal, TEXT("AppendProgramBinaryFile Binary program returned 0 bytes!"));
-			// Panic!
-		}
-	}
-}
-
-void FOpenGLProgramBinaryCache::CacheProgramBinary(const FOpenGLProgramKey& ProgramKey, TArray<uint8>& CachedProgramBinary)
-{
-	if (CachePtr)
-	{
-		if (CachePtr->IsBuildingCache_internal() == false)
-		{
-			return;
-		}
-
-		FScopeLock Lock(&GProgramBinaryCacheCS);
-
-		if (!CachePtr->ProgramToBinaryMap.Contains(ProgramKey))
-		{
-			CachePtr->AddProgramBinaryDataToBinaryCache(CachedProgramBinary, ProgramKey);
-		}
-	}
-}
-
-// Serialize out the program binary data and add to runtime structures.
-void FOpenGLProgramBinaryCache::AddProgramBinaryDataToBinaryCache(TArray<uint8>& BinaryProgramData, const FOpenGLProgramKey& ProgramKey)
-{
-	FScopeLock Lock(&GProgramBinaryCacheCS);
-	FArchive& Ar = *BinaryCacheWriteFileHandle;
-	// Serialize to output file:
-	FOpenGLProgramKey SerializedProgramKey = ProgramKey;
-	uint32 ProgramBinarySize = (uint32)BinaryProgramData.Num();
-	Ar << SerializedProgramKey;
-	uint32 ProgramBinaryOffset = Ar.Tell();
-	Ar << ProgramBinarySize;
-	Ar.Serialize(BinaryProgramData.GetData(), ProgramBinarySize);
-	if(CVarStoreCompressedBinaries.GetValueOnAnyThread())
-	{
-		static uint32 TotalUncompressed = 0;
-		static uint32 TotalCompressed = 0;
-
-		FCompressedProgramBinaryHeader* Header = (FCompressedProgramBinaryHeader*)BinaryProgramData.GetData();
-		TotalUncompressed += Header->UncompressedSize;
-		TotalCompressed += BinaryProgramData.Num();
-
-		UE_LOG(LogRHI, Verbose, TEXT("AppendProgramBinaryFile: total Uncompressed: %d, total Compressed %d, Total saved so far: %d"), TotalUncompressed, TotalCompressed, TotalUncompressed - TotalCompressed);
-	}
-
-	FGLProgramBinaryFileCacheEntry* NewIndexEntry = new FGLProgramBinaryFileCacheEntry();
-	ProgramEntryContainer.Emplace(TUniquePtr<FGLProgramBinaryFileCacheEntry>(NewIndexEntry));
-
-	// Store the program file descriptor in the runtime program/shader container:
-	NewIndexEntry->GLProgramState = FGLProgramBinaryFileCacheEntry::EGLProgramState::ProgramStored;
-	NewIndexEntry->FileInfo.ProgramOffset = ProgramBinaryOffset;
-	NewIndexEntry->FileInfo.ProgramSize = ProgramBinarySize;
-	NewIndexEntry->ProgramIndex = ProgramToBinaryMap.Num();
-	NewIndexEntry->FileInfo.ShaderHasheSet = ProgramKey;
-	AddProgramFileEntryToMap(NewIndexEntry);
-}
-
-void FOpenGLProgramBinaryCache::AppendProgramBinaryFileEofEntry(FArchive& Ar)
-{
-	// write out an all zero record that signifies eof.
-	FOpenGLProgramKey SerializedProgramKey;
-	Ar << SerializedProgramKey;
-	uint32 ProgramBinarySize = 0;
-	Ar << ProgramBinarySize;
-}
-
-void FOpenGLProgramBinaryCache::Shutdown()
-{
-	if (CachePtr)
-	{
-		delete CachePtr;
-		CachePtr = nullptr;
-	}
-}
-
-void FOpenGLProgramBinaryCache::CacheProgram(GLuint Program, const FOpenGLProgramKey& ProgramKey, TArray<uint8>& CachedProgramBinaryOUT)
-{
-	if (CachePtr)
-	{
-		CachePtr->AppendGLProgramToBinaryCache(ProgramKey, Program, CachedProgramBinaryOUT);
-	}
-}
-
-bool FOpenGLProgramBinaryCache::UseCachedProgram(GLuint& ProgramOUT, const FOpenGLProgramKey& ProgramKey, TArray<uint8>& CachedProgramBinaryOUT)
-{
-	if (CachePtr)
-	{
-		return CachePtr->UseCachedProgram_internal(ProgramOUT, ProgramKey, CachedProgramBinaryOUT);
-	}
-	return false;
-}
-
-bool FOpenGLProgramBinaryCache::UseCachedProgram_internal(GLuint& ProgramOUT, const FOpenGLProgramKey& ProgramKey, TArray<uint8>& CachedProgramBinaryOUT)
-{
-	SCOPE_CYCLE_COUNTER(STAT_OpenGLUseCachedProgramTime);
-	
-	FGLProgramBinaryFileCacheEntry** ProgramBinRefPtr = nullptr;
-
-	FScopeLock Lock(&GProgramBinaryCacheCS);
-
-	ProgramBinRefPtr = ProgramToBinaryMap.Find(ProgramKey);
-
-	if (ProgramBinRefPtr)
-	{
-		FGLProgramBinaryFileCacheEntry* FoundProgram = *ProgramBinRefPtr;
-		check(FoundProgram->FileInfo.ShaderHasheSet == ProgramKey);
-
-		TSharedPtr<IAsyncReadRequest, ESPMode::ThreadSafe> LocalReadRequest = FoundProgram->ReadRequest.Pin();
-		bool bHasReadRequest = LocalReadRequest.IsValid();
-		check(!bHasReadRequest);
-
-		// by this point the program must be either available or no attempt to load from shader library has occurred.
-		checkf(FoundProgram->GLProgramState == FGLProgramBinaryFileCacheEntry::EGLProgramState::ProgramStored
-			|| FoundProgram->GLProgramState == FGLProgramBinaryFileCacheEntry::EGLProgramState::ProgramAvailable,
-			TEXT("Unexpected program state:  (%s) == %d"), *ProgramKey.ToString(), (int32)FoundProgram->GLProgramState);
-
-		if (FoundProgram->GLProgramState == FGLProgramBinaryFileCacheEntry::EGLProgramState::ProgramAvailable)
-		{
-			UE_LOG(LogRHI, Log, TEXT("UseCachedProgram : Program (%s) GLid = %x is ready!"), *ProgramKey.ToString(), FoundProgram->GLProgramId);
-			ProgramOUT = FoundProgram->GLProgramId;
-
-			// GLProgram has been handed over.
-			FoundProgram->GLProgramId = 0;
-			FoundProgram->GLProgramState = FGLProgramBinaryFileCacheEntry::EGLProgramState::ProgramComplete;
-			return true;
-		}
-		else
-		{
-			UE_LOG(LogRHI, Log, TEXT("UseCachedProgram : %s was not ready when needed!! (state %d)"), *ProgramKey.ToString(), (uint32)FoundProgram->GLProgramState);
-		}
-	}
-	else if (BinaryFileState == EBinaryFileState::BuildingCacheFileWithMove)
-	{
-		// We're building the new cache using the original cache to warm:
-		TUniquePtr<FGLProgramBinaryFileCacheEntry>* FoundExistingBinary = PreviousBinaryCacheInfo.ProgramToOldBinaryCacheMap.Find(ProgramKey);
-		if (FoundExistingBinary)
-		{
-			TUniquePtr<FGLProgramBinaryFileCacheEntry>& ExistingBinary = *FoundExistingBinary;
-			// read old binary:
-			CachedProgramBinaryOUT.SetNumUninitialized(ExistingBinary->FileInfo.ProgramSize);
-			PreviousBinaryCacheInfo.OldCacheArchive->Seek(ExistingBinary->FileInfo.ProgramOffset);
-			PreviousBinaryCacheInfo.OldCacheArchive->Serialize(CachedProgramBinaryOUT.GetData(), ExistingBinary->FileInfo.ProgramSize);
-			bool bSuccess = CreateGLProgramFromBinary(ProgramOUT, CachedProgramBinaryOUT);
-			if (!bSuccess)
-			{
-				UE_LOG(LogRHI, Log, TEXT("[%s, %d, %d]"), *ProgramKey.ToString(), ProgramOUT, CachedProgramBinaryOUT.Num());
-				RHIGetPanicDelegate().ExecuteIfBound(FName("FailedBinaryProgramCreateFromOldCache"));
-				UE_LOG(LogRHI, Fatal, TEXT("UseCachedProgram : Failed to create GL program from binary data while BuildingCacheFileWithMove! [%s]"), *ProgramKey.ToString());
-			}
-			SetNewProgramStats(ProgramOUT);
-			// Now write to new cache, we're returning true here so no attempt will be made to add it back to the cache later.
-			AddProgramBinaryDataToBinaryCache(CachedProgramBinaryOUT, ProgramKey);
-
-			PreviousBinaryCacheInfo.NumberOfOldEntriesReused++;
-			return true;
-		}
-	}
-	return false;
-}
-
-void FOpenGLProgramBinaryCache::CompilePendingShaders(const FOpenGLLinkedProgramConfiguration& Config)
-{
-	VERIFY_GL_SCOPE();
-
-	// Find the existing compiled shader in the cache.
-	FScopeLock Lock(&GCompiledShaderCacheCS);
-	for (int32 StageIdx = 0; StageIdx < UE_ARRAY_COUNT(Config.Shaders); ++StageIdx)
-	{
-		TSharedPtr<FOpenGLCompiledShaderValue> FoundShader = GetOpenGLCompiledShaderCache().FindRef(Config.Shaders[StageIdx].ShaderKey);
-		if (FoundShader && FoundShader->bHasCompiled == false)
-		{
-			TArray<ANSICHAR> GlslCode = FoundShader->GetUncompressedShader();
-			CompileCurrentShader(Config.Shaders[StageIdx].Resource, GlslCode);
-			FoundShader->bHasCompiled = true;
-		}
-	}
-}
-
-FString FOpenGLProgramBinaryCache::GetProgramBinaryCacheFilePath() const
-{
-	FString ProgramFilename = CachePath + TEXT("/") + CacheFilename;
-	return ProgramFilename;
-}
-
-void FOpenGLProgramBinaryCache::CheckPendingGLProgramCreateRequests()
-{
-	FDelayedEvictionContainer::Get().Tick();
-	if (CachePtr)
-	{
-		QUICK_SCOPE_CYCLE_COUNTER(STAT_OpenGLShaderCreateShaderLibRequests);
-		CachePtr->CheckPendingGLProgramCreateRequests_internal();
-	}
-}
-
-void FOpenGLProgramBinaryCache::CheckPendingGLProgramCreateRequests_internal()
-{
-	check(IsInRenderingThread() || IsInRHIThread());
-	FScopeLock Lock(&GPendingGLProgramCreateRequestsCS);
-	//UE_LOG(LogRHI, Log, TEXT("CheckPendingGLProgramCreateRequests : PendingGLProgramCreateRequests = %d"), PendingGLProgramCreateRequests.Num());
-
-	if(PendingGLProgramCreateRequests.Num() > 0)
-	{
-		float TimeRemainingS = (float)GMaxShaderLibProcessingTimeMS / 1000.0f;
-		double StartTime = FPlatformTime::Seconds();
-		int32 Count = 0;
-		while(PendingGLProgramCreateRequests.Num() && TimeRemainingS > 0.0f)
-		{
-			CompleteLoadedGLProgramRequest_internal(PendingGLProgramCreateRequests.Pop());
-			TimeRemainingS -= (float)(FPlatformTime::Seconds() - StartTime);
-			StartTime = FPlatformTime::Seconds();
-			Count++;
-		}
-		float TimeTaken = (float)GMaxShaderLibProcessingTimeMS - (TimeRemainingS * 1000.0f);
-		UE_CLOG(TimeTaken>2.0f, LogRHI, Log, TEXT("CheckPendingGLProgramCreateRequests : iter count = %d, time taken = %f ms (remaining %d)"), Count, TimeTaken, PendingGLProgramCreateRequests.Num());
-	}
-}
-
-void FOpenGLProgramBinaryCache::CompleteLoadedGLProgramRequest_internal(FGLProgramBinaryFileCacheEntry* PendingGLCreate)
-{
-	VERIFY_GL_SCOPE();
-
-	check(PendingGLCreate->GLProgramState == FGLProgramBinaryFileCacheEntry::EGLProgramState::ProgramLoaded);
-
-	PendingGLCreate->ReadRequest = nullptr;
-
-	FOpenGLProgramKey& ProgramKey = PendingGLCreate->FileInfo.ShaderHasheSet;
-	const bool bProgramExists = GetOpenGLProgramsCache().Find(ProgramKey, false) != nullptr;
-
-	if (GetOpenGLProgramsCache().IsUsingLRU())
-	{
-		if (!bProgramExists)
-		{
-			// Always add programs as evicted, 1st use will create them as programs.
-			// This will reduce pressure on driver by ensuring only used programs
-			// are created.
-			// In this case do not create the GL program.
-			GetOpenGLProgramsCache().AddAsEvicted(ProgramKey, MoveTemp(PendingGLCreate->ProgramBinaryData));
-		}
-		else
-		{
-			// The program is already in use, discard the binary data.
-			PendingGLCreate->ProgramBinaryData.Empty();
-		}
-
-		// Ownership transfered to OpenGLProgramsCache.
-		PendingGLCreate->GLProgramState = FGLProgramBinaryFileCacheEntry::EGLProgramState::ProgramComplete;
-	}
-	else
-	{
-		if(!bProgramExists)
-		{
-			bool bSuccess = CreateGLProgramFromBinary(PendingGLCreate->GLProgramId, PendingGLCreate->ProgramBinaryData);
-			if (!bSuccess)
-			{
-				UE_LOG(LogRHI, Log, TEXT("[%s, %d, %d]"), *ProgramKey.ToString(), PendingGLCreate->GLProgramId, PendingGLCreate->ProgramBinaryData.Num());
-				RHIGetPanicDelegate().ExecuteIfBound(FName("FailedBinaryProgramCreateLoadRequest"));
-				UE_LOG(LogRHI, Fatal, TEXT("CompleteLoadedGLProgramRequest_internal : Failed to create GL program from binary data! [%s]"), *ProgramKey.ToString());
-			}
-			FOpenGLLinkedProgram* NewLinkedProgram = new FOpenGLLinkedProgram(ProgramKey, PendingGLCreate->GLProgramId);
-			GetOpenGLProgramsCache().Add(ProgramKey, NewLinkedProgram);
-			SetNewProgramStats(PendingGLCreate->GLProgramId);
-		}
-
-		// Finished with binary data.
-		PendingGLCreate->GLProgramState = FGLProgramBinaryFileCacheEntry::EGLProgramState::ProgramAvailable;
-		PendingGLCreate->ProgramBinaryData.Empty();
-	}
-}
-
-bool FOpenGLProgramBinaryCache::CheckSinglePendingGLProgramCreateRequest(const FOpenGLProgramKey& ProgramKey)
-{
-	if (CachePtr)
-	{
-		return CachePtr->CheckSinglePendingGLProgramCreateRequest_internal(ProgramKey);
-	}
-	return false;
-}
-
-// Any pending program must complete in this case.
-bool FOpenGLProgramBinaryCache::CheckSinglePendingGLProgramCreateRequest_internal(const FOpenGLProgramKey& ProgramKey)
-{
-	FGLProgramBinaryFileCacheEntry** ProgramBinRefPtr = nullptr;
-	FScopeLock ProgramBinaryCacheLock(&GProgramBinaryCacheCS);
-	ProgramBinRefPtr = CachePtr->ProgramToBinaryMap.Find(ProgramKey);
-	if( ProgramBinRefPtr )
-	{
-		FGLProgramBinaryFileCacheEntry* ProgramEntry = *ProgramBinRefPtr;
-		TSharedPtr<IAsyncReadRequest, ESPMode::ThreadSafe> LocalReadRequest = ProgramEntry->ReadRequest.Pin();
-		if (LocalReadRequest.IsValid())
-		{
-			ensure(ProgramEntry->GLProgramState == FGLProgramBinaryFileCacheEntry::EGLProgramState::ProgramLoading);
-			LocalReadRequest->WaitCompletion();
-			ProgramEntry->ReadRequest = nullptr;
-			ProgramEntry->GLProgramState = FGLProgramBinaryFileCacheEntry::EGLProgramState::ProgramLoaded;
-			CompleteLoadedGLProgramRequest_internal(ProgramEntry);
-		}
-		else
-		{
-			FScopeLock Lock(&GPendingGLProgramCreateRequestsCS);
-			if (ProgramEntry->GLProgramState == FGLProgramBinaryFileCacheEntry::EGLProgramState::ProgramLoaded)
-			{
-				int32 PendingRequestIndex = -1;
-				if (ensure(PendingGLProgramCreateRequests.Find(ProgramEntry, PendingRequestIndex)))
+				if (!bProgramExists)
 				{
-					CompleteLoadedGLProgramRequest_internal(ProgramEntry);
-					PendingGLProgramCreateRequests.RemoveAtSwap(PendingRequestIndex);
+					GLuint GLProgramId = 0;
+					bool bSuccess = UE::OpenGL::CreateGLProgramFromBinary(GLProgramId, ProgramBinaryData);
+					if (!bSuccess)
+					{
+						UE_LOG(LogRHI, Log, TEXT("[%s, %d, %d]"), *ProgramKey.ToString(), GLProgramId, ProgramBinaryData.Num());
+						RHIGetPanicDelegate().ExecuteIfBound(FName("FailedBinaryProgramCreateLoadRequest"));
+						UE_LOG(LogRHI, Fatal, TEXT("CompleteLoadedGLProgramRequest_internal : Failed to create GL program from binary data! [%s]"), *ProgramKey.ToString());
+					}
+					FOpenGLLinkedProgram* NewLinkedProgram = new FOpenGLLinkedProgram(ProgramKey, GLProgramId);
+					GetOpenGLProgramsCache().Add(ProgramKey, NewLinkedProgram);
+					SetNewProgramStats(GLProgramId);
 				}
-			}
-		}
-		return true;
-	}
-	return false;
-}
-
-bool OnExternalReadCallback(const TSharedPtr<IAsyncReadRequest, ESPMode::ThreadSafe>& AsyncReadRequest, FGLProgramBinaryFileCacheEntry* ProgramBinEntry, TArray<FGLProgramBinaryFileCacheEntry *>& PendingGLProgramCreateRequests, double RemainingTime)
-{
-	if (!AsyncReadRequest->WaitCompletion(RemainingTime))
-	{
-		return false;
-	}
-
-	FScopeLock ProgramBinaryCacheLock(&GProgramBinaryCacheCS);
-
-	if (ProgramBinEntry->GLProgramState == FGLProgramBinaryFileCacheEntry::EGLProgramState::ProgramLoading)
-	{
-		// Async load complete.
-		ProgramBinEntry->GLProgramState = FGLProgramBinaryFileCacheEntry::EGLProgramState::ProgramLoaded;
-		FOpenGLProgramKey& ProgramKey = ProgramBinEntry->FileInfo.ShaderHasheSet;
-
-		{
-			// Add this program to the create queue.
-			FScopeLock Lock(&GPendingGLProgramCreateRequestsCS);
-			PendingGLProgramCreateRequests.Add(ProgramBinEntry);
-		}
-	}
-
-	return true;
-}
 
-void FOpenGLProgramBinaryCache::BeginProgramReadRequest(FGLProgramBinaryFileCacheEntry* ProgramBinEntry, FArchive* Ar)
-{
-	check(ProgramBinEntry);
-
-	TSharedPtr<IAsyncReadRequest, ESPMode::ThreadSafe> LocalReadRequest = ProgramBinEntry->ReadRequest.Pin();
-	bool bHasReadRequest = LocalReadRequest.IsValid();
-
-	if (ensure(!bHasReadRequest))
-	{
-		check(ProgramBinEntry->ProgramBinaryData.Num() == 0);
-		check(ProgramBinEntry->GLProgramState == FGLProgramBinaryFileCacheEntry::EGLProgramState::ProgramStored);
-
-		int64 ReadSize = ProgramBinEntry->FileInfo.ProgramSize;
-		int64 ReadOffset = ProgramBinEntry->FileInfo.ProgramOffset;
-
-		if(ensure(ReadSize>0))
-		{
-			ProgramBinEntry->ProgramBinaryData.SetNumUninitialized(ReadSize);
-			ProgramBinEntry->GLProgramState = FGLProgramBinaryFileCacheEntry::EGLProgramState::ProgramLoading;
-			LocalReadRequest = MakeShareable(BinaryCacheAsyncReadFileHandle->ReadRequest(ReadOffset, ReadSize, AIOP_Normal, nullptr, ProgramBinEntry->ProgramBinaryData.GetData()));
-			ProgramBinEntry->ReadRequest = LocalReadRequest;
-			bHasReadRequest = true;
-
-			FExternalReadCallback ExternalReadCallback = [ProgramBinEntry, LocalReadRequest, this](double ReaminingTime)
-			{
-				return OnExternalReadCallback(LocalReadRequest, ProgramBinEntry, PendingGLProgramCreateRequests, ReaminingTime);
-			};
-
-			if (!Ar || !Ar->AttachExternalReadDependency(ExternalReadCallback))
-			{
-				// Archive does not support async loading
-				// do a blocking load
-				ExternalReadCallback(0.0);
+				// Finished with binary data.
+				ProgramBinaryData.Empty();
 			}
 		}
 	}
 }
 
-void FOpenGLProgramBinaryCache::OnShaderLibraryRequestShaderCode(const FSHAHash& Hash, FArchive* Ar)
-{
-	if (CachePtr)
-	{
-		CachePtr->OnShaderLibraryRequestShaderCode_internal(Hash, Ar);
-	}
-}
-
-void FOpenGLProgramBinaryCache::OnShaderLibraryRequestShaderCode_internal(const FSHAHash& Hash, FArchive* Ar)
+void FOpenGLDynamicRHI::EndFrameTick()
 {
-	FScopeLock Lock(&GProgramBinaryCacheCS);
-	FGLShaderToPrograms& FoundShaderToBinary = ShaderToProgramsMap.FindOrAdd(Hash);
-	if (!FoundShaderToBinary.bLoaded)
-	{
-		FoundShaderToBinary.bLoaded = true;
-
-		// if the binary cache is valid, look to see if we now have any complete programs to stream in.
-		// otherwise, we'll do this check bLoaded shaders when the binary cache loads.
-		if (BinaryFileState == EBinaryFileState::ValidCacheFile)
-		{
-			for (struct FGLProgramBinaryFileCacheEntry* ProgramBinEntry : FoundShaderToBinary.AssociatedPrograms)
-			{
-				const FOpenGLProgramKey& ProgramKey = ProgramBinEntry->FileInfo.ShaderHasheSet;
-				if (ProgramBinEntry->GLProgramState == FGLProgramBinaryFileCacheEntry::EGLProgramState::ProgramStored)
-				{
-					bool bAllShadersLoaded = true;
-					for (int32 i = 0; i < CrossCompiler::NUM_NON_COMPUTE_SHADER_STAGES && bAllShadersLoaded; i++)
-					{
-						bAllShadersLoaded = ProgramKey.ShaderHashes[i] == FSHAHash() || ShaderIsLoaded(ProgramKey.ShaderHashes[i]);
-					}
-
-					if (bAllShadersLoaded)
-					{
-						FOpenGLProgramBinaryCache::BeginProgramReadRequest(ProgramBinEntry, Ar);
-					}
-				}
-			}
-		}
-	}
+	FDelayedEvictionContainer::Get().Tick();
+	FOpenGLProgramBinaryCache::CheckPendingGLProgramCreateRequests();
+	FTextureEvictionLRU::Get().TickEviction();
 }
 
 static FDelegateHandle OnSharedShaderCodeRequest;
diff --git a/Engine/Source/Runtime/OpenGLDrv/Private/OpenGLUtil.cpp b/Engine/Source/Runtime/OpenGLDrv/Private/OpenGLUtil.cpp
--- a/Engine/Source/Runtime/OpenGLDrv/Private/OpenGLUtil.cpp
+++ b/Engine/Source/Runtime/OpenGLDrv/Private/OpenGLUtil.cpp
@@ -169,4 +169,118 @@ void RunOnGLRenderContextThread(TUniqueFunction<void(void)> GLFunc, bool bWaitFo
 			RHICmdList.ImmediateFlush(EImmediateFlushType::FlushRHIThread);
 		}
 	}
+}
+
+namespace UE
+{
+	namespace OpenGL
+	{
+		static TAutoConsoleVariable<int32> CVarStoreCompressedBinaries(
+			TEXT("r.OpenGL.StoreCompressedProgramBinaries"),
+			0,
+			TEXT(""),
+			ECVF_ReadOnly | ECVF_RenderThreadSafe
+		);
+
+		bool IsStoringCompressedBinaryPrograms()
+		{
+			return CVarStoreCompressedBinaries.GetValueOnAnyThread() != 0;
+		}
+
+		bool UncompressCompressedBinaryProgram(const TArray<uint8>& CompressedProgramBinary, TArray<uint8>& UncompressedProgramBinaryOUT)
+		{
+			if (ensure(CompressedProgramBinary.Num() > sizeof(FCompressedProgramBinaryHeader)))
+			{
+				FCompressedProgramBinaryHeader* Header = (FCompressedProgramBinaryHeader*)CompressedProgramBinary.GetData();
+
+				if (Header->UncompressedSize == FCompressedProgramBinaryHeader::NotCompressed)
+				{
+					const uint32 ProgramSize = CompressedProgramBinary.Num() - sizeof(FCompressedProgramBinaryHeader);
+					UncompressedProgramBinaryOUT.SetNumUninitialized(ProgramSize);
+					FMemory::Memcpy(UncompressedProgramBinaryOUT.GetData(), CompressedProgramBinary.GetData() + sizeof(FCompressedProgramBinaryHeader), ProgramSize);
+					return true;
+				}
+				else
+				{
+					UncompressedProgramBinaryOUT.AddUninitialized(Header->UncompressedSize);
+
+					if (Header->UncompressedSize > 0
+						&& FCompression::UncompressMemory(NAME_Zlib, UncompressedProgramBinaryOUT.GetData(), UncompressedProgramBinaryOUT.Num(), CompressedProgramBinary.GetData() + sizeof(FCompressedProgramBinaryHeader), CompressedProgramBinary.Num() - sizeof(FCompressedProgramBinaryHeader)))
+					{
+						return true;
+					}
+				}
+			}
+			return false;
+		}
+
+		bool GetUncompressedProgramBinaryFromGLProgram(GLuint Program, TArray<uint8>& ProgramBinaryOUT)
+		{
+			VERIFY_GL_SCOPE();
+
+			// pull binary from linked program
+			GLint BinaryLength = -1;
+			glGetProgramiv(Program, GL_PROGRAM_BINARY_LENGTH, &BinaryLength);
+			if (BinaryLength > 0)
+			{
+				ProgramBinaryOUT.SetNumUninitialized(BinaryLength + sizeof(GLenum));
+				uint8* ProgramBinaryPtr = ProgramBinaryOUT.GetData();
+				// BinaryFormat is stored at the start of ProgramBinary array
+				FOpenGL::GetProgramBinary(Program, BinaryLength, &BinaryLength, (GLenum*)ProgramBinaryPtr, ProgramBinaryPtr + sizeof(GLenum));
+				return true;
+			}
+			return false;
+		}
+
+		void CompressProgramBinary(const TArray<uint8>& UncompressedProgramBinary, TArray<uint8>& ProgramBinaryOUT)
+		{
+			check(IsStoringCompressedBinaryPrograms());
+			check(ProgramBinaryOUT.IsEmpty());
+
+			int32 CompressedSize = FCompression::CompressMemoryBound(NAME_Zlib, UncompressedProgramBinary.Num());
+			uint32 CompressedHeaderSize = sizeof(FCompressedProgramBinaryHeader);
+			ProgramBinaryOUT.AddUninitialized(CompressedSize + CompressedHeaderSize);
+			bool bSuccess = FCompression::CompressMemory(NAME_Zlib, ProgramBinaryOUT.GetData() + CompressedHeaderSize, CompressedSize, UncompressedProgramBinary.GetData(), UncompressedProgramBinary.Num());
+			if (bSuccess)
+			{
+				ProgramBinaryOUT.SetNum(CompressedSize + CompressedHeaderSize);
+				ProgramBinaryOUT.Shrink();
+				FCompressedProgramBinaryHeader* Header = (FCompressedProgramBinaryHeader*)ProgramBinaryOUT.GetData();
+				Header->UncompressedSize = UncompressedProgramBinary.Num();
+			}
+			else
+			{
+				// failed, store the uncompressed version.
+				UE_LOG(LogRHI, Log, TEXT("Storing binary program uncompressed (%d, %d, %d)"), UncompressedProgramBinary.Num(), ProgramBinaryOUT.Num(), CompressedSize);
+				ProgramBinaryOUT.SetNumUninitialized(UncompressedProgramBinary.Num() + CompressedHeaderSize);
+				FCompressedProgramBinaryHeader* Header = (FCompressedProgramBinaryHeader*)ProgramBinaryOUT.GetData();
+				Header->UncompressedSize = FCompressedProgramBinaryHeader::NotCompressed;
+				FMemory::Memcpy(ProgramBinaryOUT.GetData() + sizeof(FCompressedProgramBinaryHeader), UncompressedProgramBinary.GetData(), UncompressedProgramBinary.Num());
+			}
+		}
+
+		bool GetCompressedProgramBinaryFromGLProgram(GLuint Program, TArray<uint8>& ProgramBinaryOUT)
+		{
+			// get uncompressed binary
+			TArray<uint8> UncompressedProgramBinary;
+			if (GetUncompressedProgramBinaryFromGLProgram(Program, UncompressedProgramBinary))
+			{
+				CompressProgramBinary(UncompressedProgramBinary, ProgramBinaryOUT);
+				return true;
+			}
+			return false;
+		}
+
+		bool GetProgramBinaryFromGLProgram(GLuint Program, TArray<uint8>& ProgramBinaryOUT)
+		{
+			if (IsStoringCompressedBinaryPrograms())
+			{
+				return GetCompressedProgramBinaryFromGLProgram(Program, ProgramBinaryOUT);
+			}
+			else
+			{
+				return GetUncompressedProgramBinaryFromGLProgram(Program, ProgramBinaryOUT);
+			}
+		}
+	}
 }
\ No newline at end of file
diff --git a/Engine/Source/Runtime/OpenGLDrv/Private/OpenGLUtil.h b/Engine/Source/Runtime/OpenGLDrv/Private/OpenGLUtil.h
new file mode 100644
--- /dev/null
+++ b/Engine/Source/Runtime/OpenGLDrv/Private/OpenGLUtil.h
@@ -0,0 +1,185 @@
+// Copyright Epic Games, Inc. All Rights Reserved.
+
+/*=============================================================================
+	OpenGLUtil.h: OpenGL RHI utility definitions.
+=============================================================================*/
+
+#pragma once
+#include "RHICommandList.h"
+
+/** Set to 1 to enable the VERIFY_GL macros which call glGetError */
+#define ENABLE_VERIFY_GL (0 & DO_CHECK)
+#define ENABLE_VERIFY_GL_TRACE 0
+
+// Additional check that our GL calls are occurring on the expected thread
+#define ENABLE_VERIFY_GL_THREAD (UE_BUILD_DEBUG)
+
+/** Set to 1 to verify that the the engine side uniform buffer layout matches the driver side of the GLSL shader*/
+#define ENABLE_UNIFORM_BUFFER_LAYOUT_VERIFICATION ( 0 & UE_BUILD_DEBUG & (OPENGL_GL3 | OPENGL_GL4))
+
+/** Set to 1 to additinally dump uniform buffer layout at shader link time, this assumes ENABLE_UNIFORM_BUFFER_LAYOUT_VERIFICATION == 1 */
+#define ENABLE_UNIFORM_BUFFER_LAYOUT_DUMP 0
+
+/** Set to 1 to enable shader debugging which e.g. keeps the GLSL source as members of TOpenGLShader*/
+#define DEBUG_GL_SHADERS (UE_BUILD_DEBUG || UE_EDITOR)
+
+/** Set to 1 to enable calls to place event markers into the OpenGL stream
+    this is purposefully not considered for OPENGL_PERFORMANCE_DATA_INVALID, 
+	since there is an additional cvar OpenGLConsoleVariables::bEnableARBDebug*/
+
+#define ENABLE_OPENGL_DEBUG_GROUPS 1
+
+#define OPENGL_PERFORMANCE_DATA_INVALID (ENABLE_VERIFY_GL | ENABLE_UNIFORM_BUFFER_LAYOUT_VERIFICATION | DEBUG_GL_SHADERS)
+
+/**
+* Convert from ECubeFace to GLenum type
+* @param Face - ECubeFace type to convert
+* @return OpenGL cube face enum value
+*/
+GLenum GetOpenGLCubeFace(ECubeFace Face);
+
+extern bool PlatformOpenGLContextValid();
+
+#if ENABLE_VERIFY_GL_THREAD
+	#if UE_BUILD_TEST
+		#define GLCONTEXT_CLAUSE 
+	#else
+		#define GLCONTEXT_CLAUSE PlatformOpenGLContextValid() &&
+	#endif
+	// check that the current thread has a valid context and matches our RT / RHIT expectations.
+	// Note that the game thread can access the shared context.
+	// use commandline switch -norhithread if this causes issues.
+	#define CHECK_EXPECTED_GL_THREAD() \
+		if(!( \
+			 GLCONTEXT_CLAUSE\
+			(IsInGameThread() || ( (IsInRenderingThread() && !IsRunningRHIInSeparateThread()) || (IsInRHIThread() && IsRunningRHIInSeparateThread()) ))))\
+		{ \
+			UE_LOG(LogRHI, Fatal, TEXT("Potential use of GL context from incorrect thread. [ValidContext = %d] && [IsInGameThread() = %d, IsInRenderingThread() && !IsRunningRHIInSeparateThread() = %d, IsInRHIThread() && IsRunningRHIInSeparateThread() = %d]"), PlatformOpenGLContextValid(), IsInGameThread(), IsInRenderingThread() && !IsRunningRHIInSeparateThread(), IsInRHIThread() && IsRunningRHIInSeparateThread());\
+		}
+#else
+	#define CHECK_EXPECTED_GL_THREAD() 
+#endif
+
+#if ENABLE_VERIFY_GL
+	extern int32 PlatformGlGetError();
+
+	void VerifyOpenGLResult(GLenum ErrorCode, const TCHAR* Msg1, const TCHAR* Msg2, const TCHAR* Filename, uint32 Line);
+	#define VERIFY_GL(msg) { CHECK_EXPECTED_GL_THREAD(); GLenum ErrorCode = PlatformGlGetError(); if (ErrorCode != GL_NO_ERROR) { VerifyOpenGLResult(ErrorCode,TEXT(#msg),TEXT(""),TEXT(__FILE__),__LINE__); } }
+
+	struct FOpenGLErrorScope
+	{
+		const char* FunctionName;
+		const TCHAR* Filename;
+		const uint32 Line;
+
+		FOpenGLErrorScope(
+			const char* InFunctionName,
+			const TCHAR* InFilename,
+			const uint32 InLine)
+			: FunctionName(InFunctionName)
+			, Filename(InFilename)
+			, Line(InLine)
+		{
+#if ENABLE_VERIFY_GL_TRACE
+			UE_LOG(LogRHI, Log, TEXT("log before %s(%d): %s"), InFilename, InLine, ANSI_TO_TCHAR(InFunctionName));
+#endif
+			CheckForErrors(0);
+		}
+
+		~FOpenGLErrorScope()
+		{
+#if ENABLE_VERIFY_GL_TRACE
+			UE_LOG(LogRHI, Log, TEXT("log after  %s(%d): %s"), Filename, Line, ANSI_TO_TCHAR(FunctionName));
+
+#endif
+
+			CheckForErrors(1);
+		}
+
+		void CheckForErrors(int32 BeforeOrAfter)
+		{
+			check(PlatformOpenGLContextValid());
+
+			GLenum ErrorCode = PlatformGlGetError();
+			if (ErrorCode != GL_NO_ERROR)
+			{
+				const TCHAR* PrefixStrings[] = { TEXT("Before "), TEXT("During ") };
+				VerifyOpenGLResult(ErrorCode,PrefixStrings[BeforeOrAfter], ANSI_TO_TCHAR(FunctionName),Filename,Line);
+			}
+		}
+	};
+	#define MACRO_TOKENIZER(IdentifierName, Msg, FileName, LineNumber) FOpenGLErrorScope IdentifierName_ ## LineNumber (Msg, FileName, LineNumber)
+	#define MACRO_TOKENIZER2(IdentifierName, Msg, FileName, LineNumber) MACRO_TOKENIZER(IdentiferName, Msg, FileName, LineNumber)
+	#define VERIFY_GL_SCOPE_WITH_MSG_STR(MsgStr) CHECK_EXPECTED_GL_THREAD(); MACRO_TOKENIZER2(ErrorScope_, MsgStr, TEXT(__FILE__), __LINE__)
+	#define VERIFY_GL_SCOPE() VERIFY_GL_SCOPE_WITH_MSG_STR(__FUNCTION__)
+	#define VERIFY_GL_FUNC(Func, ...) { VERIFY_GL_SCOPE_WITH_MSG_STR((#Func)); Func(__VA_ARGS__); }
+
+	/**
+	 * Some important GL calls are trapped individually.
+	 */
+	#define glBlitFramebuffer(...) VERIFY_GL_FUNC(glBlitFramebuffer, __VA_ARGS__)
+	#define glTexImage2D(...) VERIFY_GL_FUNC(glTexImage2D, __VA_ARGS__)
+	#define glTexSubImage2D(...) VERIFY_GL_FUNC(glTexSubImage2D, __VA_ARGS__)
+	#define glCompressedTexImage2D(...) VERIFY_GL_FUNC(glCompressedTexImage2D, __VA_ARGS__)
+
+#else
+	#define VERIFY_GL(...) CHECK_EXPECTED_GL_THREAD()
+	#define VERIFY_GL_SCOPE(...) CHECK_EXPECTED_GL_THREAD()
+#endif
+
+struct FRHICommandGLCommandString
+{
+	static const TCHAR* TStr() { return TEXT("FRHICommandGLCommand"); }
+};
+
+#define GL_CAPTURE_CALLSTACK 0 // Capture the callstack at the point of enqueuing the command. 
+
+struct FRHICommandGLCommand final : public FRHICommand<FRHICommandGLCommand, FRHICommandGLCommandString>
+{
+#if GL_CAPTURE_CALLSTACK
+	uint64 CallStack[16];
+#endif
+	TUniqueFunction<void()> GLFunction;
+
+	FORCEINLINE_DEBUGGABLE FRHICommandGLCommand(TUniqueFunction<void()> InGLFunction)
+		: GLFunction(MoveTemp(InGLFunction))
+	{
+#if GL_CAPTURE_CALLSTACK
+		FPlatformStackWalk::CaptureStackBackTrace(CallStack, UE_ARRAY_COUNT(CallStack), nullptr);
+#endif
+	}
+
+	void Execute(FRHICommandListBase& CmdList)
+	{
+		GLFunction();
+	}
+};
+
+void RunOnGLRenderContextThread(TUniqueFunction<void(void)> GLFunc, bool bWaitForCompletion = false);
+inline bool ShouldRunGLRenderContextOpOnThisThread(FRHICommandListImmediate& RHICmdList)
+{
+	return (RHICmdList.Bypass() || !IsRunningRHIInSeparateThread() || IsInRHIThread());
+}
+
+
+// Program Binary helpers.
+namespace UE
+{
+	namespace OpenGL
+	{
+		bool IsStoringCompressedBinaryPrograms();
+
+		struct FCompressedProgramBinaryHeader
+		{
+			static const uint32 NotCompressed = 0xFFFFFFFF;
+			uint32 UncompressedSize;
+		};
+
+		bool GetProgramBinaryFromGLProgram(GLuint Program, TArray<uint8>& ProgramBinaryOUT);
+		bool GetUncompressedProgramBinaryFromGLProgram(GLuint Program, TArray<uint8>& ProgramBinaryOUT);
+		bool GetCompressedProgramBinaryFromGLProgram(GLuint Program, TArray<uint8>& ProgramBinaryOUT);
+		void CompressProgramBinary(const TArray<uint8>& UncompressedProgramBinary, TArray<uint8>& ProgramBinaryOUT);
+		bool UncompressCompressedBinaryProgram(const TArray<uint8>& CompressedProgramBinary, TArray<uint8>& UncompressedProgramBinaryOUT);
+	}
+}
+
diff --git a/Engine/Source/Runtime/OpenGLDrv/Private/OpenGLViewport.cpp b/Engine/Source/Runtime/OpenGLDrv/Private/OpenGLViewport.cpp
--- a/Engine/Source/Runtime/OpenGLDrv/Private/OpenGLViewport.cpp
+++ b/Engine/Source/Runtime/OpenGLDrv/Private/OpenGLViewport.cpp
@@ -200,10 +200,7 @@ void FOpenGLDynamicRHI::RHIEndDrawingViewport(FRHIViewport* ViewportRHI,bool bPr
 		GRenderThreadNumIdle[ERenderThreadIdleTypes::WaitingForGPUPresent]++;
 	}
 
-	// TODO: find better location to poll this, or create programs on separate thread. Gil had a prototype of this.
-	FOpenGLProgramBinaryCache::CheckPendingGLProgramCreateRequests();
-
-	FTextureEvictionLRU::Get().TickEviction();
+	EndFrameTick();
 }
 
 
diff --git a/Engine/Source/Runtime/OpenGLDrv/Public/OpenGLDrv.h b/Engine/Source/Runtime/OpenGLDrv/Public/OpenGLDrv.h
--- a/Engine/Source/Runtime/OpenGLDrv/Public/OpenGLDrv.h
+++ b/Engine/Source/Runtime/OpenGLDrv/Public/OpenGLDrv.h
@@ -939,6 +939,10 @@ private:
 
 	void PrepareGFXBoundShaderState(const FGraphicsPipelineStateInitializer& Initializer);
 
+
+	/** called once per frame, used for resource processing */
+	void EndFrameTick();
+
 	/** Counter incremented each time RHIBeginScene is called. */
 	uint32 SceneFrameCounter;
 
diff --git a/Engine/Source/Runtime/OpenGLDrv/Public/OpenGLShaderResources.h b/Engine/Source/Runtime/OpenGLDrv/Public/OpenGLShaderResources.h
--- a/Engine/Source/Runtime/OpenGLDrv/Public/OpenGLShaderResources.h
+++ b/Engine/Source/Runtime/OpenGLDrv/Public/OpenGLShaderResources.h
@@ -509,182 +509,3 @@ public:
 		return GetTypeHash(Config.ProgramKey);
 	}
 };
-
-struct FGLProgramBinaryFileCacheEntry;
-class FOpenGLProgramBinaryCache
-{
-public:
-	static void Initialize();
-	static void Shutdown();
-
-	static bool IsEnabled();
-
-	/** Compile required shaders for a program, only in case binary program was not found in the cache   */
-	static void CompilePendingShaders(const FOpenGLLinkedProgramConfiguration& Config);
-	
-	/** Try to find and load program binary from cache */
-	static bool UseCachedProgram(GLuint& ProgramOUT, const FOpenGLProgramKey& ProgramKey, TArray<uint8>& CachedProgramBinaryOUT);
-	
-	/** Extract the program binary from GL and store on disk. Only when ProgramBinaryCache is enabled */
-	static void CacheProgram(GLuint Program, const FOpenGLProgramKey& ProgramKey, TArray<uint8>& CachedProgramBinaryOUT);
-
-	/** Take an existing program binary and store on disk. Only when ProgramBinaryCache is enabled
-	* CachedProgramBinary should be const, cant be due to serialize api
-	*/
-	static void CacheProgramBinary(const FOpenGLProgramKey& ProgramKey, /*const*/ TArray<uint8>& CachedProgramBinary); 
-
-	static void OnShaderLibraryRequestShaderCode(const FSHAHash& Hash, FArchive* Ar);
-
-	/** Create any pending GL programs that have come from shader library requests */
-	static void CheckPendingGLProgramCreateRequests();
-
-	/** Create any single GL program that have come from shader library requests */
-	static bool CheckSinglePendingGLProgramCreateRequest(const FOpenGLProgramKey& ProgramKey);
-
-	/** true if the program binary cache is currently in cache build mode */
-	static bool IsBuildingCache();
-
-private:
-	FOpenGLProgramBinaryCache(const FString& InCachePath);
-	~FOpenGLProgramBinaryCache();
-
-	FString GetProgramBinaryCacheFilePath() const;
-	
-	struct FPendingShaderCode
-	{
-		TArray<ANSICHAR> GlslCode;
-		int32 UncompressedSize;
-		bool bCompressed;
-	};
-	
-private:
-	static TAutoConsoleVariable<int32> CVarPBCEnable;
-	static TAutoConsoleVariable<int32> CVarRestartAndroidAfterPrecompile;
-	static FOpenGLProgramBinaryCache* CachePtr;
-	
-	/*  Path to directory where binary programs will be stored excluding the cache filename */
-	FString CachePath;
-	/* Just the cache filename, without the path */
-	FString CacheFilename;
-
-	
-	/**
-	* Shaders that were requested for compilation
-	* They will be compiled just before linking a program only in case when there is no saved binary program
-	*/
-	TMap<GLuint, FPendingShaderCode> ShadersPendingCompilation;
-
-	bool AppendProgramBinaryFile(FArchive& Ar, const FOpenGLProgramKey& ProgramKey, GLuint Program, uint32& ProgramBinaryOffsetOUT, uint32& ProgramBinarySizeOUT);
-	void AppendProgramBinaryFileEofEntry(FArchive& Ar);
-
-	void ScanProgramCacheFile(const FGuid& ShaderPipelineCacheVersionGuid = FGuid());
-
-	/* Add a program */
-	void AddProgramFileEntryToMap(FGLProgramBinaryFileCacheEntry* IndexEntry);
-
-	void OpenAsyncReadHandle();
-	void CloseAsyncReadHandle();
-
-	bool OpenWriteHandle();
-	void CloseWriteHandle();
-
-	void AppendGLProgramToBinaryCache(const FOpenGLProgramKey& ProgramKey, GLuint Program, TArray<uint8>& CachedProgramBinaryOUT);
-	void AddUniqueGLProgramToBinaryCache(const FOpenGLProgramKey& ProgramKey, GLuint Program, TArray<uint8>& CachedProgramBinaryOUT);
-
-	void AddProgramBinaryDataToBinaryCache(TArray<uint8>& BinaryProgramData, const FOpenGLProgramKey& ProgramKey);
-
-	void ReleaseGLProgram_internal(FOpenGLLinkedProgramConfiguration& Config, GLuint Program);
-
-	FORCEINLINE_DEBUGGABLE bool ShaderIsLoaded(const FSHAHash& Hash)
-	{
-		const FGLShaderToPrograms* FoundShaderToBinary = CachePtr->ShaderToProgramsMap.Find(Hash);
-		return FoundShaderToBinary && FoundShaderToBinary->bLoaded;
-	}
-
-	bool UseCachedProgram_internal(GLuint& ProgramOUT, const FOpenGLProgramKey& ProgramKey, TArray<uint8>& CachedProgramBinaryOUT);
-
-	void OnShaderLibraryRequestShaderCode_internal(const FSHAHash& Hash, FArchive* Ar);
-
-	void BeginProgramReadRequest(FGLProgramBinaryFileCacheEntry* IndexEntry, FArchive* Ar);
-
-	void CheckPendingGLProgramCreateRequests_internal();
-	bool CheckSinglePendingGLProgramCreateRequest_internal(const FOpenGLProgramKey& ProgramKey);
-
-	void CompleteLoadedGLProgramRequest_internal(FGLProgramBinaryFileCacheEntry* PendingGLCreate);
-
-	/** Delegate handlers to track the ShaderPipelineCache precompile. */
-	void OnShaderPipelineCacheOpened(FString const& Name, EShaderPlatform Platform, uint32 Count, const FGuid& VersionGuid, FShaderPipelineCache::FShaderCachePrecompileContext& ShaderCachePrecompileContext);
-	void OnShaderPipelineCachePrecompilationComplete(uint32 Count, double Seconds, const FShaderPipelineCache::FShaderCachePrecompileContext& ShaderCachePrecompileContext);
-
-private:
-
-	FDelegateHandle OnShaderPipelineCacheOpenedDelegate;
-	FDelegateHandle OnShaderPipelineCachePrecompilationCompleteDelegate;
-
-	TArray<TUniquePtr<FGLProgramBinaryFileCacheEntry>> ProgramEntryContainer; // this is the owner of all FGLProgramBinaryFileCacheEntry ptrs
-
-	TMap<FOpenGLProgramKey, FGLProgramBinaryFileCacheEntry*> ProgramToBinaryMap; // program key to program entry.
-
-	struct FGLShaderToPrograms
-	{
-		FGLShaderToPrograms() : bLoaded(false)
-		{
-		}
-
-		FGLShaderToPrograms(FGLProgramBinaryFileCacheEntry* ProgramEntry) : bLoaded(false)
-		{
-			AssociatedPrograms.Add(ProgramEntry);
-		}
-
-		void Add(FGLProgramBinaryFileCacheEntry* ProgramEntry)
-		{
-			checkSlow(!AssociatedPrograms.Contains(ProgramEntry));
-			AssociatedPrograms.Add(ProgramEntry);
-		}
-
-		bool bLoaded;
-		TArray<FGLProgramBinaryFileCacheEntry*> AssociatedPrograms;
-	};
-
-	// Map of shader hash to a list of programs which reference it.
-	TMap<FSHAHash, FGLShaderToPrograms> ShaderToProgramsMap;
-
-	// programs loaded via async and now ready for creation on GL context owning thread.
-	TArray<FGLProgramBinaryFileCacheEntry*> PendingGLProgramCreateRequests;
-
-	IAsyncReadFileHandle* BinaryCacheAsyncReadFileHandle;
-	FArchive* BinaryCacheWriteFileHandle;
-	bool bShownLoadingScreen;
-
-	enum class EBinaryFileState : uint8
-	{
-		Uninitialized,					// No binary file is yet established and we should not read or write to it.
-		BuildingCacheFile,				// We are precompiling shaders from the PSO and storing them in a new binary cache. Do not attempt to read.
-		BuildingCacheFileWithMove,		// We are precompiling shaders from the PSO and storing them in a new binary cache, shaders matching from the existing cache are moved to the new file. Do not attempt to read.
-		ValidCacheFile,					// We have a valid cache file we can use for reading. Do not attempt to write.
-	};
-
-	bool IsBuildingCache_internal() const
-	{
-		return BinaryFileState == EBinaryFileState::BuildingCacheFile || BinaryFileState == EBinaryFileState::BuildingCacheFileWithMove;
-	}
-
-	struct FPreviousGLProgramBinaryCacheInfo
-	{
-		FPreviousGLProgramBinaryCacheInfo();
-		FPreviousGLProgramBinaryCacheInfo(FPreviousGLProgramBinaryCacheInfo&&);
-		FPreviousGLProgramBinaryCacheInfo& operator=(FPreviousGLProgramBinaryCacheInfo&&);
-		~FPreviousGLProgramBinaryCacheInfo();
-
-
-		FString OldCacheFilename;
-		TUniquePtr<FArchive> OldCacheArchive;
-		TMap<FOpenGLProgramKey, TUniquePtr<FGLProgramBinaryFileCacheEntry> > ProgramToOldBinaryCacheMap; // program key to program entry for old cache used when generating new cache.
-
-		// for logging
-		uint32 NumberOfOldEntriesReused;
-	};
-	FPreviousGLProgramBinaryCacheInfo PreviousBinaryCacheInfo;
-
-	EBinaryFileState BinaryFileState;
-};
diff --git a/Engine/Source/Runtime/OpenGLDrv/Public/OpenGLUtil.h b/Engine/Source/Runtime/OpenGLDrv/Public/OpenGLUtil.h
deleted file mode 100644
