Move OpenGLUtil.h back to public dir
Fixing compile issues.

diff --git a/Engine/Source/Runtime/OpenGLDrv/Private/OpenGLUtil.h b/Engine/Source/Runtime/OpenGLDrv/Private/OpenGLUtil.h
deleted file mode 100644
diff --git a/Engine/Source/Runtime/OpenGLDrv/Public/OpenGLUtil.h b/Engine/Source/Runtime/OpenGLDrv/Public/OpenGLUtil.h
new file mode 100644
--- /dev/null
+++ b/Engine/Source/Runtime/OpenGLDrv/Public/OpenGLUtil.h
@@ -0,0 +1,185 @@
+// Copyright Epic Games, Inc. All Rights Reserved.
+
+/*=============================================================================
+	OpenGLUtil.h: OpenGL RHI utility definitions.
+=============================================================================*/
+
+#pragma once
+#include "RHICommandList.h"
+
+/** Set to 1 to enable the VERIFY_GL macros which call glGetError */
+#define ENABLE_VERIFY_GL (0 & DO_CHECK)
+#define ENABLE_VERIFY_GL_TRACE 0
+
+// Additional check that our GL calls are occurring on the expected thread
+#define ENABLE_VERIFY_GL_THREAD (UE_BUILD_DEBUG)
+
+/** Set to 1 to verify that the the engine side uniform buffer layout matches the driver side of the GLSL shader*/
+#define ENABLE_UNIFORM_BUFFER_LAYOUT_VERIFICATION ( 0 & UE_BUILD_DEBUG & (OPENGL_GL3 | OPENGL_GL4))
+
+/** Set to 1 to additinally dump uniform buffer layout at shader link time, this assumes ENABLE_UNIFORM_BUFFER_LAYOUT_VERIFICATION == 1 */
+#define ENABLE_UNIFORM_BUFFER_LAYOUT_DUMP 0
+
+/** Set to 1 to enable shader debugging which e.g. keeps the GLSL source as members of TOpenGLShader*/
+#define DEBUG_GL_SHADERS (UE_BUILD_DEBUG || UE_EDITOR)
+
+/** Set to 1 to enable calls to place event markers into the OpenGL stream
+    this is purposefully not considered for OPENGL_PERFORMANCE_DATA_INVALID, 
+	since there is an additional cvar OpenGLConsoleVariables::bEnableARBDebug*/
+
+#define ENABLE_OPENGL_DEBUG_GROUPS 1
+
+#define OPENGL_PERFORMANCE_DATA_INVALID (ENABLE_VERIFY_GL | ENABLE_UNIFORM_BUFFER_LAYOUT_VERIFICATION | DEBUG_GL_SHADERS)
+
+/**
+* Convert from ECubeFace to GLenum type
+* @param Face - ECubeFace type to convert
+* @return OpenGL cube face enum value
+*/
+GLenum GetOpenGLCubeFace(ECubeFace Face);
+
+extern bool PlatformOpenGLContextValid();
+
+#if ENABLE_VERIFY_GL_THREAD
+	#if UE_BUILD_TEST
+		#define GLCONTEXT_CLAUSE 
+	#else
+		#define GLCONTEXT_CLAUSE PlatformOpenGLContextValid() &&
+	#endif
+	// check that the current thread has a valid context and matches our RT / RHIT expectations.
+	// Note that the game thread can access the shared context.
+	// use commandline switch -norhithread if this causes issues.
+	#define CHECK_EXPECTED_GL_THREAD() \
+		if(!( \
+			 GLCONTEXT_CLAUSE\
+			(IsInGameThread() || ( (IsInRenderingThread() && !IsRunningRHIInSeparateThread()) || (IsInRHIThread() && IsRunningRHIInSeparateThread()) ))))\
+		{ \
+			UE_LOG(LogRHI, Fatal, TEXT("Potential use of GL context from incorrect thread. [ValidContext = %d] && [IsInGameThread() = %d, IsInRenderingThread() && !IsRunningRHIInSeparateThread() = %d, IsInRHIThread() && IsRunningRHIInSeparateThread() = %d]"), PlatformOpenGLContextValid(), IsInGameThread(), IsInRenderingThread() && !IsRunningRHIInSeparateThread(), IsInRHIThread() && IsRunningRHIInSeparateThread());\
+		}
+#else
+	#define CHECK_EXPECTED_GL_THREAD() 
+#endif
+
+#if ENABLE_VERIFY_GL
+	extern int32 PlatformGlGetError();
+
+	void VerifyOpenGLResult(GLenum ErrorCode, const TCHAR* Msg1, const TCHAR* Msg2, const TCHAR* Filename, uint32 Line);
+	#define VERIFY_GL(msg) { CHECK_EXPECTED_GL_THREAD(); GLenum ErrorCode = PlatformGlGetError(); if (ErrorCode != GL_NO_ERROR) { VerifyOpenGLResult(ErrorCode,TEXT(#msg),TEXT(""),TEXT(__FILE__),__LINE__); } }
+
+	struct FOpenGLErrorScope
+	{
+		const char* FunctionName;
+		const TCHAR* Filename;
+		const uint32 Line;
+
+		FOpenGLErrorScope(
+			const char* InFunctionName,
+			const TCHAR* InFilename,
+			const uint32 InLine)
+			: FunctionName(InFunctionName)
+			, Filename(InFilename)
+			, Line(InLine)
+		{
+#if ENABLE_VERIFY_GL_TRACE
+			UE_LOG(LogRHI, Log, TEXT("log before %s(%d): %s"), InFilename, InLine, ANSI_TO_TCHAR(InFunctionName));
+#endif
+			CheckForErrors(0);
+		}
+
+		~FOpenGLErrorScope()
+		{
+#if ENABLE_VERIFY_GL_TRACE
+			UE_LOG(LogRHI, Log, TEXT("log after  %s(%d): %s"), Filename, Line, ANSI_TO_TCHAR(FunctionName));
+
+#endif
+
+			CheckForErrors(1);
+		}
+
+		void CheckForErrors(int32 BeforeOrAfter)
+		{
+			check(PlatformOpenGLContextValid());
+
+			GLenum ErrorCode = PlatformGlGetError();
+			if (ErrorCode != GL_NO_ERROR)
+			{
+				const TCHAR* PrefixStrings[] = { TEXT("Before "), TEXT("During ") };
+				VerifyOpenGLResult(ErrorCode,PrefixStrings[BeforeOrAfter], ANSI_TO_TCHAR(FunctionName),Filename,Line);
+			}
+		}
+	};
+	#define MACRO_TOKENIZER(IdentifierName, Msg, FileName, LineNumber) FOpenGLErrorScope IdentifierName_ ## LineNumber (Msg, FileName, LineNumber)
+	#define MACRO_TOKENIZER2(IdentifierName, Msg, FileName, LineNumber) MACRO_TOKENIZER(IdentiferName, Msg, FileName, LineNumber)
+	#define VERIFY_GL_SCOPE_WITH_MSG_STR(MsgStr) CHECK_EXPECTED_GL_THREAD(); MACRO_TOKENIZER2(ErrorScope_, MsgStr, TEXT(__FILE__), __LINE__)
+	#define VERIFY_GL_SCOPE() VERIFY_GL_SCOPE_WITH_MSG_STR(__FUNCTION__)
+	#define VERIFY_GL_FUNC(Func, ...) { VERIFY_GL_SCOPE_WITH_MSG_STR((#Func)); Func(__VA_ARGS__); }
+
+	/**
+	 * Some important GL calls are trapped individually.
+	 */
+	#define glBlitFramebuffer(...) VERIFY_GL_FUNC(glBlitFramebuffer, __VA_ARGS__)
+	#define glTexImage2D(...) VERIFY_GL_FUNC(glTexImage2D, __VA_ARGS__)
+	#define glTexSubImage2D(...) VERIFY_GL_FUNC(glTexSubImage2D, __VA_ARGS__)
+	#define glCompressedTexImage2D(...) VERIFY_GL_FUNC(glCompressedTexImage2D, __VA_ARGS__)
+
+#else
+	#define VERIFY_GL(...) CHECK_EXPECTED_GL_THREAD()
+	#define VERIFY_GL_SCOPE(...) CHECK_EXPECTED_GL_THREAD()
+#endif
+
+struct FRHICommandGLCommandString
+{
+	static const TCHAR* TStr() { return TEXT("FRHICommandGLCommand"); }
+};
+
+#define GL_CAPTURE_CALLSTACK 0 // Capture the callstack at the point of enqueuing the command. 
+
+struct FRHICommandGLCommand final : public FRHICommand<FRHICommandGLCommand, FRHICommandGLCommandString>
+{
+#if GL_CAPTURE_CALLSTACK
+	uint64 CallStack[16];
+#endif
+	TUniqueFunction<void()> GLFunction;
+
+	FORCEINLINE_DEBUGGABLE FRHICommandGLCommand(TUniqueFunction<void()> InGLFunction)
+		: GLFunction(MoveTemp(InGLFunction))
+	{
+#if GL_CAPTURE_CALLSTACK
+		FPlatformStackWalk::CaptureStackBackTrace(CallStack, UE_ARRAY_COUNT(CallStack), nullptr);
+#endif
+	}
+
+	void Execute(FRHICommandListBase& CmdList)
+	{
+		GLFunction();
+	}
+};
+
+void RunOnGLRenderContextThread(TUniqueFunction<void(void)> GLFunc, bool bWaitForCompletion = false);
+inline bool ShouldRunGLRenderContextOpOnThisThread(FRHICommandListImmediate& RHICmdList)
+{
+	return (RHICmdList.Bypass() || !IsRunningRHIInSeparateThread() || IsInRHIThread());
+}
+
+
+// Program Binary helpers.
+namespace UE
+{
+	namespace OpenGL
+	{
+		bool IsStoringCompressedBinaryPrograms();
+
+		struct FCompressedProgramBinaryHeader
+		{
+			static const uint32 NotCompressed = 0xFFFFFFFF;
+			uint32 UncompressedSize;
+		};
+
+		bool GetProgramBinaryFromGLProgram(GLuint Program, TArray<uint8>& ProgramBinaryOUT);
+		bool GetUncompressedProgramBinaryFromGLProgram(GLuint Program, TArray<uint8>& ProgramBinaryOUT);
+		bool GetCompressedProgramBinaryFromGLProgram(GLuint Program, TArray<uint8>& ProgramBinaryOUT);
+		void CompressProgramBinary(const TArray<uint8>& UncompressedProgramBinary, TArray<uint8>& ProgramBinaryOUT);
+		bool UncompressCompressedBinaryProgram(const TArray<uint8>& CompressedProgramBinary, TArray<uint8>& UncompressedProgramBinaryOUT);
+	}
+}
+
