Fix memory corruption bug on GLES devices with more than 48 texture units per stage. (latest mali drivers)

Clamp compute shader texture counts so we stay within the bounds of the texture state array.
Increase vertex shader texture count from 8 to 16.
Brings RHI in line with the shader compiler and possibly fixes rendering issues where >8 vertex textures are used.

diff --git a/Engine/Source/Runtime/OpenGLDrv/Private/Android/AndroidOpenGL.cpp b/Engine/Source/Runtime/OpenGLDrv/Private/Android/AndroidOpenGL.cpp
--- a/Engine/Source/Runtime/OpenGLDrv/Private/Android/AndroidOpenGL.cpp
+++ b/Engine/Source/Runtime/OpenGLDrv/Private/Android/AndroidOpenGL.cpp
@@ -871,6 +871,9 @@ void FAndroidOpenGL::ProcessExtensions(const FString& ExtensionsString)
 
 	FOpenGLES::ProcessExtensions(ExtensionsString);
 
+	// clamp the number of texture units so we stay within the limits of our ContextState.Textures array.
+	MaxComputeTextureImageUnits = FMath::Min(MaxComputeTextureImageUnits, GetMaxCombinedTextureImageUnits());
+
 	FString RendererString = FString(ANSI_TO_TCHAR((const ANSICHAR*)glGetString(GL_RENDERER)));
 
 	// Common GPU types
diff --git a/Engine/Source/Runtime/OpenGLDrv/Private/OpenGLDrv.cpp b/Engine/Source/Runtime/OpenGLDrv/Private/OpenGLDrv.cpp
--- a/Engine/Source/Runtime/OpenGLDrv/Private/OpenGLDrv.cpp
+++ b/Engine/Source/Runtime/OpenGLDrv/Private/OpenGLDrv.cpp
@@ -530,8 +530,8 @@ void FOpenGLBase::ProcessExtensions( const FString& ExtensionsString )
 		if (IsMobilePlatform(GMaxRHIShaderPlatform))
 		{
 			MaxTextureImageUnits = MaxTextureImageUnits > 16 ? 16 : MaxTextureImageUnits;
-			MaxVertexTextureImageUnits = MaxVertexTextureImageUnits > 8 ? 8 : MaxVertexTextureImageUnits;
-			MaxGeometryTextureImageUnits = MaxGeometryTextureImageUnits > 8 ? 8 : MaxGeometryTextureImageUnits;
+			MaxVertexTextureImageUnits = MaxVertexTextureImageUnits > 16 ? 16 : MaxVertexTextureImageUnits;
+			MaxGeometryTextureImageUnits = MaxGeometryTextureImageUnits > 16 ? 16 : MaxGeometryTextureImageUnits;
 			MaxCombinedTextureImageUnits = MaxCombinedTextureImageUnits > 48 ? 48 : MaxCombinedTextureImageUnits;
 		}
 	}
