Moved MaxComputeTextureImageUnits out of subclasses into FOpenGLBase.
FOpenGLBase now clamps the number of gles compute texture units to be no bigger than the total combined texture unit size.
Switched Textures and Samplerstates in FOpenGLCommonState to using TArrays.
refinement of CL 20670879

diff --git a/Engine/Source/Runtime/OpenGLDrv/Private/Android/AndroidOpenGL.cpp b/Engine/Source/Runtime/OpenGLDrv/Private/Android/AndroidOpenGL.cpp
--- a/Engine/Source/Runtime/OpenGLDrv/Private/Android/AndroidOpenGL.cpp
+++ b/Engine/Source/Runtime/OpenGLDrv/Private/Android/AndroidOpenGL.cpp
@@ -871,9 +871,6 @@ void FAndroidOpenGL::ProcessExtensions(const FString& ExtensionsString)
 
 	FOpenGLES::ProcessExtensions(ExtensionsString);
 
-	// clamp the number of texture units so we stay within the limits of our ContextState.Textures array.
-	MaxComputeTextureImageUnits = FMath::Min(MaxComputeTextureImageUnits, GetMaxCombinedTextureImageUnits());
-
 	FString RendererString = FString(ANSI_TO_TCHAR((const ANSICHAR*)glGetString(GL_RENDERER)));
 
 	// Common GPU types
diff --git a/Engine/Source/Runtime/OpenGLDrv/Private/Android/AndroidOpenGL.h b/Engine/Source/Runtime/OpenGLDrv/Private/Android/AndroidOpenGL.h
--- a/Engine/Source/Runtime/OpenGLDrv/Private/Android/AndroidOpenGL.h
+++ b/Engine/Source/Runtime/OpenGLDrv/Private/Android/AndroidOpenGL.h
@@ -153,7 +153,6 @@ struct FAndroidOpenGL : public FOpenGLES
 
 	static FORCEINLINE EImageExternalType GetImageExternalType() { return ImageExternalType; }
 
-	static FORCEINLINE GLint GetMaxComputeTextureImageUnits() { check(MaxComputeTextureImageUnits != -1); return MaxComputeTextureImageUnits; }
 	static FORCEINLINE GLint GetMaxComputeUniformComponents() { check(MaxComputeUniformComponents != -1); return MaxComputeUniformComponents; }
 	static FORCEINLINE GLint GetFirstComputeUAVUnit()			{ return 0; }
 	static FORCEINLINE GLint GetMaxComputeUAVUnits()			{ check(MaxComputeUAVUnits != -1); return MaxComputeUAVUnits; }
diff --git a/Engine/Source/Runtime/OpenGLDrv/Private/OpenGL4.cpp b/Engine/Source/Runtime/OpenGLDrv/Private/OpenGL4.cpp
--- a/Engine/Source/Runtime/OpenGLDrv/Private/OpenGL4.cpp
+++ b/Engine/Source/Runtime/OpenGLDrv/Private/OpenGL4.cpp
@@ -10,7 +10,6 @@
 
 #if OPENGL_GL4
 
-GLint FOpenGL4::MaxComputeTextureImageUnits = -1;
 GLint FOpenGL4::MaxComputeUniformComponents = -1;
 GLint FOpenGL4::MaxCombinedUAVUnits = 0;
 GLint FOpenGL4::MaxComputeUAVUnits = -1;
@@ -24,7 +23,6 @@ void FOpenGL4::ProcessQueryGLInt()
 	GET_GL_INT(GL_MAX_FRAGMENT_IMAGE_UNIFORMS, 0, MaxPixelUAVUnits);
 	GET_GL_INT(GL_TEXTURE_BUFFER_OFFSET_ALIGNMENT, 0, TextureBufferAlignment);
 
-	GET_GL_INT(GL_MAX_COMPUTE_TEXTURE_IMAGE_UNITS, 0, MaxComputeTextureImageUnits);
 	GET_GL_INT(GL_MAX_COMPUTE_UNIFORM_COMPONENTS, 0, MaxComputeUniformComponents);
 	GET_GL_INT(GL_MAX_COMPUTE_IMAGE_UNIFORMS, 0, MaxComputeUAVUnits);
 	
diff --git a/Engine/Source/Runtime/OpenGLDrv/Private/OpenGLDrv.cpp b/Engine/Source/Runtime/OpenGLDrv/Private/OpenGLDrv.cpp
--- a/Engine/Source/Runtime/OpenGLDrv/Private/OpenGLDrv.cpp
+++ b/Engine/Source/Runtime/OpenGLDrv/Private/OpenGLDrv.cpp
@@ -485,6 +485,7 @@ void FOpenGLDynamicRHI::InitializeStateResources()
 
 GLint FOpenGLBase::MaxTextureImageUnits = -1;
 GLint FOpenGLBase::MaxCombinedTextureImageUnits = -1;
+GLint FOpenGLBase::MaxComputeTextureImageUnits = -1;
 GLint FOpenGLBase::MaxVertexTextureImageUnits = -1;
 GLint FOpenGLBase::MaxGeometryTextureImageUnits = -1;
 GLint FOpenGLBase::MaxVaryingVectors = -1;
@@ -504,8 +505,9 @@ bool  FOpenGLBase::bAmdWorkaround = false;
 
 void FOpenGLBase::ProcessQueryGLInt()
 {
-	GET_GL_INT(GL_MAX_TEXTURE_IMAGE_UNITS, 0, MaxTextureImageUnits);
-	GET_GL_INT(GL_MAX_VERTEX_TEXTURE_IMAGE_UNITS, 0, MaxVertexTextureImageUnits);
+	LOG_AND_GET_GL_INT(GL_MAX_TEXTURE_IMAGE_UNITS, 0, MaxTextureImageUnits);
+	LOG_AND_GET_GL_INT(GL_MAX_VERTEX_TEXTURE_IMAGE_UNITS, 0, MaxVertexTextureImageUnits);
+	LOG_AND_GET_GL_INT(GL_MAX_COMPUTE_TEXTURE_IMAGE_UNITS, 0, MaxComputeTextureImageUnits);
 	GET_GL_INT(GL_MAX_COMBINED_TEXTURE_IMAGE_UNITS, 0, MaxCombinedTextureImageUnits);
 }
 
@@ -513,27 +515,29 @@ void FOpenGLBase::ProcessExtensions( const FString& ExtensionsString )
 {
 	ProcessQueryGLInt();
 
-	// For now, just allocate additional units if available and advertise no tessellation units for HW that can't handle more
-	if ( MaxCombinedTextureImageUnits < 48 )
+	auto CheckAndSetImageUnits = [](GLint& StageImageUnitsINOUT, GLint Limit, const TCHAR* Msg) 
 	{
-		// To work around AMD driver limitation of 32 GL_MAX_COMBINED_TEXTURE_IMAGE_UNITS,
-		// Going to hard code this for now (16 units in PS, 8 units in VS, 8 units in GS).
-		// This is going to be a problem for tessellation.
-		MaxTextureImageUnits = MaxTextureImageUnits > 16 ? 16 : MaxTextureImageUnits;
-		MaxVertexTextureImageUnits = MaxVertexTextureImageUnits > 8 ? 8 : MaxVertexTextureImageUnits;
-		MaxGeometryTextureImageUnits = MaxGeometryTextureImageUnits > 8 ? 8 : MaxGeometryTextureImageUnits;
-		MaxCombinedTextureImageUnits = MaxCombinedTextureImageUnits > 32 ? 32 : MaxCombinedTextureImageUnits;
+		const bool bUnsupported = StageImageUnitsINOUT < Limit;
+		UE_CLOG(bUnsupported, LogRHI, Error, TEXT("GL RHI requires a minimum %s texture unit count of %d, this device reports %d."), Msg, Limit, StageImageUnitsINOUT);
+		check(!bUnsupported);
+		StageImageUnitsINOUT = Limit;
+	};
+
+	static const GLint GLESMaxImageUnitsPerStage = 16; // gles 3 spec is a minimum of 16 per stage. 
+	static const GLint MaxCombinedImageUnits = 48;
+
+	if (IsMobilePlatform(GMaxRHIShaderPlatform))
+	{
+		// clamp things to the levels that the spec is expecting, check the minimum is supported.
+		CheckAndSetImageUnits(MaxTextureImageUnits, GLESMaxImageUnitsPerStage, TEXT("pixel stage"));
+		CheckAndSetImageUnits(MaxVertexTextureImageUnits, GLESMaxImageUnitsPerStage, TEXT("vertex stage"));
+		CheckAndSetImageUnits(MaxGeometryTextureImageUnits, 0, TEXT("geometry stage")); // gles is not expecting this.
+		CheckAndSetImageUnits(MaxComputeTextureImageUnits, GLESMaxImageUnitsPerStage, TEXT("compute stage"));
+		CheckAndSetImageUnits(MaxCombinedTextureImageUnits, MaxCombinedImageUnits, TEXT("combined"));
 	}
 	else
 	{
-		// clamp things to the levels that the other path is going, but allow additional units for tessellation
-		if (IsMobilePlatform(GMaxRHIShaderPlatform))
-		{
-			MaxTextureImageUnits = MaxTextureImageUnits > 16 ? 16 : MaxTextureImageUnits;
-			MaxVertexTextureImageUnits = MaxVertexTextureImageUnits > 16 ? 16 : MaxVertexTextureImageUnits;
-			MaxGeometryTextureImageUnits = MaxGeometryTextureImageUnits > 16 ? 16 : MaxGeometryTextureImageUnits;
-			MaxCombinedTextureImageUnits = MaxCombinedTextureImageUnits > 48 ? 48 : MaxCombinedTextureImageUnits;
-		}
+		UE_CLOG(MaxCombinedTextureImageUnits<MaxCombinedImageUnits, LogRHI, Fatal, TEXT("GL RHI requires a minimum combined texture unit count of %d, this device reports %d."), MaxCombinedImageUnits, MaxCombinedTextureImageUnits);
 	}
 
 	// Check for support for advanced texture compression (desktop and mobile)
diff --git a/Engine/Source/Runtime/OpenGLDrv/Private/OpenGLDrvPrivate.h b/Engine/Source/Runtime/OpenGLDrv/Private/OpenGLDrvPrivate.h
--- a/Engine/Source/Runtime/OpenGLDrv/Private/OpenGLDrvPrivate.h
+++ b/Engine/Source/Runtime/OpenGLDrv/Private/OpenGLDrvPrivate.h
@@ -16,8 +16,20 @@
 
 #define GL_CHECK(x)		x; do { GLint Err = glGetError(); if (Err != 0) {FPlatformMisc::LowLevelOutputDebugStringf(TEXT("(%s:%d) GL_CHECK Failed '%s'! %d (%x)\n"), ANSI_TO_TCHAR(__FILE__), __LINE__, ANSI_TO_TCHAR( #x ), Err, Err); check(!Err);}} while (0)
 
-#if !defined(__GNUC__) && !defined(__clang__)
-	#define LOG_AND_GET_GL_INT(IntEnum, Default, Dest) \
+// 
+// #if !defined(__GNUC__) && !defined(__clang__) && !defined(__INTEL_COMPILER) && !(defined(_MSC_VER) && _MSC_VER >= 1900)
+// #define LOG_AND_GET_GL_INT_TEMP(IntEnum,Default) 
+// GLint Value_##IntEnum = Default; if (IntEnum) {glGetIntegerv(IntEnum, &Value_##IntEnum); glGetError();} else {Value_##IntEnum = Default;} 
+// UE_LOG(LogRHI, Log, TEXT("  ") ## TEXT(#IntEnum) ## TEXT(": %d"), Value_##IntEnum)
+// #else
+// #define LOG_AND_GET_GL_INT_TEMP(IntEnum,Default) GLint Value_##IntEnum = Default; if (IntEnum) {glGetIntegerv(IntEnum, &Value_##IntEnum); glGetError();} else {Value_##IntEnum = Default;} 
+// UE_LOG(LogRHI, Log, TEXT("  " #IntEnum ": %d"), Value_##IntEnum)
+// #endif
+
+
+
+#if !defined(__GNUC__) && !defined(__clang__) && !defined(__INTEL_COMPILER) && !(defined(_MSC_VER) && _MSC_VER >= 1900)
+#define LOG_AND_GET_GL_INT(IntEnum, Default, Dest) \
 		do \
 		{ \
 			Dest = Default; \
diff --git a/Engine/Source/Runtime/OpenGLDrv/Private/OpenGLES.cpp b/Engine/Source/Runtime/OpenGLDrv/Private/OpenGLES.cpp
--- a/Engine/Source/Runtime/OpenGLDrv/Private/OpenGLES.cpp
+++ b/Engine/Source/Runtime/OpenGLDrv/Private/OpenGLES.cpp
@@ -134,7 +134,6 @@ bool FOpenGLES::bHasHardwareHiddenSurfaceRemoval = false;
 bool FOpenGLES::bSupportsMobileMultiView = false;
 GLint FOpenGLES::MaxMSAASamplesTileMem = 1;
 
-GLint FOpenGLES::MaxComputeTextureImageUnits = -1;
 GLint FOpenGLES::MaxComputeUniformComponents = -1;
 
 GLint FOpenGLES::MaxComputeUAVUnits = -1;
@@ -170,8 +169,7 @@ void FOpenGLES::ProcessQueryGLInt()
 	LOG_AND_GET_GL_INT(GL_MAX_FRAGMENT_UNIFORM_VECTORS, 0, MaxPixelUniformComponents);
 	LOG_AND_GET_GL_INT(GL_TEXTURE_BUFFER_OFFSET_ALIGNMENT, 0, TextureBufferAlignment);
 	
-	GET_GL_INT(GL_MAX_COMPUTE_TEXTURE_IMAGE_UNITS, 0, MaxComputeTextureImageUnits);
-	GET_GL_INT(GL_MAX_COMPUTE_UNIFORM_COMPONENTS, 0, MaxComputeUniformComponents);
+	LOG_AND_GET_GL_INT(GL_MAX_COMPUTE_UNIFORM_COMPONENTS, 0, MaxComputeUniformComponents);
 	LOG_AND_GET_GL_INT(GL_MAX_COMBINED_IMAGE_UNIFORMS, 0, MaxCombinedUAVUnits);
 	LOG_AND_GET_GL_INT(GL_MAX_COMPUTE_IMAGE_UNIFORMS, 0, MaxComputeUAVUnits);
 	LOG_AND_GET_GL_INT(GL_MAX_FRAGMENT_IMAGE_UNIFORMS, 0, MaxPixelUAVUnits);
diff --git a/Engine/Source/Runtime/OpenGLDrv/Private/OpenGLShaders.cpp b/Engine/Source/Runtime/OpenGLDrv/Private/OpenGLShaders.cpp
--- a/Engine/Source/Runtime/OpenGLDrv/Private/OpenGLShaders.cpp
+++ b/Engine/Source/Runtime/OpenGLDrv/Private/OpenGLShaders.cpp
@@ -1545,11 +1545,11 @@ class FGLProgramCacheLRU
 	class FEvictedGLProgram
 	{
 		FOpenGLLinkedProgram* LinkedProgram;
-
+ 
 		FORCEINLINE_DEBUGGABLE TArray<uint8>& GetProgramBinary()
-		{
+ 		{
 			return LinkedProgram->LRUInfo.CachedProgramBinary;
-		}
+ 		}
 
 	public:
 
@@ -2089,6 +2089,17 @@ void FOpenGLLinkedProgram::ConfigureShaderStage( int Stage, uint32 FirstUniformB
 		0,
 		FOpenGL::GetFirstComputeTextureUnit()
 	};
+
+	static const GLint MaxTextureUnit[CrossCompiler::NUM_SHADER_STAGES] =
+	{
+		FOpenGL::GetMaxVertexTextureImageUnits(),
+		FOpenGL::GetMaxTextureImageUnits(),
+		FOpenGL::GetMaxGeometryTextureImageUnits(),
+		0,
+		0,
+		FOpenGL::GetMaxComputeTextureImageUnits()
+	};
+
 	static const GLint FirstUAVUnit[CrossCompiler::NUM_SHADER_STAGES] =
 	{
 		FOpenGL::GetFirstVertexUAVUnit(),
@@ -2208,6 +2219,11 @@ void FOpenGLLinkedProgram::ConfigureShaderStage( int Stage, uint32 FirstUniformB
 				FOpenGL::ProgramUniform1i(StageProgram, Location, FirstTextureUnit[Stage] + SamplerIndex);
 				TextureStageNeeds[ FirstTextureUnit[Stage] + SamplerIndex ] = true;
 				MaxTextureStage = FMath::Max( MaxTextureStage, FirstTextureUnit[Stage] + SamplerIndex);
+				if (SamplerIndex >= MaxTextureUnit[Stage])
+				{
+					UE_LOG(LogShaders, Error, TEXT("%s has a shader using too many textures (idx %d, max allowed %d) at stage %d"), *Config.ProgramKey.ToString(), SamplerIndex, MaxTextureUnit[Stage]-1, Stage);
+					checkNoEntry();
+				}
 			}
 			else
 			{
diff --git a/Engine/Source/Runtime/OpenGLDrv/Private/OpenGLTexture.cpp b/Engine/Source/Runtime/OpenGLDrv/Private/OpenGLTexture.cpp
--- a/Engine/Source/Runtime/OpenGLDrv/Private/OpenGLTexture.cpp
+++ b/Engine/Source/Runtime/OpenGLDrv/Private/OpenGLTexture.cpp
@@ -1951,29 +1951,20 @@ void FOpenGLDynamicRHI::RHIUpdateTexture3D(FRHITexture3D* TextureRHI, uint32 Mip
 void FOpenGLDynamicRHI::InvalidateTextureResourceInCache(GLuint Resource)
 {
 	VERIFY_GL_SCOPE();
-	if (SharedContextState.Textures || RenderingContextState.Textures || PendingState.Textures)
+	auto InvalidateContextTextureResources = [&Resource](TArray<FTextureStage>& Textures)
 	{
-		for (int32 SamplerIndex = 0; SamplerIndex < FOpenGL::GetMaxCombinedTextureImageUnits(); ++SamplerIndex)
+		for (FTextureStage& TextureStage : Textures)
 		{
-			if (SharedContextState.Textures && SharedContextState.Textures[SamplerIndex].Resource == Resource)
+			if (TextureStage.Resource == Resource)
 			{
-				SharedContextState.Textures[SamplerIndex].Target = GL_NONE;
-				SharedContextState.Textures[SamplerIndex].Resource = 0;
-			}
-
-			if (RenderingContextState.Textures && RenderingContextState.Textures[SamplerIndex].Resource == Resource)
-			{
-				RenderingContextState.Textures[SamplerIndex].Target = GL_NONE;
-				RenderingContextState.Textures[SamplerIndex].Resource = 0;
-			}
-
-			if (PendingState.Textures && PendingState.Textures[SamplerIndex].Resource == Resource)
-			{
-				PendingState.Textures[SamplerIndex].Target = GL_NONE;
-				PendingState.Textures[SamplerIndex].Resource = 0;
+				TextureStage.Target = GL_NONE;
+				TextureStage.Resource = 0;
 			}
 		}
-	}
+	};
+	InvalidateContextTextureResources(SharedContextState.Textures);
+	InvalidateContextTextureResources(RenderingContextState.Textures);
+	InvalidateContextTextureResources(PendingState.Textures);
 	
 	TextureMipLimits.Remove(Resource);
 	
diff --git a/Engine/Source/Runtime/OpenGLDrv/Public/OpenGL.h b/Engine/Source/Runtime/OpenGLDrv/Public/OpenGL.h
--- a/Engine/Source/Runtime/OpenGLDrv/Public/OpenGL.h
+++ b/Engine/Source/Runtime/OpenGLDrv/Public/OpenGL.h
@@ -143,7 +143,7 @@ public:
 	static FORCEINLINE GLint GetMaxTextureImageUnits()			{ check(MaxTextureImageUnits != -1); return MaxTextureImageUnits; }
 	static FORCEINLINE GLint GetMaxVertexTextureImageUnits()	{ check(MaxVertexTextureImageUnits != -1); return MaxVertexTextureImageUnits; }
 	static FORCEINLINE GLint GetMaxGeometryTextureImageUnits()	{ check(MaxGeometryTextureImageUnits != -1); return MaxGeometryTextureImageUnits; }
-	static FORCEINLINE GLint GetMaxComputeTextureImageUnits()	{ return 0; }
+	static FORCEINLINE GLint GetMaxComputeTextureImageUnits()	{ check(MaxComputeTextureImageUnits != -1); return MaxComputeTextureImageUnits; }
 	static FORCEINLINE GLint GetMaxCombinedTextureImageUnits()	{ check(MaxCombinedTextureImageUnits != -1); return MaxCombinedTextureImageUnits; }
 	static FORCEINLINE GLint GetTextureBufferAlignment()		{ return TextureBufferAlignment; }
 
@@ -370,6 +370,7 @@ public:
 protected:
 	static GLint MaxTextureImageUnits;
 	static GLint MaxCombinedTextureImageUnits;
+	static GLint MaxComputeTextureImageUnits;
 	static GLint MaxVertexTextureImageUnits;
 	static GLint MaxGeometryTextureImageUnits;
 	static GLint MaxVertexUniformComponents;
diff --git a/Engine/Source/Runtime/OpenGLDrv/Public/OpenGL4.h b/Engine/Source/Runtime/OpenGLDrv/Public/OpenGL4.h
--- a/Engine/Source/Runtime/OpenGLDrv/Public/OpenGL4.h
+++ b/Engine/Source/Runtime/OpenGLDrv/Public/OpenGL4.h
@@ -99,7 +99,6 @@ struct FOpenGL4 : public FOpenGL3
 	static void ProcessQueryGLInt();
 	static void ProcessExtensions( const FString& ExtensionsString );
 
-	static FORCEINLINE GLint GetMaxComputeTextureImageUnits() { check(MaxComputeTextureImageUnits != -1); return MaxComputeTextureImageUnits; }
 	static FORCEINLINE GLint GetMaxComputeUniformComponents() { check(MaxComputeUniformComponents != -1); return MaxComputeUniformComponents; }
 
 	static FORCEINLINE GLint GetFirstComputeUAVUnit()			{ return 0; }
@@ -114,7 +113,6 @@ struct FOpenGL4 : public FOpenGL3
 	static FORCEINLINE int32 GetReadHalfFloatPixelsEnum() { return GL_HALF_FLOAT; }
 
 protected:
-	static GLint MaxComputeTextureImageUnits;
 	static GLint MaxComputeUniformComponents;
 
 	static GLint MaxCombinedUAVUnits;
diff --git a/Engine/Source/Runtime/OpenGLDrv/Public/OpenGLES.h b/Engine/Source/Runtime/OpenGLDrv/Public/OpenGLES.h
--- a/Engine/Source/Runtime/OpenGLDrv/Public/OpenGLES.h
+++ b/Engine/Source/Runtime/OpenGLDrv/Public/OpenGLES.h
@@ -859,7 +859,6 @@ public:
 	/** Whether device supports mobile multi-view */
 	static bool bSupportsMobileMultiView;
 
-	static GLint MaxComputeTextureImageUnits;
 	static GLint MaxComputeUniformComponents;
 
 	static GLint MaxCombinedUAVUnits;
diff --git a/Engine/Source/Runtime/OpenGLDrv/Public/OpenGLState.h b/Engine/Source/Runtime/OpenGLDrv/Public/OpenGLState.h
--- a/Engine/Source/Runtime/OpenGLDrv/Public/OpenGLState.h
+++ b/Engine/Source/Runtime/OpenGLDrv/Public/OpenGLState.h
@@ -271,13 +271,11 @@ struct FOpenGLStream
 
 struct FOpenGLCommonState
 {
-	FTextureStage*			Textures;
-	FOpenGLSamplerState**	SamplerStates;
+	TArray<FTextureStage>	Textures;
+	TArray<FOpenGLSamplerState*>	SamplerStates;
 	TArray<FUAVStage>		UAVs;
 
 	FOpenGLCommonState()
-	: Textures(NULL)
-	, SamplerStates(NULL)
 	{}
 
 	virtual ~FOpenGLCommonState()
@@ -291,10 +289,9 @@ struct FOpenGLCommonState
 	{
 		check(NumCombinedTextures >= FOpenGL::GetMaxCombinedTextureImageUnits());
 		check(NumCombinedUAVUnits >= FOpenGL::GetMaxCombinedUAVUnits());
-		check(!Textures && !SamplerStates && UAVs.Num() == 0);
-		Textures = new FTextureStage[NumCombinedTextures];
-		SamplerStates = new FOpenGLSamplerState*[NumCombinedTextures];
-		FMemory::Memset( SamplerStates, 0, NumCombinedTextures * sizeof(*SamplerStates) );
+		check(Textures.IsEmpty() && SamplerStates.IsEmpty() && UAVs.Num() == 0);
+		Textures.SetNum(NumCombinedTextures);
+		SamplerStates.SetNumZeroed(NumCombinedTextures);
 		
 		UAVs.Reserve(NumCombinedUAVUnits);
 		UAVs.AddDefaulted(NumCombinedUAVUnits);
@@ -302,12 +299,9 @@ struct FOpenGLCommonState
 
 	virtual void CleanupResources()
 	{
-		delete [] SamplerStates;
-		delete [] Textures;
-
+		SamplerStates.Empty();
+		Textures.Empty();
 		UAVs.Empty();
-		SamplerStates = NULL;
-		Textures = NULL;
 	}
 };
 
