Create DummyFrameBuffer in rendering context to make renderdoc work properly.

diff --git a/Engine/Source/Runtime/OpenGLDrv/Private/Android/AndroidEGL.cpp b/Engine/Source/Runtime/OpenGLDrv/Private/Android/AndroidEGL.cpp
--- a/Engine/Source/Runtime/OpenGLDrv/Private/Android/AndroidEGL.cpp
+++ b/Engine/Source/Runtime/OpenGLDrv/Private/Android/AndroidEGL.cpp
@@ -640,12 +640,6 @@ void AndroidEGL::InitBackBuffer()
 	PImplData->RenderingContext.ViewportFramebuffer = GetResolveFrameBuffer();
 	PImplData->SharedContext.ViewportFramebuffer = GetResolveFrameBuffer();
 	PImplData->SingleThreadedContext.ViewportFramebuffer = GetResolveFrameBuffer();
-	
-	// Dummy FBO we bind right after SwapBuffers to tell driver that backbuffer is no longer in use by the App
-	glGenFramebuffers(1, &PImplData->DummyFrameBuffer);
-	PImplData->RenderingContext.DummyFrameBuffer = PImplData->DummyFrameBuffer;
-	PImplData->SharedContext.DummyFrameBuffer = PImplData->DummyFrameBuffer;
-	PImplData->SingleThreadedContext.DummyFrameBuffer = PImplData->DummyFrameBuffer;
 }
 
 extern void AndroidThunkCpp_SetDesiredViewSize(int32 Width, int32 Height);
@@ -912,6 +906,15 @@ void AndroidEGL::SetCurrentSharedContext()
 void AndroidEGL::AcquireCurrentRenderingContext()
 {
 	SetCurrentRenderingContext();
+
+	if (!PImplData->DummyFrameBuffer)
+	{
+		// Dummy FBO we bind right after SwapBuffers to tell driver that backbuffer is no longer in use by the App
+		glGenFramebuffers(1, &PImplData->DummyFrameBuffer);
+		PImplData->RenderingContext.DummyFrameBuffer = PImplData->DummyFrameBuffer;
+		PImplData->SharedContext.DummyFrameBuffer = PImplData->DummyFrameBuffer;
+		PImplData->SingleThreadedContext.DummyFrameBuffer = PImplData->DummyFrameBuffer;
+	}
 }
 
 void AndroidEGL::SetCurrentRenderingContext()
