OpenGL: Discard transient color targets at the end of render-pass

diff --git a/Engine/Source/Runtime/OpenGLDrv/Private/OpenGLRenderTarget.cpp b/Engine/Source/Runtime/OpenGLDrv/Private/OpenGLRenderTarget.cpp
--- a/Engine/Source/Runtime/OpenGLDrv/Private/OpenGLRenderTarget.cpp
+++ b/Engine/Source/Runtime/OpenGLDrv/Private/OpenGLRenderTarget.cpp
@@ -1164,21 +1164,35 @@ void FOpenGLDynamicRHI::RHIEndRenderPass()
 		EndOcclusionQueryBatch();
 	}
 
-	// Drop depth and stencil to avoid export
-	if (RenderPassInfo.DepthStencilRenderTarget.DepthStencilTarget)
+	// Discard transient color targets
+	uint32 ColorMask = 0u;
+	for (int32 ColorIndex = 0; ColorIndex < MaxSimultaneousRenderTargets; ++ColorIndex)
 	{
-		ERenderTargetActions DepthActions = GetDepthActions(RenderPassInfo.DepthStencilRenderTarget.Action);
-		ERenderTargetActions StencilActions = GetStencilActions(RenderPassInfo.DepthStencilRenderTarget.Action);
-		const bool bIsUsingStencil = RenderPassInfo.DepthStencilRenderTarget.ExclusiveDepthStencil.IsUsingStencil();
+		const FRHIRenderPassInfo::FColorEntry& Entry = RenderPassInfo.ColorRenderTargets[ColorIndex];
+		if (!Entry.RenderTarget)
+		{
+			break;
+		}
 
-		bool bDiscardDepth = GetStoreAction(DepthActions) == ERenderTargetStoreAction::ENoAction;
-		bool bDiscardStencil = bIsUsingStencil && GetStoreAction(StencilActions) == ERenderTargetStoreAction::ENoAction;
-		if (bDiscardDepth || bDiscardStencil)
+		if (GetStoreAction(Entry.Action) == ERenderTargetStoreAction::ENoAction)
 		{
-			RHIDiscardRenderTargets(bDiscardDepth, bDiscardStencil, 0);
+			ColorMask |= (1u << ColorIndex);
 		}
 	}
 
+	// Discard transient DepthStencil
+	bool bDiscardDepthStencil = false;
+	if (RenderPassInfo.DepthStencilRenderTarget.DepthStencilTarget)
+	{
+		ERenderTargetActions DepthActions = GetDepthActions(RenderPassInfo.DepthStencilRenderTarget.Action);
+		bDiscardDepthStencil = GetStoreAction(DepthActions) == ERenderTargetStoreAction::ENoAction;
+	}
+
+	if (bDiscardDepthStencil || ColorMask != 0)
+	{
+		RHIDiscardRenderTargets(bDiscardDepthStencil, bDiscardDepthStencil, ColorMask);
+	}
+
 	FRHIRenderTargetView RTV(nullptr, ERenderTargetLoadAction::ENoAction);
 	FRHIDepthRenderTargetView DepthRTV(nullptr, ERenderTargetLoadAction::ENoAction, ERenderTargetStoreAction::ENoAction);
 	SetRenderTargets(1, &RTV, &DepthRTV);
@@ -1196,7 +1210,8 @@ void FOpenGLDynamicRHI::RHINextSubpass()
 {
 	IRHICommandContext::RHINextSubpass();
 	
-	if (RenderPassInfo.SubpassHint == ESubpassHint::DepthReadSubpass)
+	if (RenderPassInfo.SubpassHint == ESubpassHint::DepthReadSubpass ||
+		RenderPassInfo.SubpassHint == ESubpassHint::DeferredShadingSubpass)
 	{
 		FOpenGL::FrameBufferFetchBarrier();
 	}
