Fix the flip screen issue on GLES mali devices.

diff --git a/Engine/Source/Runtime/OpenGLDrv/Private/Android/AndroidEGL.cpp b/Engine/Source/Runtime/OpenGLDrv/Private/Android/AndroidEGL.cpp
--- a/Engine/Source/Runtime/OpenGLDrv/Private/Android/AndroidEGL.cpp
+++ b/Engine/Source/Runtime/OpenGLDrv/Private/Android/AndroidEGL.cpp
@@ -31,6 +31,14 @@ DEFINE_LOG_CATEGORY(LogEGL);
 typedef int32(*PFN_ANativeWindow_setBuffersTransform)(struct ANativeWindow* window, int32 transform);
 static PFN_ANativeWindow_setBuffersTransform ANativeWindow_setBuffersTransform_API = nullptr;
 
+static TAutoConsoleVariable<int32> CVarAndroidGLESFlipYMethod(
+	TEXT("r.Android.GLESFlipYMethod"),
+	0,
+	TEXT(" 0: Flip Y method detected automatically by GPU vendor.\n"
+		 " 1: Force flip Y by native window setBuffersTransform.\n"
+		 " 2: Force flip Y by BlitFrameBuffer."),
+	ECVF_RenderThreadSafe);
+
 
 const  int EGLMinRedBits		= 5;
 const  int EGLMinGreenBits		= 6;
@@ -160,7 +168,7 @@ AndroidEGL::AndroidEGL()
 	{
 		ANativeWindow_setBuffersTransform_API = reinterpret_cast<PFN_ANativeWindow_setBuffersTransform>(dlsym(LibNativeWindow, "ANativeWindow_setBuffersTransform"));
 	}
-	UE_LOG(LogAndroid, Log, TEXT("ANativeWindow_setBuffersTransform is %s on this device"), ANativeWindow_setBuffersTransform_API == nullptr ? TEXT("not supported") : TEXT("supported"));
+	FPlatformMisc::LowLevelOutputDebugStringf(TEXT("ANativeWindow_setBuffersTransform is %s on this device"), ANativeWindow_setBuffersTransform_API == nullptr ? TEXT("not supported") : TEXT("supported"));
 }
 
 void AndroidEGL::ResetDisplay()
@@ -914,7 +922,7 @@ void AndroidEGL::AcquireCurrentRenderingContext()
 		PImplData->SingleThreadedContext.DummyFrameBuffer = PImplData->DummyFrameBuffer;
 	}
 
-	if (FAndroidMisc::SupportsBackbufferSampling())
+	if (IsOfflineSurfaceRequired())
 	{
 		// Needs to be generated on rendering context
 		if (!PImplData->ResolveFrameBuffer)
@@ -1136,7 +1144,7 @@ void AndroidEGL::LogConfigInfo(EGLConfig  EGLConfigInfo)
 
 void AndroidEGL::UpdateBuffersTransform()
 {
-	if (ANativeWindow_setBuffersTransform_API != nullptr)
+	if (ANativeWindow_setBuffersTransform_API != nullptr && !IsOfflineSurfaceRequired())
 	{
 		int32 BufferTransform = ANATIVEWINDOW_TRANSFORM_IDENTITY;
 
@@ -1168,6 +1176,14 @@ void AndroidEGL::UpdateBuffersTransform()
 	}
 }
 
+bool AndroidEGL::IsOfflineSurfaceRequired()
+{
+	return FAndroidMisc::SupportsBackbufferSampling() ||
+		// setBuffersTransform doesn't work on arm GPU devices, uses BlitFrameBuffer instead
+		(CVarAndroidGLESFlipYMethod.GetValueOnAnyThread() == 0 && GRHIVendorId == 0x13B5) ||
+		CVarAndroidGLESFlipYMethod.GetValueOnAnyThread() == 2;
+}
+
 ///
 extern FCriticalSection GAndroidWindowLock;
 
diff --git a/Engine/Source/Runtime/OpenGLDrv/Private/Android/AndroidEGL.h b/Engine/Source/Runtime/OpenGLDrv/Private/Android/AndroidEGL.h
--- a/Engine/Source/Runtime/OpenGLDrv/Private/Android/AndroidEGL.h
+++ b/Engine/Source/Runtime/OpenGLDrv/Private/Android/AndroidEGL.h
@@ -82,6 +82,7 @@ public:
 	void InitRenderSurface(bool bUseSmallSurface, bool bCreateWndSurface);
 	void InitSharedSurface(bool bUseSmallSurface);
 	void UpdateBuffersTransform();
+	bool IsOfflineSurfaceRequired();
 
 	void GetDimensions(uint32& OutWidth, uint32& OutHeight);
 	
diff --git a/Engine/Source/Runtime/OpenGLDrv/Private/Android/AndroidOpenGL.cpp b/Engine/Source/Runtime/OpenGLDrv/Private/Android/AndroidOpenGL.cpp
--- a/Engine/Source/Runtime/OpenGLDrv/Private/Android/AndroidOpenGL.cpp
+++ b/Engine/Source/Runtime/OpenGLDrv/Private/Android/AndroidOpenGL.cpp
@@ -225,8 +225,7 @@ bool PlatformBlitToViewport(FPlatformOpenGLDevice* Device, const FOpenGLViewport
 	
 	FPlatformOpenGLContext* const Context = Viewport.GetGLContext();
 
-	if (bPresent && FPlatformMisc::SupportsBackbufferSampling())
-
+	if (bPresent && AndroidEGL::GetInstance()->IsOfflineSurfaceRequired())
 	{
 		if (Device->TargetDirty)
 		{
@@ -247,7 +246,7 @@ bool PlatformBlitToViewport(FPlatformOpenGLDevice* Device, const FOpenGLViewport
 
 			FOpenGL::BlitFramebuffer(
 				0, 0, BackbufferSizeX, BackbufferSizeY,
-				0, 0, BackbufferSizeX, BackbufferSizeY,
+				0, BackbufferSizeY, BackbufferSizeX, 0,
 				GL_COLOR_BUFFER_BIT,
 				GL_NEAREST
 			);
@@ -455,7 +454,7 @@ FOpenGLTexture* PlatformCreateBuiltinBackBuffer(FOpenGLDynamicRHI* OpenGLRHI, ui
 	check(IsInRenderingThread());
 	// Create the built-in back buffer if we disable backbuffer sampling.
 	// Otherwise return null and we will create an off-screen surface afterward.
-	if (!FPlatformMisc::SupportsBackbufferSampling())
+	if (!AndroidEGL::GetInstance()->IsOfflineSurfaceRequired())
 	{
 		const FOpenGLTextureCreateDesc CreateDesc =
 			FRHITextureCreateDesc::Create2D(TEXT("PlatformCreateBuiltinBackBuffer"), SizeX, SizeY, PF_B8G8R8A8)
@@ -478,7 +477,7 @@ void PlatformResizeGLContext( FPlatformOpenGLDevice* Device, FPlatformOpenGLCont
 	Context->BackBufferResource = BackBufferResource;
 	Context->BackBufferTarget = BackBufferTarget;
 
-	if (FPlatformMisc::SupportsBackbufferSampling())
+	if (AndroidEGL::GetInstance()->IsOfflineSurfaceRequired())
 	{
 		Device->TargetDirty = true;
 		glBindFramebuffer(GL_FRAMEBUFFER, Context->ViewportFramebuffer);
