OpenGL: Fixing an ancient out of bounds memory read bug

diff --git a/Engine/Source/Runtime/OpenGLDrv/Public/OpenGLShaderResources.h b/Engine/Source/Runtime/OpenGLDrv/Public/OpenGLShaderResources.h
--- a/Engine/Source/Runtime/OpenGLDrv/Public/OpenGLShaderResources.h
+++ b/Engine/Source/Runtime/OpenGLDrv/Public/OpenGLShaderResources.h
@@ -138,13 +138,14 @@ struct FOpenGLShaderBindings
 
 		bEqual &= FMemory::Memcmp(A.PackedGlobalArrays.GetData(),B.PackedGlobalArrays.GetData(),A.PackedGlobalArrays.GetTypeSize()*A.PackedGlobalArrays.Num()) == 0; 
 
-		for (int32 Item = 0; Item < A.PackedUniformBuffers.Num(); Item++)
+		for (int32 Item = 0; bEqual && Item < A.PackedUniformBuffers.Num(); Item++)
 		{
-			const TArray<CrossCompiler::FPackedArrayInfo> &ArrayA = A.PackedUniformBuffers[Item];
-			const TArray<CrossCompiler::FPackedArrayInfo> &ArrayB = B.PackedUniformBuffers[Item];
+			const TArray<CrossCompiler::FPackedArrayInfo>& ArrayA = A.PackedUniformBuffers[Item];
+			const TArray<CrossCompiler::FPackedArrayInfo>& ArrayB = B.PackedUniformBuffers[Item];
 
-			bEqual &= FMemory::Memcmp(ArrayA.GetData(),ArrayB.GetData(),ArrayA.GetTypeSize()*ArrayA.Num()) == 0;
+			bEqual = bEqual && (ArrayA.Num() == ArrayB.Num()) && (FMemory::Memcmp(ArrayA.GetData(), ArrayB.GetData(), ArrayA.GetTypeSize() * ArrayA.Num()) == 0);
 		}
+
 		
 		for (int32 Item = 0; bEqual && Item < A.InputVaryings.Num(); Item++)
 		{
