Android: Prevent hang during pause by flushing GT before blocking RT

During Android Pause, the Event Thread (ET) will block the Rendering Thread (RT), followed by waiting for a Game Thread (GT) command to complete.
With Vulkan enabled, there could be a scenario where the game thread (GT)  is waiting on the RT to flush. If the GT is waiting on RT while the ET blocks the RT, the ET will then deadlock waiting for GT.

Instead, the ET can wait on the GT to finish first, with an unblocked RT.

diff --git a/Engine/Source/Runtime/OpenGLDrv/Private/Android/AndroidEGL.cpp b/Engine/Source/Runtime/OpenGLDrv/Private/Android/AndroidEGL.cpp
--- a/Engine/Source/Runtime/OpenGLDrv/Private/Android/AndroidEGL.cpp
+++ b/Engine/Source/Runtime/OpenGLDrv/Private/Android/AndroidEGL.cpp
@@ -1263,6 +1263,15 @@ void BlockRendering()
 		FPlatformProcess::ReturnSynchEventToPool(EventToDelete);
 	});
 
+	// Flush GT first in case it has any dependency on RT work to complete
+	FGraphEventRef GTBlockTask = FFunctionGraphTask::CreateAndDispatchWhenReady([BlockedTrigger]()
+		{
+			SetSharedContextGameCommand(BlockedTrigger);
+		}, TStatId(), NULL, ENamedThreads::GameThread);
+
+	UE_LOG(LogAndroid, Log, TEXT("Waiting for game thread to release EGL context/surface."));
+	BlockedTrigger->Wait();
+
 	FGraphEventRef RTBlockTask = FFunctionGraphTask::CreateAndDispatchWhenReady([BlockedTrigger]()
 	{
 		BlockOnLostWindowRenderCommand(BlockedTrigger);
@@ -1271,14 +1280,6 @@ void BlockRendering()
 	// wait for the render thread to process.
 	UE_LOG(LogAndroid, Log, TEXT("Waiting for renderer to encounter blocking command."));
 	BlockedTrigger->Wait();
-
-	FGraphEventRef GTBlockTask = FFunctionGraphTask::CreateAndDispatchWhenReady([BlockedTrigger]()
-	{
-		SetSharedContextGameCommand(BlockedTrigger);
-	}, TStatId(), NULL, ENamedThreads::GameThread);
-
-	UE_LOG(LogAndroid, Log, TEXT("Waiting for game thread to encounter blocking command."));
-	BlockedTrigger->Wait();
 }
 
 #endif
