Fix a crash when launching mobile opengl preview.

diff --git a/Engine/Source/Runtime/OpenGLDrv/Private/Android/AndroidOpenGL.cpp b/Engine/Source/Runtime/OpenGLDrv/Private/Android/AndroidOpenGL.cpp
--- a/Engine/Source/Runtime/OpenGLDrv/Private/Android/AndroidOpenGL.cpp
+++ b/Engine/Source/Runtime/OpenGLDrv/Private/Android/AndroidOpenGL.cpp
@@ -362,6 +362,21 @@ bool PlatformInitOpenGL()
 				FAndroidMisc::MessageBoxExt(EAppMsgType::Ok, *Message, TEXT("Unable to run on this device!"));
 			}
 		}
+
+		// Needs to initialize GPU vendor id before AndroidEGL::AcquireCurrentRenderingContext
+		FString& VendorName = FAndroidGPUInfo::Get().VendorName;
+		if (VendorName.Contains(TEXT("ImgTec")))
+		{
+			GRHIVendorId = 0x1010;
+		}
+		else if (VendorName.Contains(TEXT("ARM")))
+		{
+			GRHIVendorId = 0x13B5;
+		}
+		else if (VendorName.Contains(TEXT("Qualcomm")))
+		{
+			GRHIVendorId = 0x5143;
+		}
 	}
 	return true;
 }
diff --git a/Engine/Source/Runtime/OpenGLDrv/Private/Android/AndroidOpenGLPrivate.h b/Engine/Source/Runtime/OpenGLDrv/Private/Android/AndroidOpenGLPrivate.h
--- a/Engine/Source/Runtime/OpenGLDrv/Private/Android/AndroidOpenGLPrivate.h
+++ b/Engine/Source/Runtime/OpenGLDrv/Private/Android/AndroidOpenGLPrivate.h
@@ -25,6 +25,7 @@ public:
 
 	FString GPUFamily;
 	FString GLVersion;
+	FString VendorName;
 	bool bSupportsFloatingPointRenderTargets;
 	bool bSupportsFrameBufferFetch;
 	TArray<FString> TargetPlatformNames;
@@ -89,5 +90,7 @@ private:
 			|| ExtensionsString.Contains(TEXT("GL_ARM_shader_framebuffer_fetch ")); // has space at the end to exclude GL_ARM_shader_framebuffer_fetch_depth_stencil match
 
 		GAndroidGPUInfoReady = true;
+
+		VendorName = FString(ANSI_TO_TCHAR((const ANSICHAR*)glGetString(GL_VENDOR)));
 	}
 };
