Flip Y with offline surface on PowerVR devices and Adreno 4xx devices because the setBuffersTransform is not working on these devices.

diff --git a/Engine/Config/BaseDeviceProfiles.ini b/Engine/Config/BaseDeviceProfiles.ini
--- a/Engine/Config/BaseDeviceProfiles.ini
+++ b/Engine/Config/BaseDeviceProfiles.ini
@@ -724,6 +724,7 @@ BaseProfileName=Android_Mid
 [Android_Adreno4xx DeviceProfile]
 DeviceType=Android
 BaseProfileName=Android_Low
++CVars=r.Android.GLESFlipYMethod=2
 
 [Android_Adreno5xx_Low DeviceProfile]
 DeviceType=Android
diff --git a/Engine/Source/Runtime/OpenGLDrv/Private/Android/AndroidEGL.cpp b/Engine/Source/Runtime/OpenGLDrv/Private/Android/AndroidEGL.cpp
--- a/Engine/Source/Runtime/OpenGLDrv/Private/Android/AndroidEGL.cpp
+++ b/Engine/Source/Runtime/OpenGLDrv/Private/Android/AndroidEGL.cpp
@@ -1179,8 +1179,8 @@ void AndroidEGL::UpdateBuffersTransform()
 bool AndroidEGL::IsOfflineSurfaceRequired()
 {
 	return FAndroidMisc::SupportsBackbufferSampling() ||
-		// setBuffersTransform doesn't work on arm GPU devices, uses BlitFrameBuffer instead
-		(CVarAndroidGLESFlipYMethod.GetValueOnAnyThread() == 0 && GRHIVendorId == 0x13B5) ||
+		// setBuffersTransform doesn't work on arm and powerVR GPU devices, uses BlitFrameBuffer instead
+		(CVarAndroidGLESFlipYMethod.GetValueOnAnyThread() == 0 && (GRHIVendorId == 0x13B5 || GRHIVendorId == 0x1010)) ||
 		CVarAndroidGLESFlipYMethod.GetValueOnAnyThread() == 2;
 }
 
diff --git a/Engine/Source/Runtime/OpenGLDrv/Private/Android/AndroidOpenGL.cpp b/Engine/Source/Runtime/OpenGLDrv/Private/Android/AndroidOpenGL.cpp
--- a/Engine/Source/Runtime/OpenGLDrv/Private/Android/AndroidOpenGL.cpp
+++ b/Engine/Source/Runtime/OpenGLDrv/Private/Android/AndroidOpenGL.cpp
@@ -365,7 +365,7 @@ bool PlatformInitOpenGL()
 
 		// Needs to initialize GPU vendor id before AndroidEGL::AcquireCurrentRenderingContext
 		FString& VendorName = FAndroidGPUInfo::Get().VendorName;
-		if (VendorName.Contains(TEXT("ImgTec")))
+		if (VendorName.Contains(TEXT("ImgTec")) || VendorName.Contains(TEXT("Imagination")))
 		{
 			GRHIVendorId = 0x1010;
 		}
diff --git a/Engine/Source/Runtime/OpenGLDrv/Private/OpenGLDrv.cpp b/Engine/Source/Runtime/OpenGLDrv/Private/OpenGLDrv.cpp
--- a/Engine/Source/Runtime/OpenGLDrv/Private/OpenGLDrv.cpp
+++ b/Engine/Source/Runtime/OpenGLDrv/Private/OpenGLDrv.cpp
@@ -583,7 +583,7 @@ void FOpenGLBase::ProcessExtensions( const FString& ExtensionsString )
 	{
 		GRHIVendorId = 0x10DE;
 	}
-	else if (VendorName.Contains(TEXT("ImgTec")))
+	else if (VendorName.Contains(TEXT("ImgTec")) || VendorName.Contains(TEXT("Imagination")))
 	{
 		GRHIVendorId = 0x1010;
 	}
