Fix an ensure issue of setBuffersTransform on Samsung S9.

diff --git a/Engine/Source/Runtime/OpenGLDrv/Private/Android/AndroidEGL.cpp b/Engine/Source/Runtime/OpenGLDrv/Private/Android/AndroidEGL.cpp
--- a/Engine/Source/Runtime/OpenGLDrv/Private/Android/AndroidEGL.cpp
+++ b/Engine/Source/Runtime/OpenGLDrv/Private/Android/AndroidEGL.cpp
@@ -1149,6 +1149,14 @@ void AndroidEGL::UpdateBuffersTransform()
 		int32 BufferTransform = ANATIVEWINDOW_TRANSFORM_IDENTITY;
 
 		EDeviceScreenOrientation ScreenOrientation = FPlatformMisc::GetDeviceOrientation();
+		
+		// Update the device orientation in case it hasn't been updated yet.
+		if (ScreenOrientation == EDeviceScreenOrientation::Unknown)
+		{
+			FAndroidMisc::UpdateDeviceOrientation();
+			ScreenOrientation = FPlatformMisc::GetDeviceOrientation();
+		}
+
 		switch (ScreenOrientation)
 		{
 		case EDeviceScreenOrientation::Portrait:
