Fixed: rendering issues with MSAA enabled on Android OpenGL
Make MSAA targets 'Memoryless' when possible

diff --git a/Engine/Source/Runtime/Engine/Public/SceneTexturesConfig.h b/Engine/Source/Runtime/Engine/Public/SceneTexturesConfig.h
--- a/Engine/Source/Runtime/Engine/Public/SceneTexturesConfig.h
+++ b/Engine/Source/Runtime/Engine/Public/SceneTexturesConfig.h
@@ -106,6 +106,7 @@ struct ENGINE_API FSceneTexturesConfig
 		, bKeepDepthContent{ 1 }
 		, bPreciseDepthAux{}
 		, bSamplesCustomDepthAndStencil{}
+		, bMemorylessMSAA{}
 		, bSupportsXRTargetManagerDepthAlloc{}
 	{}
 
@@ -167,6 +168,9 @@ struct ENGINE_API FSceneTexturesConfig
 
 	// (Mobile) True if both CustomDepth and CustomStencil are sampled in a shader
 	uint32 bSamplesCustomDepthAndStencil : 1;
+	
+	// (Mobile) True if MSAA targets can be memoryless
+	uint32 bMemorylessMSAA : 1;
 
 	// (XR) True if we can request an XR depth swapchain
 	uint32 bSupportsXRTargetManagerDepthAlloc : 1;
diff --git a/Engine/Source/Runtime/OpenGLDrv/Private/OpenGLCommands.cpp b/Engine/Source/Runtime/OpenGLDrv/Private/OpenGLCommands.cpp
--- a/Engine/Source/Runtime/OpenGLDrv/Private/OpenGLCommands.cpp
+++ b/Engine/Source/Runtime/OpenGLDrv/Private/OpenGLCommands.cpp
@@ -1283,7 +1283,7 @@ void FOpenGLDynamicRHI::SetPendingBlendStateForActiveRenderTargets( FOpenGLConte
 		else if (RenderTargetIndex == 0)
 		{
 			FOpenGLTexture* RenderTarget2D = PendingState.RenderTargets[RenderTargetIndex];
-			bMSAAEnabled = RenderTarget2D->GetNumSamplesRendered() > 1;
+			bMSAAEnabled = PendingState.NumRenderingSamples > 1 || RenderTarget2D->IsMultisampled();
 		}
 
 		const FOpenGLBlendStateData::FRenderTarget& RenderTargetBlendState = PendingState.BlendState.RenderTargets[RenderTargetIndex];
@@ -1536,7 +1536,7 @@ void FOpenGLDynamicRHI::SetRenderTargets(
 		return;
 	}
 
-	PendingState.Framebuffer = GetOpenGLFramebuffer(NumSimultaneousRenderTargets, PendingState.RenderTargets, PendingState.RenderTargetArrayIndex, PendingState.RenderTargetMipmapLevels, PendingState.DepthStencil);
+	PendingState.Framebuffer = GetOpenGLFramebuffer(NumSimultaneousRenderTargets, PendingState.RenderTargets, PendingState.RenderTargetArrayIndex, PendingState.RenderTargetMipmapLevels, PendingState.DepthStencil, PendingState.NumRenderingSamples);
 	PendingState.bFramebufferSetupInvalid = false;
 
 	if (PendingState.FirstNonzeroRenderTarget != -1)
diff --git a/Engine/Source/Runtime/OpenGLDrv/Private/OpenGLRenderTarget.cpp b/Engine/Source/Runtime/OpenGLDrv/Private/OpenGLRenderTarget.cpp
--- a/Engine/Source/Runtime/OpenGLDrv/Private/OpenGLRenderTarget.cpp
+++ b/Engine/Source/Runtime/OpenGLDrv/Private/OpenGLRenderTarget.cpp
@@ -43,9 +43,11 @@ public:
 		const uint32* InRenderTargetArrayIndices,
 		const uint32* InRenderTargetMipmapLevels,
 		FOpenGLTexture* InDepthStencilTarget,
+		int32 InNumRenderingSamples,
 		EOpenGLCurrentContext InContext
 		)
 		:	DepthStencilTarget(InDepthStencilTarget)
+		,	NumRenderingSamples(InNumRenderingSamples)
 		,	Context(InContext)
 	{
 		uint32 RenderTargetIndex;
@@ -71,7 +73,11 @@ public:
 	*/
 	friend bool operator ==(const FOpenGLFramebufferKey& A,const FOpenGLFramebufferKey& B)
 	{
-		return	!FMemory::Memcmp(A.RenderTargets, B.RenderTargets, sizeof(A.RenderTargets) ) && A.DepthStencilTarget == B.DepthStencilTarget && A.Context == B.Context;
+		return	
+			!FMemory::Memcmp(A.RenderTargets, B.RenderTargets, sizeof(A.RenderTargets) ) && 
+			A.DepthStencilTarget == B.DepthStencilTarget &&
+			A.NumRenderingSamples == B.NumRenderingSamples &&
+			A.Context == B.Context;
 	}
 
 	/**
@@ -81,19 +87,42 @@ public:
 	*/
 	friend uint32 GetTypeHash(const FOpenGLFramebufferKey &Key)
 	{
-		return FCrc::MemCrc_DEPRECATED(Key.RenderTargets, sizeof(Key.RenderTargets)) ^ GetTypeHash(Key.DepthStencilTarget) ^ GetTypeHash(Key.Context);
+		return FCrc::MemCrc_DEPRECATED(Key.RenderTargets, sizeof(Key.RenderTargets)) ^ GetTypeHash(Key.DepthStencilTarget) ^ GetTypeHash(Key.NumRenderingSamples) ^ GetTypeHash(Key.Context);
 	}
 
 	const FOpenGLTexture* GetRenderTarget( int32 Index ) const { return RenderTargets[Index].Texture; }
 	const FOpenGLTexture* GetDepthStencilTarget( void ) const { return DepthStencilTarget; }
+	int32 GetNumRenderingSamples( void ) const { return NumRenderingSamples; }
 
 private:
 
 	RenderTargetInfo RenderTargets[MaxSimultaneousRenderTargets];
 	FOpenGLTexture* DepthStencilTarget;
+	int32 NumRenderingSamples; // MSAA on tile
 	EOpenGLCurrentContext Context;
 };
 
+static void ConditionallyAllocateRenderbufferStorage(FOpenGLTexture& RenderTarget)
+{
+	if (RenderTarget.bMultisampleRenderbuffer && 
+		RenderTarget.GetAllocatedStorageForMip(0,0) == false)
+	{
+		check(RenderTarget.IsMultisampled());
+		check(RenderTarget.Target == GL_RENDERBUFFER);
+
+		GLuint TextureID = RenderTarget.GetRawResourceName();
+		const FRHITextureDesc& Desc = RenderTarget.GetDesc();
+		const FOpenGLTextureFormat& GLFormat = GOpenGLTextureFormats[Desc.Format];
+		const bool bSRGB = EnumHasAnyFlags(Desc.Flags, TexCreate_SRGB);
+		
+		glBindRenderbuffer(GL_RENDERBUFFER, TextureID);
+		FOpenGL::RenderbufferStorageMultisample(GL_RENDERBUFFER, Desc.NumSamples, GLFormat.InternalFormat[bSRGB], Desc.Extent.X, Desc.Extent.Y);
+		glBindRenderbuffer(GL_RENDERBUFFER, 0);
+
+		RenderTarget.SetAllocatedStorage(true);
+	}
+}
+
 typedef TMap<FOpenGLFramebufferKey,GLuint> FOpenGLFramebufferCache;
 
 /** Lazily initialized framebuffer cache singleton. */
@@ -104,12 +133,18 @@ static FOpenGLFramebufferCache& GetOpenGLFramebufferCache()
 }
 
 GLuint FOpenGLDynamicRHI::GetOpenGLFramebuffer(uint32 NumSimultaneousRenderTargets, FOpenGLTexture** RenderTargets, const uint32* ArrayIndices, const uint32* MipmapLevels, FOpenGLTexture* DepthStencilTarget)
+{
+	const int32 NumRenderingSamples = 1;
+	return GetOpenGLFramebuffer(NumSimultaneousRenderTargets, RenderTargets, ArrayIndices, MipmapLevels, DepthStencilTarget, NumRenderingSamples);
+}
+
+GLuint FOpenGLDynamicRHI::GetOpenGLFramebuffer(uint32 NumSimultaneousRenderTargets, FOpenGLTexture** RenderTargets, const uint32* ArrayIndices, const uint32* MipmapLevels, FOpenGLTexture* DepthStencilTarget, int32 NumRenderingSamples)
 {
 	VERIFY_GL_SCOPE();
 
 	check( NumSimultaneousRenderTargets <= MaxSimultaneousRenderTargets );
 
-	uint32 FramebufferRet = GetOpenGLFramebufferCache().FindRef(FOpenGLFramebufferKey(NumSimultaneousRenderTargets, RenderTargets, ArrayIndices, MipmapLevels, DepthStencilTarget, PlatformOpenGLCurrentContext(PlatformDevice)));
+	uint32 FramebufferRet = GetOpenGLFramebufferCache().FindRef(FOpenGLFramebufferKey(NumSimultaneousRenderTargets, RenderTargets, ArrayIndices, MipmapLevels, DepthStencilTarget, NumRenderingSamples, PlatformOpenGLCurrentContext(PlatformDevice)));
 	if( FramebufferRet > 0 )
 	{
 		// Found and is valid. We never store zero as a result, increasing all results by 1 to avoid range overlap.
@@ -130,7 +165,6 @@ GLuint FOpenGLDynamicRHI::GetOpenGLFramebuffer(uint32 NumSimultaneousRenderTarge
 	glBindFramebuffer(GL_FRAMEBUFFER, Framebuffer);
 	VERIFY_GL(glBindFramebuffer)
 
-#if PLATFORM_ANDROID
 	static const auto CVarMobileMultiView = IConsoleManager::Get().FindTConsoleVariableDataInt(TEXT("vr.MobileMultiView"));
 
 	// Allocate mobile multi-view frame buffer if enabled and supported.
@@ -147,26 +181,25 @@ GLuint FOpenGLDynamicRHI::GetOpenGLFramebuffer(uint32 NumSimultaneousRenderTarge
 		glBindFramebuffer(GL_FRAMEBUFFER, 0);
 		glBindFramebuffer(GL_DRAW_FRAMEBUFFER, Framebuffer);
 
-		const uint32 NumSamplesRendered = RenderTarget->GetNumSamplesRendered();
-		if (NumSamplesRendered > 1)
+		if (NumRenderingSamples > 1)
 		{
-			glFramebufferTextureMultisampleMultiviewOVR(GL_DRAW_FRAMEBUFFER, GL_COLOR_ATTACHMENT0, RenderTarget->GetResource(), 0, NumSamplesRendered, 0, 2);
+			FOpenGL::FramebufferTextureMultisampleMultiviewOVR(GL_DRAW_FRAMEBUFFER, GL_COLOR_ATTACHMENT0, RenderTarget->GetResource(), 0, NumRenderingSamples, 0, 2);
 			VERIFY_GL(glFramebufferTextureMultisampleMultiviewOVR);
 
 			if (DepthStencilTarget)
 			{
-				glFramebufferTextureMultisampleMultiviewOVR(GL_DRAW_FRAMEBUFFER, GL_DEPTH_ATTACHMENT, DepthStencilTarget->GetResource(), 0, NumSamplesRendered, 0, 2);
+				FOpenGL::FramebufferTextureMultisampleMultiviewOVR(GL_DRAW_FRAMEBUFFER, GL_DEPTH_ATTACHMENT, DepthStencilTarget->GetResource(), 0, NumRenderingSamples, 0, 2);
 				VERIFY_GL(glFramebufferTextureMultisampleMultiviewOVR);
 			}
 		}
 		else
 		{
-			glFramebufferTextureMultiviewOVR(GL_DRAW_FRAMEBUFFER, GL_COLOR_ATTACHMENT0, RenderTarget->GetResource(), 0, 0, 2);
+			FOpenGL::FramebufferTextureMultiviewOVR(GL_DRAW_FRAMEBUFFER, GL_COLOR_ATTACHMENT0, RenderTarget->GetResource(), 0, 0, 2);
 			VERIFY_GL(glFramebufferTextureMultiviewOVR);
 
 			if (DepthStencilTarget)
 			{
-				glFramebufferTextureMultiviewOVR(GL_DRAW_FRAMEBUFFER, GL_DEPTH_ATTACHMENT, DepthStencilTarget->GetResource(), 0, 0, 2);
+				FOpenGL::FramebufferTextureMultiviewOVR(GL_DRAW_FRAMEBUFFER, GL_DEPTH_ATTACHMENT, DepthStencilTarget->GetResource(), 0, 0, 2);
 				VERIFY_GL(glFramebufferTextureMultiviewOVR);
 			}
 		}
@@ -176,12 +209,11 @@ GLuint FOpenGLDynamicRHI::GetOpenGLFramebuffer(uint32 NumSimultaneousRenderTarge
 		FOpenGL::ReadBuffer(GL_NONE);
 		FOpenGL::DrawBuffer(GL_COLOR_ATTACHMENT0);
 
-		GetOpenGLFramebufferCache().Add(FOpenGLFramebufferKey(NumSimultaneousRenderTargets, RenderTargets, ArrayIndices, MipmapLevels, DepthStencilTarget, PlatformOpenGLCurrentContext(PlatformDevice)), Framebuffer + 1);
+		GetOpenGLFramebufferCache().Add(FOpenGLFramebufferKey(NumSimultaneousRenderTargets, RenderTargets, ArrayIndices, MipmapLevels, DepthStencilTarget, NumRenderingSamples, PlatformOpenGLCurrentContext(PlatformDevice)), Framebuffer + 1);
 		
 		return Framebuffer;
 	}
-#endif
-
+	
 	int32 FirstNonzeroRenderTarget = -1;
 	for (int32 RenderTargetIndex = NumSimultaneousRenderTargets - 1; RenderTargetIndex >= 0; --RenderTargetIndex)
 	{
@@ -196,19 +228,23 @@ GLuint FOpenGLDynamicRHI::GetOpenGLFramebuffer(uint32 NumSimultaneousRenderTarge
 			// If no index was specified, bind the entire object, rather than a slice
 			switch (RenderTarget->Target)
 			{
+			case GL_RENDERBUFFER:
+			{
+				// lazily allocate render buffer storage in case it's multisampled
+				ConditionallyAllocateRenderbufferStorage(*RenderTarget);
+				glFramebufferRenderbuffer(GL_FRAMEBUFFER, GL_COLOR_ATTACHMENT0 + RenderTargetIndex, GL_RENDERBUFFER, RenderTarget->GetResource());
+				break;
+			}
 			case GL_TEXTURE_2D:
 			case GL_TEXTURE_EXTERNAL_OES:
 			case GL_TEXTURE_2D_MULTISAMPLE:
 			{
-#if PLATFORM_ANDROID
-				if (RenderTarget->IsTiledMSAA() && glFramebufferTexture2DMultisampleEXT)
+				if (NumRenderingSamples > 1)
 				{
-					// GL_EXT_multisampled_render_to_texture requires GL_COLOR_ATTACHMENT0
-					glFramebufferTexture2DMultisampleEXT(GL_FRAMEBUFFER, GL_COLOR_ATTACHMENT0 + RenderTargetIndex, RenderTarget->Target, RenderTarget->GetResource(), MipmapLevels[RenderTargetIndex], RenderTarget->GetNumSamplesRendered());
-					VERIFY_GL(glFramebufferTexture2DMultisampleEXT);
+					// GL_EXT_multisampled_render_to_texture
+					FOpenGL::FramebufferTexture2DMultisample(GL_FRAMEBUFFER, GL_COLOR_ATTACHMENT0 + RenderTargetIndex, RenderTarget->Target, RenderTarget->GetResource(), MipmapLevels[RenderTargetIndex], NumRenderingSamples);
 				}
 				else
-#endif
 				{
 					FOpenGL::FramebufferTexture2D(GL_FRAMEBUFFER, GL_COLOR_ATTACHMENT0 + RenderTargetIndex, RenderTarget->Target, RenderTarget->GetResource(), MipmapLevels[RenderTargetIndex]);
 				}
@@ -230,20 +266,25 @@ GLuint FOpenGLDynamicRHI::GetOpenGLFramebuffer(uint32 NumSimultaneousRenderTarge
 			// Bind just one slice of the object
 			switch( RenderTarget->Target )
 			{
+			case GL_RENDERBUFFER:
+			{
+				check(ArrayIndices[RenderTargetIndex] == 0);
+				// lazily allocate render buffer storage in case it's multisampled
+				ConditionallyAllocateRenderbufferStorage(*RenderTarget);
+				glFramebufferRenderbuffer(GL_FRAMEBUFFER, GL_COLOR_ATTACHMENT0 + RenderTargetIndex, GL_RENDERBUFFER, RenderTarget->GetResource());
+				break;
+			}
 			case GL_TEXTURE_2D:
 			case GL_TEXTURE_EXTERNAL_OES:
 			case GL_TEXTURE_2D_MULTISAMPLE:
 			{
 				check(ArrayIndices[RenderTargetIndex] == 0);
-#if PLATFORM_ANDROID
-				if (RenderTarget->IsTiledMSAA() && glFramebufferTexture2DMultisampleEXT)
+				if (NumRenderingSamples > 1)
 				{
-					// GL_EXT_multisampled_render_to_texture requires GL_COLOR_ATTACHMENT0
-					glFramebufferTexture2DMultisampleEXT(GL_FRAMEBUFFER, GL_COLOR_ATTACHMENT0 + RenderTargetIndex, RenderTarget->Target, RenderTarget->GetResource(), MipmapLevels[RenderTargetIndex], RenderTarget->GetNumSamplesRendered());
-					VERIFY_GL(glFramebufferTexture2DMultisampleEXT);
+					// GL_EXT_multisampled_render_to_texture
+					FOpenGL::FramebufferTexture2DMultisample(GL_FRAMEBUFFER, GL_COLOR_ATTACHMENT0 + RenderTargetIndex, RenderTarget->Target, RenderTarget->GetResource(), MipmapLevels[RenderTargetIndex], NumRenderingSamples);
 				}
 				else
-#endif
 				{
 					FOpenGL::FramebufferTexture2D(GL_FRAMEBUFFER, GL_COLOR_ATTACHMENT0 + RenderTargetIndex, RenderTarget->Target, RenderTarget->GetResource(), MipmapLevels[RenderTargetIndex]);
 				}
@@ -282,6 +323,8 @@ GLuint FOpenGLDynamicRHI::GetOpenGLFramebuffer(uint32 NumSimultaneousRenderTarge
 		}
 		case GL_RENDERBUFFER:
 		{
+			// lazily allocate render buffer storage in case it's multisampled
+			ConditionallyAllocateRenderbufferStorage(*DepthStencilTarget);
 			glFramebufferRenderbuffer(GL_FRAMEBUFFER, GL_DEPTH_ATTACHMENT, GL_RENDERBUFFER, DepthStencilTarget->GetResource());
 			glFramebufferRenderbuffer(GL_FRAMEBUFFER, GL_STENCIL_ATTACHMENT, GL_RENDERBUFFER, DepthStencilTarget->GetResource());
 			VERIFY_GL(glFramebufferRenderbuffer);
@@ -294,7 +337,7 @@ GLuint FOpenGLDynamicRHI::GetOpenGLFramebuffer(uint32 NumSimultaneousRenderTarge
 			FOpenGL::FramebufferTexture(GL_FRAMEBUFFER, DepthStencilTarget->Attachment, DepthStencilTarget->GetResource(), 0);
 			break;
 		default:
-				FOpenGL::FramebufferRenderbuffer(GL_FRAMEBUFFER, DepthStencilTarget->Attachment, GL_RENDERBUFFER, DepthStencilTarget->GetResource());
+			FOpenGL::FramebufferRenderbuffer(GL_FRAMEBUFFER, DepthStencilTarget->Attachment, GL_RENDERBUFFER, DepthStencilTarget->GetResource());
 			break;
 		}
 	}
@@ -323,7 +366,7 @@ GLuint FOpenGLDynamicRHI::GetOpenGLFramebuffer(uint32 NumSimultaneousRenderTarge
 	
 	FOpenGL::CheckFrameBuffer();
 
-	GetOpenGLFramebufferCache().Add(FOpenGLFramebufferKey(NumSimultaneousRenderTargets, RenderTargets, ArrayIndices, MipmapLevels, DepthStencilTarget, PlatformOpenGLCurrentContext(PlatformDevice)), Framebuffer+1);
+	GetOpenGLFramebufferCache().Add(FOpenGLFramebufferKey(NumSimultaneousRenderTargets, RenderTargets, ArrayIndices, MipmapLevels, DepthStencilTarget, NumRenderingSamples, PlatformOpenGLCurrentContext(PlatformDevice)), Framebuffer+1);
 
 	return Framebuffer;
 }
@@ -1125,11 +1168,35 @@ void FOpenGLDynamicRHI::BindPendingFramebuffer( FOpenGLContextState& ContextStat
 	}
 }
 
+// Replaces RenderTargets with ResoveTargets to utilize GL_EXT_multisampled_render_to_texture
+static int32 SetupMultisampleRenderingInfo(FRHISetRenderTargetsInfo& RTInfo)
+{
+	if (FOpenGL::GetMaxMSAASamplesTileMem() > 1 && RTInfo.NumColorRenderTargets > 0)
+	{
+		int32 NumRenderingSamples = RTInfo.ColorRenderTarget[0].Texture->GetDesc().NumSamples;
+		if (NumRenderingSamples > 1)
+		{
+			for (int32 i = 0; i < RTInfo.NumColorRenderTargets; ++i)
+			{
+				if (RTInfo.ColorResolveRenderTarget[i].Texture)
+				{
+					RTInfo.ColorRenderTarget[i].Texture = RTInfo.ColorResolveRenderTarget[i].Texture;
+				}
+			}
+			
+			return NumRenderingSamples;
+		}
+	}
+		
+	return 1;
+}
 
 void FOpenGLDynamicRHI::RHIBeginRenderPass(const FRHIRenderPassInfo& InInfo, const TCHAR* InName)
 {
 	FRHISetRenderTargetsInfo RTInfo;
 	InInfo.ConvertToRenderTargetsInfo(RTInfo);
+	// Begin GL_EXT_multisampled_render_to_texture if any
+	PendingState.NumRenderingSamples = SetupMultisampleRenderingInfo(RTInfo);
 	SetRenderTargetsAndClear(RTInfo);
 
 	RenderPassInfo = InInfo;
@@ -1164,6 +1231,9 @@ void FOpenGLDynamicRHI::RHIEndRenderPass()
 		EndOcclusionQueryBatch();
 	}
 
+	// End GL_EXT_multisampled_render_to_texture
+	PendingState.NumRenderingSamples = 1;
+
 	// Discard transient color targets
 	uint32 ColorMask = 0u;
 	for (int32 ColorIndex = 0; ColorIndex < MaxSimultaneousRenderTargets; ++ColorIndex)
diff --git a/Engine/Source/Runtime/OpenGLDrv/Private/OpenGLTexture.cpp b/Engine/Source/Runtime/OpenGLDrv/Private/OpenGLTexture.cpp
--- a/Engine/Source/Runtime/OpenGLDrv/Private/OpenGLTexture.cpp
+++ b/Engine/Source/Runtime/OpenGLDrv/Private/OpenGLTexture.cpp
@@ -173,38 +173,38 @@ bool FOpenGLDynamicRHI::RHIGetTextureMemoryVisualizeData( FColor* /*TextureData*
 }
 
 FOpenGLTextureDesc::FOpenGLTextureDesc(FRHITextureDesc const& InDesc)
-	: NumSamplesRendered (InDesc.NumSamples)
-	, NumSamplesStored   (InDesc.NumSamples)
-	, bCubemap           (InDesc.IsTextureCube())
+	: bCubemap           (InDesc.IsTextureCube())
 	, bArrayTexture      (InDesc.IsTextureArray())
 	, bStreamable        (EnumHasAnyFlags(InDesc.Flags, TexCreate_Streamable))
 	, bDepthStencil      (EnumHasAnyFlags(InDesc.Flags, TexCreate_DepthStencilTargetable))
 	, bCanCreateAsEvicted(false)
 	, bIsPowerOfTwo      (false)
-	, bTileMemDepthBuffer(false)
+	, bMultisampleRenderbuffer(EnumHasAnyFlags(InDesc.Flags, TexCreate_Memoryless) && InDesc.NumSamples > 1)
 {
-	checkf(!bCubemap || NumSamplesStored == 1, TEXT("Texture cubes cannot be multisampled."));
+	checkf(!bCubemap || InDesc.NumSamples == 1, TEXT("Texture cubes cannot be multisampled."));
 	checkf(FOpenGL::SupportsTexture3D() || (!InDesc.IsTexture3D() && !InDesc.IsTextureArray()), TEXT("Texture3D / Texture2DArray support requires FOpenGL::SupportsTexture3D()."));
+	checkf(!bMultisampleRenderbuffer || EnumHasAnyFlags(InDesc.Flags, TexCreate_RenderTargetable | TexCreate_DepthStencilTargetable), TEXT("Only render targets can be memoryless"));
 
-	// Use on-chip tile memory for MSAA if available
-	if (NumSamplesStored <= FOpenGL::GetMaxMSAASamplesTileMem())
+	// Special case for multiview MSAA depth target. It has to be a non-MSAA texture with multisample rendering
+	const bool bMultiviewMSAADepthTarget = (bDepthStencil && InDesc.NumSamples > 1 && InDesc.Dimension == ETextureDimension::Texture2DArray);
+	if (bMultiviewMSAADepthTarget)
 	{
-		NumSamplesRendered = FMath::Min<uint32>(NumSamplesStored, FOpenGL::GetMaxMSAASamplesTileMem());
-		NumSamplesStored = 1;
+		bMultisampleRenderbuffer = false;
 	}
-
+			
 	// Select an appropriate texture target
-#if PLATFORM_ANDROID
-	if (bDepthStencil && NumSamplesRendered != NumSamplesStored)
+	if (bMultisampleRenderbuffer)
 	{
-		// Special case for MSAA depth render target on tiled GPUs / mobile renderer.
+		// Special case for multisample memoryless render targets
 		Target = GL_RENDERBUFFER;
-		bTileMemDepthBuffer = true;
 	}
 	else 
-#endif
 	{
-		if (EnumHasAnyFlags(InDesc.Flags, TexCreate_External))
+		if (bMultiviewMSAADepthTarget)
+		{
+			Target = GL_TEXTURE_2D_ARRAY;
+		}
+		else if (EnumHasAnyFlags(InDesc.Flags, TexCreate_External))
 		{
 			check(InDesc.IsTexture2D());
 			check(!InDesc.IsTextureArray());
@@ -225,8 +225,8 @@ FOpenGLTextureDesc::FOpenGLTextureDesc(FRHITextureDesc const& InDesc)
 			switch (InDesc.Dimension)
 			{
 			default: checkNoEntry();
-			case ETextureDimension::Texture2D:		  Target = (NumSamplesStored > 1) ? GL_TEXTURE_2D_MULTISAMPLE       : GL_TEXTURE_2D;       break;
-			case ETextureDimension::Texture2DArray:	  Target = (NumSamplesStored > 1) ? GL_TEXTURE_2D_MULTISAMPLE_ARRAY : GL_TEXTURE_2D_ARRAY; break;
+			case ETextureDimension::Texture2D:		  Target = (InDesc.NumSamples > 1) ? GL_TEXTURE_2D_MULTISAMPLE       : GL_TEXTURE_2D;       break;
+			case ETextureDimension::Texture2DArray:	  Target = (InDesc.NumSamples > 1) ? GL_TEXTURE_2D_MULTISAMPLE_ARRAY : GL_TEXTURE_2D_ARRAY; break;
 			case ETextureDimension::TextureCubeArray: Target = GL_TEXTURE_CUBE_MAP_ARRAY; break;
 			case ETextureDimension::TextureCube:	  Target = GL_TEXTURE_CUBE_MAP;       break;
 			case ETextureDimension::Texture3D:		  Target = GL_TEXTURE_3D;             break;
@@ -282,11 +282,11 @@ FOpenGLTextureDesc::FOpenGLTextureDesc(FRHITextureDesc const& InDesc)
 	switch (Attachment)
 	{
 	case GL_COLOR_ATTACHMENT0:
-		check(GMaxOpenGLColorSamples >= (GLint)NumSamplesRendered);
+		check(GMaxOpenGLColorSamples >= (GLint)InDesc.NumSamples);
 		break;
 	case GL_DEPTH_ATTACHMENT:
 	case GL_DEPTH_STENCIL_ATTACHMENT:
-		check(GMaxOpenGLDepthSamples >= (GLint)NumSamplesRendered);
+		check(GMaxOpenGLDepthSamples >= (GLint)InDesc.NumSamples);
 		break;
 	default:
 		break;
@@ -299,15 +299,14 @@ FOpenGLTexture::FOpenGLTexture(FOpenGLTexture& Other, const FString& Name, EAlia
 	, Target             (Other.Target)
 	, Attachment         (Other.Attachment)
 	, MemorySize         (0)
-	, NumSamplesInternal (Other.NumSamplesInternal)
 	, bIsPowerOfTwo      (Other.bIsPowerOfTwo)
 	, bCanCreateAsEvicted(false)
 	, bStreamable        (Other.bStreamable)
 	, bCubemap           (Other.bCubemap)
 	, bArrayTexture      (Other.bArrayTexture)
 	, bDepthStencil      (Other.bDepthStencil)
-	, bTileMemDepthBuffer(Other.bTileMemDepthBuffer)
 	, bAlias             (true)
+	, bMultisampleRenderbuffer(Other.bMultisampleRenderbuffer)
 {
 	RunOnGLRenderContextThread([&]()
 	{
@@ -336,15 +335,14 @@ FOpenGLTexture::FOpenGLTexture(FOpenGLTextureCreateDesc const& CreateDesc, GLuin
 	, Target             (CreateDesc.Target)
 	, Attachment         (CreateDesc.Attachment)
 	, MemorySize         (CreateDesc.MemorySize)
-	, NumSamplesInternal (CreateDesc.NumSamplesRendered)
 	, bIsPowerOfTwo      (CreateDesc.bIsPowerOfTwo)
 	, bCanCreateAsEvicted(false)
 	, bStreamable        (CreateDesc.bStreamable)
 	, bCubemap           (CreateDesc.bCubemap)
 	, bArrayTexture      (CreateDesc.bArrayTexture)
 	, bDepthStencil      (CreateDesc.bDepthStencil)
-	, bTileMemDepthBuffer(CreateDesc.bTileMemDepthBuffer)
 	, bAlias             (true)
+	, bMultisampleRenderbuffer(CreateDesc.bMultisampleRenderbuffer)
 {}
 
 // Standard constructor.
@@ -353,15 +351,14 @@ FOpenGLTexture::FOpenGLTexture(FOpenGLTextureCreateDesc const& CreateDesc)
 	, Target             (CreateDesc.Target)
 	, Attachment         (CreateDesc.Attachment)
 	, MemorySize         (CreateDesc.MemorySize)
-	, NumSamplesInternal (CreateDesc.NumSamplesRendered)
 	, bIsPowerOfTwo      (CreateDesc.bIsPowerOfTwo)
 	, bCanCreateAsEvicted(CreateDesc.bCanCreateAsEvicted)
 	, bStreamable        (CreateDesc.bStreamable)
 	, bCubemap           (CreateDesc.bCubemap)
 	, bArrayTexture      (CreateDesc.bArrayTexture)
 	, bDepthStencil      (CreateDesc.bDepthStencil)
-	, bTileMemDepthBuffer(CreateDesc.bTileMemDepthBuffer)
 	, bAlias             (false)
+	, bMultisampleRenderbuffer(CreateDesc.bMultisampleRenderbuffer)
 {
 	check(IsInRenderingThread());
 	FRHICommandListImmediate& RHICmdList = FRHICommandListExecutor::GetImmediateCommandList();
@@ -487,7 +484,7 @@ void FOpenGLDynamicRHI::InitializeGLTexture(FOpenGLTexture* Texture, const void*
 
 	// Allocate the GL resource ID
 	GLuint TextureID;
-	if (Texture->bTileMemDepthBuffer)
+	if (Texture->bMultisampleRenderbuffer)
 	{
 		check(Texture->Target == GL_RENDERBUFFER);
 		glGenRenderbuffers(1, &TextureID);
@@ -554,6 +551,8 @@ void FOpenGLDynamicRHI::InitializeGLTextureInternal(FOpenGLTexture* Texture, voi
 	{
 		UE_LOG(LogRHI, Fatal,TEXT("Texture format '%s' not supported (sRGB=%d)."), GPixelFormats[Desc.Format].Name, bSRGB);
 	}
+	
+	const bool bMultiviewMSAADepthTarget = (Texture->bDepthStencil && Desc.NumSamples > 1 && Desc.Dimension == ETextureDimension::Texture2DArray);
 
 	FOpenGLContextState& ContextState = GetContextStateForCurrentContext();
 
@@ -561,20 +560,13 @@ void FOpenGLDynamicRHI::InitializeGLTextureInternal(FOpenGLTexture* Texture, voi
 	CachedBindPixelUnpackBuffer(ContextState, 0);
 
 	bool bAllocatedStorage = false;
-#if PLATFORM_ANDROID
-	if (Texture->bTileMemDepthBuffer)
+	if (Texture->bMultisampleRenderbuffer)
 	{
-		glBindRenderbuffer(GL_RENDERBUFFER, TextureID);
-
-		glRenderbufferStorageMultisampleEXT(GL_RENDERBUFFER, Texture->GetNumSamplesRendered(), GL_DEPTH24_STENCIL8, Desc.Extent.X, Desc.Extent.Y);
-		VERIFY_GL(glRenderbufferStorageMultisampleEXT);
-
-		glBindRenderbuffer(GL_RENDERBUFFER, 0);
-
-		bAllocatedStorage = true;
+		check(Texture->IsMultisampled());
+		check(Target == GL_RENDERBUFFER);
+		// Multisample Renderbuffers will be allocated on first use. See ConditionallyAllocateRenderbufferStorage
 	}
 	else
-#endif
 	{
 		// Use a texture stage that's not likely to be used for draws, to avoid waiting
 		CachedSetupTextureStage(ContextState, FOpenGL::GetMaxCombinedTextureImageUnits() - 1, Target, TextureID, 0, Desc.NumMips);
@@ -589,7 +581,7 @@ void FOpenGLDynamicRHI::InitializeGLTextureInternal(FOpenGLTexture* Texture, voi
 			glTexParameteri(Target, GL_TEXTURE_SWIZZLE_B, GL_RED);
 		}
 
-		if (Texture->GetNumSamplesStored() == 1)
+		if (!Texture->IsMultisampled())
 		{
 			if (Target == GL_TEXTURE_EXTERNAL_OES || !FMath::IsPowerOfTwo(Desc.Extent.X) || !FMath::IsPowerOfTwo(Desc.Extent.Y))
 			{
@@ -837,14 +829,15 @@ void FOpenGLDynamicRHI::InitializeGLTextureInternal(FOpenGLTexture* Texture, voi
 					checkf(BulkDataPtr == nullptr, TEXT("Multisample textures cannot be created with initial bulk data."));
 
 					// Try to create an immutable storage texture and fallback if it fails
+					const int32 NumSamples = Texture->GetDesc().NumSamples;
 					const bool FixedSampleLocations = true;
-					if (FOpenGL::TexStorage2DMultisample(Target, Texture->GetNumSamplesStored(), GLFormat.InternalFormat[bSRGB], Desc.Extent.X, Desc.Extent.Y, FixedSampleLocations))
+					if (FOpenGL::TexStorage2DMultisample(Target, NumSamples, GLFormat.InternalFormat[bSRGB], Desc.Extent.X, Desc.Extent.Y, FixedSampleLocations))
 					{
 						bAllocatedStorage = true;
 					}
 					else
 					{
-						FOpenGL::TexImage2DMultisample(Target, Texture->GetNumSamplesStored(), GLFormat.InternalFormat[bSRGB], Desc.Extent.X, Desc.Extent.Y, FixedSampleLocations);
+						FOpenGL::TexImage2DMultisample(Target, NumSamples, GLFormat.InternalFormat[bSRGB], Desc.Extent.X, Desc.Extent.Y, FixedSampleLocations);
 					}
 				}
 				break;
diff --git a/Engine/Source/Runtime/OpenGLDrv/Public/OpenGL.h b/Engine/Source/Runtime/OpenGLDrv/Public/OpenGL.h
--- a/Engine/Source/Runtime/OpenGLDrv/Public/OpenGL.h
+++ b/Engine/Source/Runtime/OpenGLDrv/Public/OpenGL.h
@@ -220,12 +220,15 @@ public:
 	{
 		glFramebufferTexture2D(Target, Attachment, TexTarget, Texture, Level);
 	}
+	static FORCEINLINE void FramebufferTexture2DMultisample(GLenum Target, GLenum Attachment, GLenum TexTarget, GLuint Texture, GLint Level, GLint NumSamples) UGL_REQUIRED_VOID
 	static FORCEINLINE void FramebufferTexture3D(GLenum Target, GLenum Attachment, GLenum TexTarget, GLuint Texture, GLint Level, GLint ZOffset) UGL_REQUIRED_VOID
 	static FORCEINLINE void FramebufferTextureLayer(GLenum Target, GLenum Attachment, GLuint Texture, GLint Level, GLint Layer) UGL_REQUIRED_VOID
 	static FORCEINLINE void FramebufferRenderbuffer(GLenum Target, GLenum Attachment, GLenum RenderBufferTarget, GLuint RenderBuffer)
 	{
 		glFramebufferRenderbuffer(Target, Attachment, RenderBufferTarget, RenderBuffer);
 	}
+	static FORCEINLINE void FramebufferTextureMultiviewOVR(GLenum Target, GLenum Attachment, GLuint Texture, GLint Level, GLint BaseViewIndex, GLsizei NumViews) UGL_REQUIRED_VOID
+	static FORCEINLINE void FramebufferTextureMultisampleMultiviewOVR(GLenum Target, GLenum Attachment, GLuint Texture, GLint Level, GLsizei NumSamples, GLint BaseViewIndex, GLsizei NumViews) UGL_REQUIRED_VOID
 	static FORCEINLINE void BlitFramebuffer(GLint SrcX0, GLint SrcY0, GLint SrcX1, GLint SrcY1, GLint DstX0, GLint DstY0, GLint DstX1, GLint DstY1, GLbitfield Mask, GLenum Filter) UGL_REQUIRED_VOID
 	static FORCEINLINE void DrawBuffers(GLsizei NumBuffers, const GLenum *Buffers) UGL_REQUIRED_VOID
 	static FORCEINLINE void DepthRange(GLdouble Near, GLdouble Far) UGL_REQUIRED_VOID
@@ -282,6 +285,7 @@ public:
 	static FORCEINLINE void MemoryBarrier(GLbitfield Barriers) UGL_REQUIRED_VOID
 	static FORCEINLINE bool TexStorage2D(GLenum Target, GLint Levels, GLint InternalFormat, GLsizei Width, GLsizei Height, GLenum Format, GLenum Type, ETextureCreateFlags Flags) UGL_OPTIONAL(false)
 	static FORCEINLINE bool TexStorage2DMultisample(GLenum Target, GLsizei Samples, GLint InternalFormat, GLsizei Width, GLsizei Height, GLboolean FixedSampleLocations) UGL_OPTIONAL(false)
+	static FORCEINLINE void RenderbufferStorageMultisample(GLenum Target, GLsizei Samples, GLint InternalFormat, GLsizei Width, GLsizei Height) UGL_REQUIRED_VOID
 	static FORCEINLINE void TexStorage3D(GLenum Target, GLint Levels, GLint InternalFormat, GLsizei Width, GLsizei Height, GLsizei Depth, GLenum Format, GLenum Type) UGL_REQUIRED_VOID
 	static FORCEINLINE void CompressedTexSubImage3D(GLenum Target, GLint Level, GLint XOffset, GLint YOffset, GLint ZOffset, GLsizei Width, GLsizei Height, GLsizei Depth, GLenum Format, GLsizei ImageSize, const GLvoid* PixelData) UGL_REQUIRED_VOID
 	static FORCEINLINE void CopyImageSubData(GLuint SrcName, GLenum SrcTarget, GLint SrcLevel, GLint SrcX, GLint SrcY, GLint SrcZ, GLuint DstName, GLenum DstTarget, GLint DstLevel, GLint DstX, GLint DstY, GLint DstZ, GLsizei Width, GLsizei Height, GLsizei Depth) UGL_REQUIRED_VOID
@@ -329,7 +333,7 @@ public:
 
 	static FORCEINLINE void CheckFrameBuffer()
 	{
-#if UE_BUILD_DEBUG 
+#if UE_BUILD_DEBUG || UE_BUILD_DEVELOPMENT 
 		GLenum CompleteResult = glCheckFramebufferStatus(GL_FRAMEBUFFER);
 		if (CompleteResult != GL_FRAMEBUFFER_COMPLETE)
 		{
diff --git a/Engine/Source/Runtime/OpenGLDrv/Public/OpenGLDrv.h b/Engine/Source/Runtime/OpenGLDrv/Public/OpenGLDrv.h
--- a/Engine/Source/Runtime/OpenGLDrv/Public/OpenGLDrv.h
+++ b/Engine/Source/Runtime/OpenGLDrv/Public/OpenGLDrv.h
@@ -908,6 +908,7 @@ public:
 	}
 
 	GLuint GetOpenGLFramebuffer(uint32 NumSimultaneousRenderTargets, FOpenGLTexture** RenderTargets, const uint32* ArrayIndices, const uint32* MipmapLevels, FOpenGLTexture* DepthStencilTarget);
+	GLuint GetOpenGLFramebuffer(uint32 NumSimultaneousRenderTargets, FOpenGLTexture** RenderTargets, const uint32* ArrayIndices, const uint32* MipmapLevels, FOpenGLTexture* DepthStencilTarget, int32 NumRenderingSamples);
 
 	FRHIPixelShader* GetNULLPixelShader() const;
 	
diff --git a/Engine/Source/Runtime/OpenGLDrv/Public/OpenGLES.h b/Engine/Source/Runtime/OpenGLDrv/Public/OpenGLES.h
--- a/Engine/Source/Runtime/OpenGLDrv/Public/OpenGLES.h
+++ b/Engine/Source/Runtime/OpenGLDrv/Public/OpenGLES.h
@@ -383,6 +383,12 @@ struct FOpenGLES : public FOpenGLBase
 		return true;
 	}
 
+	static FORCEINLINE void RenderbufferStorageMultisample(GLenum Target, GLsizei Samples, GLint InternalFormat, GLsizei Width, GLsizei Height)
+	{
+		check(glRenderbufferStorageMultisampleEXT);
+		glRenderbufferStorageMultisampleEXT(Target, Samples, InternalFormat, Width, Height);
+	}
+
 	static FORCEINLINE void ClearBufferfv(GLenum Buffer, GLint DrawBufferIndex, const GLfloat* Value)
 	{
 		glClearBufferfv(Buffer, DrawBufferIndex, Value);
@@ -657,6 +663,24 @@ struct FOpenGLES : public FOpenGLBase
 		VERIFY_GL(FramebufferTexture_2D);
 	}
 
+	static FORCEINLINE void FramebufferTexture2DMultisample(GLenum Target, GLenum Attachment, GLenum TexTarget, GLuint Texture, GLint Level, GLint NumSamples)
+	{
+		check(glFramebufferTexture2DMultisampleEXT != nullptr);
+		glFramebufferTexture2DMultisampleEXT(Target, Attachment, TexTarget, Texture, Level, NumSamples);
+	}
+
+	static FORCEINLINE void FramebufferTextureMultiviewOVR(GLenum Target, GLenum Attachment, GLuint Texture, GLint Level, GLint BaseViewIndex, GLsizei NumViews)
+	{
+		check(glFramebufferTextureMultiviewOVR);
+		glFramebufferTextureMultiviewOVR(Target, Attachment, Texture, Level, BaseViewIndex, NumViews);
+	}
+	
+	static FORCEINLINE void FramebufferTextureMultisampleMultiviewOVR(GLenum Target, GLenum Attachment, GLuint Texture, GLint Level, GLsizei NumSamples, GLint BaseViewIndex, GLsizei NumViews)
+	{
+		check(glFramebufferTextureMultisampleMultiviewOVR);
+		glFramebufferTextureMultisampleMultiviewOVR(Target, Attachment, Texture, Level, NumSamples, BaseViewIndex, NumViews);
+	}
+
 	static FORCEINLINE void BlitFramebuffer(GLint SrcX0, GLint SrcY0, GLint SrcX1, GLint SrcY1, GLint DstX0, GLint DstY0, GLint DstX1, GLint DstY1, GLbitfield Mask, GLenum Filter)
 	{
 		glBlitFramebuffer(SrcX0, SrcY0, SrcX1, SrcY1, DstX0, DstY0, DstX1, DstY1, Mask, Filter);
diff --git a/Engine/Source/Runtime/OpenGLDrv/Public/OpenGLResources.h b/Engine/Source/Runtime/OpenGLDrv/Public/OpenGLResources.h
--- a/Engine/Source/Runtime/OpenGLDrv/Public/OpenGLResources.h
+++ b/Engine/Source/Runtime/OpenGLDrv/Public/OpenGLResources.h
@@ -1244,9 +1244,6 @@ public:
 	GLenum Target = GL_NONE;
 	GLenum Attachment = GL_NONE;
 
-	uint8 NumSamplesRendered = 1;
-	uint8 NumSamplesStored   = 1;
-
 	uint32 MemorySize = 0;
 
 	uint8 bCubemap            : 1;
@@ -1255,7 +1252,7 @@ public:
 	uint8 bDepthStencil       : 1;
 	uint8 bCanCreateAsEvicted : 1;
 	uint8 bIsPowerOfTwo       : 1;
-	uint8 bTileMemDepthBuffer : 1;
+	uint8 bMultisampleRenderbuffer : 1;
 
 private:
 	static bool CanDeferTextureCreation();
@@ -1268,9 +1265,6 @@ public:
 		: FRHITextureCreateDesc(CreateDesc)
 		, FOpenGLTextureDesc(CreateDesc)
 	{
-		// Tiled GPUs modify the number of samples the FRHITexture reports.
-		// Ensure the base descriptor is updated with the actual number of samples.
-		FRHITextureCreateDesc::NumSamples = FOpenGLTextureDesc::NumSamplesStored;
 	}
 };
 
@@ -1327,25 +1321,6 @@ public:
 		Resource = InResource;
 	}
 
-	// Returns true if this texture will perform automatic MSAA resolve on tile-based GPUs.
-	bool IsTiledMSAA() const
-    {
-		return NumSamplesInternal > 1 && NumSamplesInternal <= FOpenGL::GetMaxMSAASamplesTileMem(); //-V547
-    }
-
-	// Returns the number of samples used during rendering. These samples may be kept in on-chip GPU tiled memory.
-	// For non-tiled GPUs, this is always equal to the number of stored samples.
-	uint32 GetNumSamplesRendered() const
-	{
-		return NumSamplesInternal;
-	}
-
-	// Returns the number of samples stored in the actual GPU resource memory.
-	uint32 GetNumSamplesStored() const
-	{
-		return IsTiledMSAA() ? 1 : NumSamplesInternal;
-	}
-
 	bool IsEvicted() const { VERIFY_GL_SCOPE(); return EvictionParamsPtr.IsValid() && !EvictionParamsPtr->bHasRestored; }
 
 	virtual void* GetTextureBaseRHI() override final
@@ -1471,9 +1446,6 @@ private:
 public:
 	uint32 const MemorySize;
 
-private:
-	uint8 const NumSamplesInternal;
-
 public:
 	uint8 const bIsPowerOfTwo       : 1;
 	uint8 const bCanCreateAsEvicted : 1;
@@ -1481,8 +1453,8 @@ public:
 	uint8 const bCubemap            : 1;
 	uint8 const bArrayTexture       : 1;
 	uint8 const bDepthStencil       : 1;
-	uint8 const bTileMemDepthBuffer : 1;
 	uint8 const bAlias              : 1;
+	uint8 const bMultisampleRenderbuffer : 1;
 };
 
 template <typename T>
diff --git a/Engine/Source/Runtime/OpenGLDrv/Public/OpenGLState.h b/Engine/Source/Runtime/OpenGLDrv/Public/OpenGLState.h
--- a/Engine/Source/Runtime/OpenGLDrv/Public/OpenGLState.h
+++ b/Engine/Source/Runtime/OpenGLDrv/Public/OpenGLState.h
@@ -427,6 +427,7 @@ struct FOpenGLRHIState final : public FOpenGLCommonState
 	bool							bAlphaToCoverageEnabled;
 
 	// Pending framebuffer setup
+	int32							NumRenderingSamples;// Only used with GL_EXT_multisampled_render_to_texture
 	int32							FirstNonzeroRenderTarget;
 	FOpenGLTexture*					RenderTargets[MaxSimultaneousRenderTargets];
 	uint32							RenderTargetMipmapLevels[MaxSimultaneousRenderTargets];
@@ -484,6 +485,7 @@ struct FOpenGLRHIState final : public FOpenGLCommonState
 	,	RenderTargetHeight(0)
 	,	RunningOcclusionQuery(0)
 	,	bAlphaToCoverageEnabled(false)
+	,	NumRenderingSamples(1)
 	,	FirstNonzeroRenderTarget(-1)
 	,	DepthStencil(0)
 	,	StencilStoreAction(ERenderTargetStoreAction::ENoAction)
diff --git a/Engine/Source/Runtime/RHI/Public/RHIDefinitions.h b/Engine/Source/Runtime/RHI/Public/RHIDefinitions.h
--- a/Engine/Source/Runtime/RHI/Public/RHIDefinitions.h
+++ b/Engine/Source/Runtime/RHI/Public/RHIDefinitions.h
@@ -2215,7 +2215,7 @@ inline bool IsFeatureLevelSupported(const FStaticShaderPlatform InShaderPlatform
 inline bool RHISupportsSeparateMSAAAndResolveTextures(const FStaticShaderPlatform Platform)
 {
 	// Metal mobile devices and Android ES3.1 need to handle MSAA and resolve textures internally (unless RHICreateTexture2D was changed to take an optional resolve target)
-	return !IsMetalMobilePlatform(Platform) && !IsAndroidOpenGLESPlatform(Platform);
+	return !IsMetalMobilePlatform(Platform);
 }
 
 inline bool RHISupportsComputeShaders(const FStaticShaderPlatform Platform)
diff --git a/Engine/Source/Runtime/Renderer/Private/MobileShadingRenderer.cpp b/Engine/Source/Runtime/Renderer/Private/MobileShadingRenderer.cpp
--- a/Engine/Source/Runtime/Renderer/Private/MobileShadingRenderer.cpp
+++ b/Engine/Source/Runtime/Renderer/Private/MobileShadingRenderer.cpp
@@ -530,6 +530,8 @@ void FMobileSceneRenderer::InitViews(FRDGBuilder& GraphBuilder, FSceneTexturesCo
 
 	// Update the bKeepDepthContent based on the mobile renderer status.
 	SceneTexturesConfig.bKeepDepthContent = bKeepDepthContent;
+	// If we render in a single pass MSAA targets can be memoryless
+	SceneTexturesConfig.bMemorylessMSAA = !bRequiresMultiPass;
 	
 	if (bDeferredShading) 
 	{
@@ -1215,6 +1217,7 @@ void FMobileSceneRenderer::RenderForwardMultiPass(FRDGBuilder& GraphBuilder, FMo
 
 	// resolve MSAA depth for translucency
 	AddResolveSceneDepthPass(GraphBuilder, View, SceneTextures.Depth);
+	AddResolveSceneColorPass(GraphBuilder, View, SceneTextures.DepthAux);
 
 	FExclusiveDepthStencil::Type ExclusiveDepthStencil = FExclusiveDepthStencil::DepthRead_StencilRead;
 	if (bModulatedShadowsInUse)
@@ -1267,6 +1270,8 @@ void FMobileSceneRenderer::RenderForwardMultiPass(FRDGBuilder& GraphBuilder, FMo
 		// Pre-tonemap before MSAA resolve (iOS only)
 		PreTonemapMSAA(RHICmdList, SceneTextures);
 	});
+
+	AddResolveSceneColorPass(GraphBuilder, View, SceneTextures.Color);
 }
 
 class FMobileDeferredCopyPLSPS : public FGlobalShader
diff --git a/Engine/Source/Runtime/Renderer/Private/SceneTextures.cpp b/Engine/Source/Runtime/Renderer/Private/SceneTextures.cpp
--- a/Engine/Source/Runtime/Renderer/Private/SceneTextures.cpp
+++ b/Engine/Source/Runtime/Renderer/Private/SceneTextures.cpp
@@ -638,9 +638,10 @@ void FMinimalSceneTextures::InitializeViewFamily(FRDGBuilder& GraphBuilder, FVie
 	{
 		ETextureCreateFlags Flags = TexCreate_DepthStencilTargetable | TexCreate_ShaderResource | TexCreate_InputAttachmentRead | GFastVRamConfig.SceneDepth;
 
-		if (!Config.bKeepDepthContent)
+		ETextureCreateFlags MemorylessFlag = TexCreate_None;
+		if (!Config.bKeepDepthContent || (Config.NumSamples > 1 && Config.bMemorylessMSAA))
 		{
-			Flags |= TexCreate_Memoryless;
+			MemorylessFlag = TexCreate_Memoryless;
 		}
 
 		if (Config.NumSamples == 1 && GRHISupportsDepthUAV)
@@ -650,8 +651,8 @@ void FMinimalSceneTextures::InitializeViewFamily(FRDGBuilder& GraphBuilder, FVie
 
 		// TODO: Array-size could be values > 2, on upcoming XR devices. This should be part of the config.
 		FRDGTextureDesc Desc(Config.bRequireMultiView ?
-							 FRDGTextureDesc::Create2DArray(SceneTextures.Config.Extent, PF_DepthStencil, Config.DepthClearValue, Flags, 2) :
-							 FRDGTextureDesc::Create2D(SceneTextures.Config.Extent, PF_DepthStencil, Config.DepthClearValue, Flags));
+							 FRDGTextureDesc::Create2DArray(SceneTextures.Config.Extent, PF_DepthStencil, Config.DepthClearValue, Flags | MemorylessFlag, 2) :
+							 FRDGTextureDesc::Create2D(SceneTextures.Config.Extent, PF_DepthStencil, Config.DepthClearValue, Flags | MemorylessFlag));
 		Desc.NumSamples = Config.NumSamples;
 		SceneTextures.Depth = GraphBuilder.CreateTexture(Desc, TEXT("SceneDepthZ"));
 
@@ -683,6 +684,12 @@ void FMinimalSceneTextures::InitializeViewFamily(FRDGBuilder& GraphBuilder, FVie
 		{
 			Flags |= TexCreate_UAV;
 		}
+		
+		if (Config.NumSamples > 1 && Config.bMemorylessMSAA)
+		{
+			Flags |= TexCreate_Memoryless;
+		}
+
 		Flags |= sRGBFlag;
 
 		const TCHAR* SceneColorName = TEXT("SceneColor");
@@ -794,11 +801,16 @@ void FSceneTextures::InitializeViewFamily(FRDGBuilder& GraphBuilder, FViewFamily
 	{
 		const float FarDepth = (float)ERHIZBuffer::FarPlane;
 		const FLinearColor FarDepthColor(FarDepth, FarDepth, FarDepth, FarDepth);
-		ETextureCreateFlags FlagsToAdd = IsMobileDeferredShadingEnabled(Config.ShaderPlatform)? TexCreate_Memoryless : TexCreate_None;
+		ETextureCreateFlags MemorylessFlag = TexCreate_None;
+		if (IsMobileDeferredShadingEnabled(Config.ShaderPlatform) || (Config.NumSamples > 1 && Config.bMemorylessMSAA))
+		{
+			MemorylessFlag = TexCreate_Memoryless;
+		}
+
 		EPixelFormat DepthAuxFormat = GetMobileSceneDepthAuxPixelFormat(Config.ShaderPlatform, Config.bPreciseDepthAux);
 		FRDGTextureDesc Desc = Config.bRequireMultiView ? 
-			FRDGTextureDesc::Create2DArray(Config.Extent, DepthAuxFormat, FClearValueBinding(FarDepthColor), TexCreate_RenderTargetable | TexCreate_ShaderResource | TexCreate_InputAttachmentRead | FlagsToAdd, 2) :
-			FRDGTextureDesc::Create2D(Config.Extent, DepthAuxFormat, FClearValueBinding(FarDepthColor), TexCreate_RenderTargetable | TexCreate_ShaderResource | TexCreate_InputAttachmentRead| FlagsToAdd);
+			FRDGTextureDesc::Create2DArray(Config.Extent, DepthAuxFormat, FClearValueBinding(FarDepthColor), TexCreate_RenderTargetable | TexCreate_ShaderResource | TexCreate_InputAttachmentRead | MemorylessFlag, 2) :
+			FRDGTextureDesc::Create2D(Config.Extent, DepthAuxFormat, FClearValueBinding(FarDepthColor), TexCreate_RenderTargetable | TexCreate_ShaderResource | TexCreate_InputAttachmentRead| MemorylessFlag);
 		Desc.NumSamples = Config.NumSamples;
 		SceneTextures.DepthAux = CreateTextureMSAA(GraphBuilder, Desc, TEXT("SceneDepthAux"));
 	}
