Fixed: Static analyzer warnings in OpenGLRenderTarget.cpp

diff --git a/Engine/Source/Runtime/OpenGLDrv/Private/OpenGLRenderTarget.cpp b/Engine/Source/Runtime/OpenGLDrv/Private/OpenGLRenderTarget.cpp
--- a/Engine/Source/Runtime/OpenGLDrv/Private/OpenGLRenderTarget.cpp
+++ b/Engine/Source/Runtime/OpenGLDrv/Private/OpenGLRenderTarget.cpp
@@ -151,8 +151,10 @@ GLuint FOpenGLDynamicRHI::GetOpenGLFramebuffer(uint32 NumSimultaneousRenderTarge
 		return FramebufferRet-1;
 	}
 
+	const bool bRenderTargetsDefined = (RenderTargets != nullptr) && RenderTargets[0];
+
 	// Check for rendering to screen back buffer.
-	if(0 < NumSimultaneousRenderTargets && RenderTargets[0] && RenderTargets[0]->GetResource() == GL_NONE)
+	if (NumSimultaneousRenderTargets > 0 && bRenderTargetsDefined && RenderTargets[0]->GetResource() == GL_NONE)
 	{
 		// Use the default framebuffer (screen back/depth buffer)
 		return GL_NONE;
@@ -170,7 +172,6 @@ GLuint FOpenGLDynamicRHI::GetOpenGLFramebuffer(uint32 NumSimultaneousRenderTarge
 	// Allocate mobile multi-view frame buffer if enabled and supported.
 	// Multi-view doesn't support read buffers, explicitly disable and only bind GL_DRAW_FRAMEBUFFER
 	// TODO: We can't reliably use packed depth stencil?
-	const bool bRenderTargetsDefined = (RenderTargets != nullptr) && RenderTargets[0];
 	const bool bValidMultiViewDepthTarget = !DepthStencilTarget || DepthStencilTarget->Target == GL_TEXTURE_2D_ARRAY;
 	const bool bUsingArrayTextures = (bRenderTargetsDefined) ? (RenderTargets[0]->Target == GL_TEXTURE_2D_ARRAY && bValidMultiViewDepthTarget) : false;
 	const bool bMultiViewCVar = CVarMobileMultiView && CVarMobileMultiView->GetValueOnAnyThread() != 0;
@@ -215,7 +216,7 @@ GLuint FOpenGLDynamicRHI::GetOpenGLFramebuffer(uint32 NumSimultaneousRenderTarge
 	}
 	
 	int32 FirstNonzeroRenderTarget = -1;
-	for (int32 RenderTargetIndex = NumSimultaneousRenderTargets - 1; RenderTargetIndex >= 0; --RenderTargetIndex)
+	for (int32 RenderTargetIndex = NumSimultaneousRenderTargets - 1; RenderTargetIndex >= 0 && bRenderTargetsDefined; --RenderTargetIndex)
 	{
 		FOpenGLTexture* RenderTarget = RenderTargets[RenderTargetIndex];
 		if (!RenderTarget)
