Test if UB is not empty before shader parameter validation in OpenGL RHI

diff --git a/Engine/Source/Runtime/OpenGLDrv/Private/OpenGLUniformBuffer.cpp b/Engine/Source/Runtime/OpenGLDrv/Private/OpenGLUniformBuffer.cpp
--- a/Engine/Source/Runtime/OpenGLDrv/Private/OpenGLUniformBuffer.cpp
+++ b/Engine/Source/Runtime/OpenGLDrv/Private/OpenGLUniformBuffer.cpp
@@ -652,7 +652,7 @@ FUniformBufferRHIRef FOpenGLDynamicRHI::RHICreateUniformBuffer(const void* Conte
 	// Explicitly check that the size is nonzero before allowing CreateBuffer to opaquely fail.
 	check(Layout->Resources.Num() > 0 || Layout->ConstantBufferSize > 0);
 	
-	if (Validation == EUniformBufferValidation::ValidateResources)
+	if (Contents && Validation == EUniformBufferValidation::ValidateResources)
 	{
 		ValidateShaderParameterResourcesRHI(Contents, *Layout);
 	}
