Only uses setBuffersTransform on android 10 and above devices.

diff --git a/Engine/Config/BaseDeviceProfiles.ini b/Engine/Config/BaseDeviceProfiles.ini
--- a/Engine/Config/BaseDeviceProfiles.ini
+++ b/Engine/Config/BaseDeviceProfiles.ini
@@ -736,6 +736,8 @@ BaseProfileName=Android_Mid
 [Android_Adreno6xx DeviceProfile]
 DeviceType=Android
 BaseProfileName=Android_High
+; seems like all Adreno6xx devices properly support setBuffersTransform
++CVars=r.Android.GLESFlipYMethod=1
 
 [Android_Adreno6xx_Vulkan DeviceProfile]
 DeviceType=Android
diff --git a/Engine/Source/Runtime/OpenGLDrv/Private/Android/AndroidEGL.cpp b/Engine/Source/Runtime/OpenGLDrv/Private/Android/AndroidEGL.cpp
--- a/Engine/Source/Runtime/OpenGLDrv/Private/Android/AndroidEGL.cpp
+++ b/Engine/Source/Runtime/OpenGLDrv/Private/Android/AndroidEGL.cpp
@@ -1186,10 +1186,13 @@ void AndroidEGL::UpdateBuffersTransform()
 
 bool AndroidEGL::IsOfflineSurfaceRequired()
 {
-	return FAndroidMisc::SupportsBackbufferSampling() ||
-		// setBuffersTransform doesn't work on arm and powerVR GPU devices, uses BlitFrameBuffer instead
-		(CVarAndroidGLESFlipYMethod.GetValueOnAnyThread() == 0 && (GRHIVendorId == 0x13B5 || GRHIVendorId == 0x1010)) ||
-		CVarAndroidGLESFlipYMethod.GetValueOnAnyThread() == 2;
+	return FAndroidMisc::SupportsBackbufferSampling()
+		// force to use BlitFrameBuffer
+		|| CVarAndroidGLESFlipYMethod.GetValueOnAnyThread() == 2
+		// setBuffersTransform doesn't work on android 9 and below devices
+		|| !(CVarAndroidGLESFlipYMethod.GetValueOnAnyThread() == 1 || FAndroidMisc::GetAndroidMajorVersion() >= 10)
+		// setBuffersTransform doesn't work on arm and powerVR GPU devices
+		|| (CVarAndroidGLESFlipYMethod.GetValueOnAnyThread() == 0 && (GRHIVendorId == 0x13B5 || GRHIVendorId == 0x1010));
 }
 
 ///
