Fix broken SRVs on OGL

diff --git a/Engine/Source/Runtime/OpenGLDrv/Public/OpenGLResources.h b/Engine/Source/Runtime/OpenGLDrv/Public/OpenGLResources.h
--- a/Engine/Source/Runtime/OpenGLDrv/Public/OpenGLResources.h
+++ b/Engine/Source/Runtime/OpenGLDrv/Public/OpenGLResources.h
@@ -442,7 +442,7 @@ public:
 	{
 		RealSize = ResourceSize ? ResourceSize : InSize;
 
-		if (!RHICmdList || RHICmdList->IsBottomOfPipe())
+		if (ShouldRunGLRenderContextOpOnThisThread(RHICmdList))
 		{
 			CreateGLBuffer(InData, ResourceToUse, ResourceSize);
 		}
diff --git a/Engine/Source/Runtime/OpenGLDrv/Public/OpenGLUtil.h b/Engine/Source/Runtime/OpenGLDrv/Public/OpenGLUtil.h
--- a/Engine/Source/Runtime/OpenGLDrv/Public/OpenGLUtil.h
+++ b/Engine/Source/Runtime/OpenGLDrv/Public/OpenGLUtil.h
@@ -156,10 +156,40 @@ struct FRHICommandGLCommand final : public FRHICommand<FRHICommandGLCommand, FRH
 };
 
 void RunOnGLRenderContextThread(TUniqueFunction<void(void)> GLFunc, bool bWaitForCompletion = false);
+
+
+#if 0
+// TODO: investigate simplifying GL's explicit RHIT tests to this.
+inline bool ShouldRunGLRenderContextOpOnThisThread(FRHICommandListBase* RHICmdList)
+{
+	return !RHICmdList || RHICmdList->IsBottomOfPipe();
+}
+
 inline bool ShouldRunGLRenderContextOpOnThisThread(FRHICommandListBase& RHICmdList)
 {
-	return (RHICmdList.Bypass() || !IsRunningRHIInSeparateThread() || IsInRHIThread());
+	return ShouldRunGLRenderContextOpOnThisThread(&RHICmdList);
 }
+#else
+
+inline bool ShouldRunGLRenderContextOpOnThisThread(FRHICommandListBase* RHICmdList)
+{
+	if (RHICmdList)
+	{
+		return (RHICmdList->Bypass() || RHICmdList->IsBottomOfPipe() || !IsRunningRHIInSeparateThread() || IsInRHIThread());
+	}
+	else
+	{
+		check(!IsRunningRHIInSeparateThread() || IsInRHIThread());
+		return true;
+	}
+}
+
+inline bool ShouldRunGLRenderContextOpOnThisThread(FRHICommandListBase& RHICmdList)
+{
+	return ShouldRunGLRenderContextOpOnThisThread(&RHICmdList);
+}
+
+#endif
 
 
 // Program Binary helpers.
