Fix a crash bug on Android 9 adreno 630 devices caused by rendering to a render target and sampling it in a shader at the same time but with different mips.

diff --git a/Engine/Source/Runtime/OpenGLDrv/Private/OpenGLCommands.cpp b/Engine/Source/Runtime/OpenGLDrv/Private/OpenGLCommands.cpp
--- a/Engine/Source/Runtime/OpenGLDrv/Private/OpenGLCommands.cpp
+++ b/Engine/Source/Runtime/OpenGLDrv/Private/OpenGLCommands.cpp
@@ -511,22 +511,49 @@ void FOpenGLDynamicRHI::CachedSetupTextureStageInner(FOpenGLContextState& Contex
 		if (!bSameLimitMip || !bSameNumMips)
 		{
 			DETAILED_QUICK_SCOPE_CYCLE_COUNTER(STAT_CachedSetupTextureStage_TexParameter);
-			if (!bSameLimitMip)
-			{
-				FOpenGL::TexParameter(Target, GL_TEXTURE_BASE_LEVEL, BaseMip);
-			}
-			if (!bSameNumMips)
+			
+			bool bBoundAsRenderTarget = false;
+			for (uint32 RenderTargetIndex = 0; RenderTargetIndex < MaxSimultaneousRenderTargets; ++RenderTargetIndex)
 			{
-				FOpenGL::TexParameter(Target, GL_TEXTURE_MAX_LEVEL, MaxMip);
+				if (PendingState.RenderTargets[RenderTargetIndex] == 0)
+				{
+					break;
+				}
+				else
+				{
+					if (PendingState.RenderTargets[RenderTargetIndex]->GetResource() == Resource)
+					{
+						bBoundAsRenderTarget = true;
+						break;
+					}
+				}
 			}
-			if (MipLimits)
+
+			// If a SRV is bound as render target, skip the BASE_LEVEL and MAX_LEVEL settings because it would cause crash on some android devices.
+			if (!bBoundAsRenderTarget)
 			{
-				MipLimits->Key = BaseMip;
-				MipLimits->Value = MaxMip;
+				if (!bSameLimitMip)
+				{
+					FOpenGL::TexParameter(Target, GL_TEXTURE_BASE_LEVEL, BaseMip);
+				}
+				if (!bSameNumMips)
+				{
+					FOpenGL::TexParameter(Target, GL_TEXTURE_MAX_LEVEL, MaxMip);
+				}
+				if (MipLimits)
+				{
+					MipLimits->Key = BaseMip;
+					MipLimits->Value = MaxMip;
+				}
+				else
+				{
+					TextureMipLimits.Add(Resource, TPair<GLenum, GLenum>(BaseMip, MaxMip));
+				}
 			}
 			else
 			{
-				TextureMipLimits.Add(Resource, TPair<GLenum, GLenum>(BaseMip, MaxMip));
+				LimitMip = 0;
+				NumMips = 0;
 			}
 		}
 	}
