Add LinearFilter texture capability flag to fix an issue that R16G16_UNORM doesn't work with linear filter on some mali devices.

diff --git a/Engine/Source/Runtime/D3D12RHI/Private/D3D12Device.cpp b/Engine/Source/Runtime/D3D12RHI/Private/D3D12Device.cpp
--- a/Engine/Source/Runtime/D3D12RHI/Private/D3D12Device.cpp
+++ b/Engine/Source/Runtime/D3D12RHI/Private/D3D12Device.cpp
@@ -242,7 +242,7 @@ void FD3D12Device::SetupAfterDeviceCreation()
 				ConvertCap1(FormatSupport, EPixelFormatCapabilities::DepthStencil,        D3D12_FORMAT_SUPPORT1_DEPTH_STENCIL);
 				ConvertCap1(FormatSupport, EPixelFormatCapabilities::TextureMipmaps,      D3D12_FORMAT_SUPPORT1_MIP);
 				ConvertCap1(SRVFormatSupport, EPixelFormatCapabilities::TextureLoad,      D3D12_FORMAT_SUPPORT1_SHADER_LOAD);
-				ConvertCap1(SRVFormatSupport, EPixelFormatCapabilities::TextureSample,    D3D12_FORMAT_SUPPORT1_SHADER_SAMPLE);
+				ConvertCap1(SRVFormatSupport, EPixelFormatCapabilities::TextureSample | EPixelFormatCapabilities::TextureFilterable,    D3D12_FORMAT_SUPPORT1_SHADER_SAMPLE);
 				ConvertCap1(SRVFormatSupport, EPixelFormatCapabilities::TextureGather,    D3D12_FORMAT_SUPPORT1_SHADER_GATHER);
 				ConvertCap2(UAVFormatSupport, EPixelFormatCapabilities::TextureAtomics,   D3D12_FORMAT_SUPPORT2_UAV_ATOMIC_EXCHANGE);
 				ConvertCap1(SRVFormatSupport, EPixelFormatCapabilities::TextureBlendable, D3D12_FORMAT_SUPPORT1_BLENDABLE);
diff --git a/Engine/Source/Runtime/OpenGLDrv/Private/OpenGLDevice.cpp b/Engine/Source/Runtime/OpenGLDrv/Private/OpenGLDevice.cpp
--- a/Engine/Source/Runtime/OpenGLDrv/Private/OpenGLDevice.cpp
+++ b/Engine/Source/Runtime/OpenGLDrv/Private/OpenGLDevice.cpp
@@ -493,12 +493,190 @@ static void APIENTRY OpenGLDebugMessageCallbackAMD(
 PFNWGLSWAPINTERVALEXTPROC wglSwapIntervalEXT_ProcAddress = NULL;
 #endif
 
+// Information from OpenGLES specification 3.2 Table 8.10
+// Qualcomm Adreno OpenGLES Developer Guide 4.1.4 section and Table 9-1
+enum class EOpenGLFormatCapabilities : uint32
+{
+	None = 0,
+	Texture = 1ull << 0,
+	Render = 1ull << 1,
+	Filterable = 1ull << 2,
+	Image = 1ull << 3,
+	DepthStencil = 1ull << 4,
+};
+ENUM_CLASS_FLAGS(EOpenGLFormatCapabilities);
+
+static EOpenGLFormatCapabilities GetOpenGLFormatCapabilities(const FOpenGLTextureFormat& GLFormat)
+{
+	EOpenGLFormatCapabilities Capabilities = EOpenGLFormatCapabilities::None;
+	
+	switch (GLFormat.InternalFormat[0])
+	{
+	default: checkNoEntry();
+	case GL_NONE:
+		break;
+
+	case GL_R8:
+	case GL_RG8:
+	case GL_RGB8:
+	case GL_RGBA4:
+	case GL_RGB10_A2:
+	case GL_SRGB8_ALPHA8:
+	case GL_R16F:
+	case GL_RG16F:
+	case GL_RGBA16:
+	case GL_RGB565:
+	case GL_RGB5_A1:
+	case GL_R11F_G11F_B10F:
+		Capabilities |= EOpenGLFormatCapabilities::Texture | EOpenGLFormatCapabilities::Render | EOpenGLFormatCapabilities::Filterable;
+		break;
+
+	case GL_RGBA8:
+	case GL_RGBA16F:
+		Capabilities |= EOpenGLFormatCapabilities::Texture | EOpenGLFormatCapabilities::Render | EOpenGLFormatCapabilities::Filterable | EOpenGLFormatCapabilities::Image;
+		break;
+
+	case GL_R16:
+	case GL_RG16:
+	case GL_RGB10_A2UI:
+	case GL_RG32F:
+	case GL_R8I:
+	case GL_R8UI:
+	case GL_R16I:
+	case GL_R16UI:
+	case GL_RG8I:
+	case GL_RG8UI:
+	case GL_RG16I:
+	case GL_RG16UI:
+	case GL_RG32I:
+	case GL_RG32UI:
+		Capabilities |= EOpenGLFormatCapabilities::Texture | EOpenGLFormatCapabilities::Render;
+		break;
+
+	case GL_R32I:
+	case GL_R32F:
+	case GL_R32UI:
+	case GL_RGBA8I:
+	case GL_RGBA8UI:
+	case GL_RGBA16I:
+	case GL_RGBA16UI:
+	case GL_RGBA32I:
+	case GL_RGBA32F:
+	case GL_RGBA32UI:
+		Capabilities |= EOpenGLFormatCapabilities::Texture | EOpenGLFormatCapabilities::Render | EOpenGLFormatCapabilities::Image;
+		break;
+
+	case GL_RGB32F:
+	case GL_RGB8I:
+	case GL_RGB8UI:
+	case GL_RGB16I:
+	case GL_RGB16UI:
+	case GL_RGB32I:
+	case GL_RGB32UI:
+		Capabilities |= EOpenGLFormatCapabilities::Texture;
+		break;
+
+	case GL_DEPTH24_STENCIL8:
+	case GL_DEPTH_COMPONENT16:
+	case GL_DEPTH_COMPONENT24:
+		Capabilities |= EOpenGLFormatCapabilities::Texture | EOpenGLFormatCapabilities::Render | EOpenGLFormatCapabilities::Filterable | EOpenGLFormatCapabilities::DepthStencil;
+		break;
+
+#if PLATFORM_ANDROID
+	case GL_COMPRESSED_RGB8_ETC2:
+	case GL_COMPRESSED_RGBA8_ETC2_EAC:
+	case GL_COMPRESSED_R11_EAC:
+	case GL_COMPRESSED_RG11_EAC:
+#endif
+
+#if PLATFORM_DESKTOP
+	case GL_COMPRESSED_RG_RGTC2:
+	case GL_COMPRESSED_RED_RGTC1:
+#endif
+
+	case GL_COMPRESSED_RGBA_S3TC_DXT1_EXT:
+	case GL_COMPRESSED_RGBA_S3TC_DXT3_EXT:
+	case GL_COMPRESSED_RGBA_S3TC_DXT5_EXT:
+
+	case GL_COMPRESSED_RGBA_ASTC_4x4_KHR:
+	case GL_COMPRESSED_RGBA_ASTC_6x6_KHR:
+	case GL_COMPRESSED_RGBA_ASTC_8x8_KHR:
+	case GL_COMPRESSED_RGBA_ASTC_10x10_KHR:
+	case GL_COMPRESSED_RGBA_ASTC_12x12_KHR:
+
+	case GL_R8_SNORM:
+	case GL_RG8_SNORM:
+	case GL_RGB8_SNORM:
+	case GL_RGB9_E5:
+	case GL_RGB16F:
+		Capabilities |= EOpenGLFormatCapabilities::Texture | EOpenGLFormatCapabilities::Filterable;
+		break;
+
+	case GL_RGBA8_SNORM:
+		Capabilities |= EOpenGLFormatCapabilities::Texture | EOpenGLFormatCapabilities::Filterable | EOpenGLFormatCapabilities::Image;
+		break;
+	}
+
+	return Capabilities;
+}
 
 static inline void SetupTextureFormat( EPixelFormat Format, const FOpenGLTextureFormat& GLFormat)
 {
 	GOpenGLTextureFormats[Format] = GLFormat;
 	GPixelFormats[Format].Supported = (GLFormat.Format != GL_NONE && (GLFormat.InternalFormat[0] != GL_NONE || GLFormat.InternalFormat[1] != GL_NONE));
 	GPixelFormats[Format].PlatformFormat = GLFormat.InternalFormat[0];
+	
+	if (GPixelFormats[Format].Supported)
+	{
+		EPixelFormatCapabilities& Capabilities = GPixelFormats[Format].Capabilities;
+		EOpenGLFormatCapabilities OpenGLCapabilities = GetOpenGLFormatCapabilities(GLFormat);
+
+		if (EnumHasAllFlags(OpenGLCapabilities, EOpenGLFormatCapabilities::Texture))
+		{
+			EnumAddFlags(Capabilities, EPixelFormatCapabilities::Texture1D);
+			EnumAddFlags(Capabilities, EPixelFormatCapabilities::Texture2D);
+			EnumAddFlags(Capabilities, EPixelFormatCapabilities::Texture3D);
+			EnumAddFlags(Capabilities, EPixelFormatCapabilities::TextureCube);
+
+			EnumAddFlags(Capabilities, EPixelFormatCapabilities::TextureMipmaps);
+			EnumAddFlags(Capabilities, EPixelFormatCapabilities::TextureLoad);
+			EnumAddFlags(Capabilities, EPixelFormatCapabilities::TextureSample);
+			EnumAddFlags(Capabilities, EPixelFormatCapabilities::TextureGather);
+		}
+
+		if (EnumHasAllFlags(OpenGLCapabilities, EOpenGLFormatCapabilities::Filterable))
+		{
+			EnumAddFlags(Capabilities, EPixelFormatCapabilities::TextureFilterable);
+		}
+
+		if (EnumHasAllFlags(OpenGLCapabilities, EOpenGLFormatCapabilities::Render))
+		{
+			EnumAddFlags(Capabilities, EPixelFormatCapabilities::RenderTarget);
+			EnumAddFlags(Capabilities, EPixelFormatCapabilities::TextureBlendable);
+		}
+
+		if (EnumHasAllFlags(OpenGLCapabilities, EOpenGLFormatCapabilities::DepthStencil))
+		{
+			EnumAddFlags(Capabilities, EPixelFormatCapabilities::DepthStencil);
+		}
+
+		if (EnumHasAllFlags(OpenGLCapabilities, EOpenGLFormatCapabilities::Image))
+		{
+			EnumAddFlags(Capabilities, EPixelFormatCapabilities::TextureAtomics);
+			EnumAddFlags(Capabilities, EPixelFormatCapabilities::TextureStore);
+
+			EnumAddFlags(Capabilities, EPixelFormatCapabilities::Buffer);
+			EnumAddFlags(Capabilities, EPixelFormatCapabilities::VertexBuffer);
+			EnumAddFlags(Capabilities, EPixelFormatCapabilities::IndexBuffer);
+			EnumAddFlags(Capabilities, EPixelFormatCapabilities::BufferLoad);
+			EnumAddFlags(Capabilities, EPixelFormatCapabilities::BufferStore);
+			EnumAddFlags(Capabilities, EPixelFormatCapabilities::BufferAtomics);
+
+			EnumAddFlags(Capabilities, EPixelFormatCapabilities::UAV);
+			EnumAddFlags(Capabilities, EPixelFormatCapabilities::TypedUAVLoad);
+			EnumAddFlags(Capabilities, EPixelFormatCapabilities::TypedUAVStore);
+		}
+	}
 }
 
 
diff --git a/Engine/Source/Runtime/RHI/Private/RHI.cpp b/Engine/Source/Runtime/RHI/Private/RHI.cpp
--- a/Engine/Source/Runtime/RHI/Private/RHI.cpp
+++ b/Engine/Source/Runtime/RHI/Private/RHI.cpp
@@ -2816,7 +2816,7 @@ void RHIInitDefaultPixelFormatCapabilities()
 			if (IsBlockCompressedFormat(PixelFormat))
 			{
 				// Block compressed formats should have limited capabilities
-				EnumAddFlags(Info.Capabilities, EPixelFormatCapabilities::AnyTexture | EPixelFormatCapabilities::TextureMipmaps | EPixelFormatCapabilities::TextureLoad | EPixelFormatCapabilities::TextureSample | EPixelFormatCapabilities::TextureGather);
+				EnumAddFlags(Info.Capabilities, EPixelFormatCapabilities::AnyTexture | EPixelFormatCapabilities::TextureMipmaps | EPixelFormatCapabilities::TextureLoad | EPixelFormatCapabilities::TextureSample | EPixelFormatCapabilities::TextureGather | EPixelFormatCapabilities::TextureFilterable);
 			}
 			else
 			{
diff --git a/Engine/Source/Runtime/RHI/Public/RHI.h b/Engine/Source/Runtime/RHI/Public/RHI.h
--- a/Engine/Source/Runtime/RHI/Public/RHI.h
+++ b/Engine/Source/Runtime/RHI/Public/RHI.h
@@ -856,6 +856,8 @@ enum class EPixelFormatCapabilities : uint32
     TypedUAVLoad     = 1ull << 21,
 	TypedUAVStore    = 1ull << 22,
 
+	TextureFilterable = 1ull << 23,
+
 	AnyTexture       = Texture1D | Texture2D | Texture3D | TextureCube,
 
 	AllTextureFlags  = AnyTexture | RenderTarget | DepthStencil | TextureMipmaps | TextureLoad | TextureSample | TextureGather | TextureAtomics | TextureBlendable | TextureStore,
diff --git a/Engine/Source/Runtime/Renderer/Private/SystemTextures.cpp b/Engine/Source/Runtime/Renderer/Private/SystemTextures.cpp
--- a/Engine/Source/Runtime/Renderer/Private/SystemTextures.cpp
+++ b/Engine/Source/Runtime/Renderer/Private/SystemTextures.cpp
@@ -409,7 +409,7 @@ void FSystemTextures::InitializeFeatureLevelDependentTextures(FRHICommandListImm
 
 		EPixelFormat Format = PF_R8G8;
 		// for low roughness we would get banding with PF_R8G8 but for low spec it could be used, for now we don't do this optimization
-		if (GPixelFormats[PF_G16R16].Supported)
+		if (GPixelFormats[PF_G16R16].Supported && RHIPixelFormatHasCapabilities(PF_G16R16, EPixelFormatCapabilities::TextureFilterable))
 		{
 			Format = PF_G16R16;
 		}
diff --git a/Engine/Source/Runtime/VulkanRHI/Private/VulkanDevice.cpp b/Engine/Source/Runtime/VulkanRHI/Private/VulkanDevice.cpp
--- a/Engine/Source/Runtime/VulkanRHI/Private/VulkanDevice.cpp
+++ b/Engine/Source/Runtime/VulkanRHI/Private/VulkanDevice.cpp
@@ -874,6 +874,7 @@ void FVulkanDevice::MapImageFormatSupport(FPixelFormatInfo& PixelFormatInfo, con
 		ConvertImageCap(LocalFormatProperties, EPixelFormatCapabilities::AllUAVFlags | EPixelFormatCapabilities::TextureStore, VK_FORMAT_FEATURE_STORAGE_IMAGE_BIT);
 
 		ConvertImageCap(LocalFormatProperties, EPixelFormatCapabilities::TextureAtomics, VK_FORMAT_FEATURE_STORAGE_IMAGE_ATOMIC_BIT);
+		ConvertImageCap(LocalFormatProperties, EPixelFormatCapabilities::TextureFilterable, VK_FORMAT_FEATURE_SAMPLED_IMAGE_FILTER_LINEAR_BIT);
 
 		if (EnumHasAllFlags(Capabilities, EPixelFormatCapabilities::AnyTexture))
 		{
diff --git a/Engine/Source/Runtime/Windows/D3D11RHI/Private/D3D11Device.cpp b/Engine/Source/Runtime/Windows/D3D11RHI/Private/D3D11Device.cpp
--- a/Engine/Source/Runtime/Windows/D3D11RHI/Private/D3D11Device.cpp
+++ b/Engine/Source/Runtime/Windows/D3D11RHI/Private/D3D11Device.cpp
@@ -466,7 +466,7 @@ void FD3D11DynamicRHI::SetupAfterDeviceCreation()
                 ConvertCap1(FormatSupport, EPixelFormatCapabilities::DepthStencil,        D3D11_FORMAT_SUPPORT_DEPTH_STENCIL);
                 ConvertCap1(FormatSupport, EPixelFormatCapabilities::TextureMipmaps,      D3D11_FORMAT_SUPPORT_MIP);
                 ConvertCap1(SRVFormatSupport, EPixelFormatCapabilities::TextureLoad,      D3D11_FORMAT_SUPPORT_SHADER_LOAD);
-                ConvertCap1(SRVFormatSupport, EPixelFormatCapabilities::TextureSample,    D3D11_FORMAT_SUPPORT_SHADER_SAMPLE);
+                ConvertCap1(SRVFormatSupport, EPixelFormatCapabilities::TextureSample | EPixelFormatCapabilities::TextureFilterable,    D3D11_FORMAT_SUPPORT_SHADER_SAMPLE);
                 ConvertCap1(SRVFormatSupport, EPixelFormatCapabilities::TextureGather,    D3D11_FORMAT_SUPPORT_SHADER_GATHER);
                 ConvertCap2(UAVFormatSupport, EPixelFormatCapabilities::TextureAtomics,   D3D11_FORMAT_SUPPORT2_UAV_ATOMIC_EXCHANGE);
                 ConvertCap1(RTVFormatSupport, EPixelFormatCapabilities::TextureBlendable, D3D11_FORMAT_SUPPORT_BLENDABLE);
