Fix a compilation error on Android OpenGL caused by CL#22451147.

diff --git a/Engine/Source/Runtime/OpenGLDrv/Private/OpenGLDevice.cpp b/Engine/Source/Runtime/OpenGLDrv/Private/OpenGLDevice.cpp
--- a/Engine/Source/Runtime/OpenGLDrv/Private/OpenGLDevice.cpp
+++ b/Engine/Source/Runtime/OpenGLDrv/Private/OpenGLDevice.cpp
@@ -537,7 +537,6 @@ static EOpenGLFormatCapabilities GetOpenGLFormatCapabilities(const FOpenGLTextur
 		break;
 
 	case GL_R16:
-	case GL_RG16:
 	case GL_RGB10_A2UI:
 	case GL_RG32F:
 	case GL_R8I:
