Fix texture format assert on GL mobile preview start

diff --git a/Engine/Source/Runtime/OpenGLDrv/Private/OpenGLDevice.cpp b/Engine/Source/Runtime/OpenGLDrv/Private/OpenGLDevice.cpp
--- a/Engine/Source/Runtime/OpenGLDrv/Private/OpenGLDevice.cpp
+++ b/Engine/Source/Runtime/OpenGLDrv/Private/OpenGLDevice.cpp
@@ -591,6 +591,8 @@ static EOpenGLFormatCapabilities GetOpenGLFormatCapabilities(const FOpenGLTextur
 #if PLATFORM_DESKTOP
 	case GL_COMPRESSED_RG_RGTC2:
 	case GL_COMPRESSED_RED_RGTC1:
+	case GL_RG16:
+	case GL_RG16_SNORM:
 #endif
 
 	case GL_COMPRESSED_RGBA_S3TC_DXT1_EXT:
