Fixed: UAV limitations check for GL programs

diff --git a/Engine/Source/Runtime/OpenGLDrv/Private/OpenGLShaders.cpp b/Engine/Source/Runtime/OpenGLDrv/Private/OpenGLShaders.cpp
--- a/Engine/Source/Runtime/OpenGLDrv/Private/OpenGLShaders.cpp
+++ b/Engine/Source/Runtime/OpenGLDrv/Private/OpenGLShaders.cpp
@@ -2110,9 +2110,6 @@ void FOpenGLLinkedProgram::ConfigureShaderStage( int Stage, uint32 FirstUniformB
 		FOpenGL::GetFirstComputeUAVUnit()
 	};
 	
-	// verify that only CS and PS uses UAVs
-	check(!(Stage == CrossCompiler::SHADER_STAGE_COMPUTE || Stage == CrossCompiler::SHADER_STAGE_PIXEL) ? (CountSetBits(UAVStageNeeds) == 0) : true);
-
 	SCOPE_CYCLE_COUNTER(STAT_OpenGLShaderBindParameterTime);
 	VERIFY_GL_SCOPE();
 
@@ -2281,6 +2278,9 @@ void FOpenGLLinkedProgram::ConfigureShaderStage( int Stage, uint32 FirstUniformB
 		{
 			// compute shaders have layout(binding) for images
 			// glUniform1i(Location, FirstUAVUnit[Stage] + UAVIndex);
+
+			// verify that only CS and PS uses UAVs (limitation on MALI GPUs)
+			check(Stage == CrossCompiler::SHADER_STAGE_COMPUTE || Stage == CrossCompiler::SHADER_STAGE_PIXEL);
 			
 			UAVStageNeeds[ FirstUAVUnit[Stage] + UAVIndex ] = true;
 			MaxUAVUnitUsed = FMath::Max(MaxUAVUnitUsed, FirstUAVUnit[Stage] + UAVIndex);
