Fix RT / RHIT stats
 - The RHI and render threads are now blocking on FGraphEvents in various platform RHI implementations, but the timing stats did not take this time into account.
 - Some time was double counted (e.g. on the render thread, due to a FRenderThreadIdleScope plus a nested call to FEvent::Wait which includes a per-thread wait stat).

 - The fix is to remove use of GRenderThreadIdle / FRenderThreadIdleScope etc in places where the thread is going to wait on an FGraphEvent, to avoid the double counting. This fixes the odd spikes on the RHIT in CSV profiles.
 - Also removed the GWorkingRHIThreadStallTime global. RHI idle time is now tracked using standard thread stats (FThreadIdleStats::FScopeIdle, which is included in the FRenderThreadIdleScope when called from the RHI thread).

 - Simplified copy/paste instances of GRenderThreadIdle and GRenderThreadNumIdle by replacing them with a FRenderThreadIdleScope across platform RHIs.

After this change, the RT and RHIT unit stats are stable in Lyra game, when clicking away from the window to force the title into a low framerate mode. Before this change, the RT/RHIT stats would change significantly.

diff --git a/Engine/Source/Runtime/Apple/MetalRHI/Private/MetalRHIRenderQuery.cpp b/Engine/Source/Runtime/Apple/MetalRHI/Private/MetalRHIRenderQuery.cpp
--- a/Engine/Source/Runtime/Apple/MetalRHI/Private/MetalRHIRenderQuery.cpp
+++ b/Engine/Source/Runtime/Apple/MetalRHI/Private/MetalRHIRenderQuery.cpp
@@ -336,9 +336,10 @@ bool FMetalRHIRenderQuery::GetResult(uint64& OutNumPixels, bool bWait, uint32 GP
 				FRHICommandListExecutor::GetImmediateCommandList().ImmediateFlush(EImmediateFlushType::FlushRHIThread);
 			}
 
-			uint32 IdleStart = FPlatformTime::Cycles();
-
-			bOK = Buffer.Wait(WaitMS);
+			{
+				FRenderThreadIdleScope IdleScope(ERenderThreadIdleTypes::WaitingForGPUQuery);
+				bOK = Buffer.Wait(WaitMS);
+			}
 
 			// Result is written in one of potentially many command buffer completion handlers
 			// But the command buffer wait above may return before the query completion handler fires
@@ -348,16 +349,6 @@ bool FMetalRHIRenderQuery::GetResult(uint64& OutNumPixels, bool bWait, uint32 GP
 				QueryWrittenEvent->Wait();
 			}
 
-			if (IsInRHIThread())
-			{
-				GWorkingRHIThreadStallTime += FPlatformTime::Cycles() - IdleStart;
-			}
-			else
-			{
-				GRenderThreadIdle[ERenderThreadIdleTypes::WaitingForGPUQuery] += FPlatformTime::Cycles() - IdleStart;
-				GRenderThreadNumIdle[ERenderThreadIdleTypes::WaitingForGPUQuery]++;
-			}
-
 			// Never wait for a failed signal again.
 			bAvailable = Buffer.bCompleted;
 		}
diff --git a/Engine/Source/Runtime/Apple/MetalRHI/Private/MetalViewport.cpp b/Engine/Source/Runtime/Apple/MetalRHI/Private/MetalViewport.cpp
--- a/Engine/Source/Runtime/Apple/MetalRHI/Private/MetalViewport.cpp
+++ b/Engine/Source/Runtime/Apple/MetalRHI/Private/MetalViewport.cpp
@@ -308,54 +308,46 @@ id<CAMetalDrawable> FMetalViewport::GetDrawable(EMetalViewportAccessFlag Accesso
 
 		@autoreleasepool
 		{
-			uint32 IdleStart = FPlatformTime::Cycles();
+			{
+				FRenderThreadIdleScope IdleScope(ERenderThreadIdleTypes::WaitingForGPUPresent);
 
 #if PLATFORM_MAC
-			CAMetalLayer* CurrentLayer = (CAMetalLayer*)[View layer];
-			if (GMetalNonBlockingPresent == 0 || [((id<CAMetalLayerSPI>)CurrentLayer) isDrawableAvailable])
-			{
-				Drawable = CurrentLayer ? [CurrentLayer nextDrawable] : nil;
-			}
+				CAMetalLayer* CurrentLayer = (CAMetalLayer*)[View layer];
+				if (GMetalNonBlockingPresent == 0 || [((id<CAMetalLayerSPI>)CurrentLayer) isDrawableAvailable])
+				{
+					Drawable = CurrentLayer ? [CurrentLayer nextDrawable] : nil;
+				}
 
 #if METAL_DEBUG_OPTIONS
-			if (Drawable)
-			{
-				CGSize Size = Drawable.layer.drawableSize;
-				if ((Size.width != BackBuffer[GetViewportIndex(Accessor)]->GetSizeX() || Size.height != BackBuffer[GetViewportIndex(Accessor)]->GetSizeY()))
+				if (Drawable)
 				{
-					UE_LOG(LogMetal, Display, TEXT("Viewport Size Mismatch: Drawable W:%f H:%f, Viewport W:%u H:%u"), Size.width, Size.height, BackBuffer[GetViewportIndex(Accessor)]->GetSizeX(), BackBuffer[GetViewportIndex(Accessor)]->GetSizeY());
+					CGSize Size = Drawable.layer.drawableSize;
+					if ((Size.width != BackBuffer[GetViewportIndex(Accessor)]->GetSizeX() || Size.height != BackBuffer[GetViewportIndex(Accessor)]->GetSizeY()))
+					{
+						UE_LOG(LogMetal, Display, TEXT("Viewport Size Mismatch: Drawable W:%f H:%f, Viewport W:%u H:%u"), Size.width, Size.height, BackBuffer[GetViewportIndex(Accessor)]->GetSizeX(), BackBuffer[GetViewportIndex(Accessor)]->GetSizeY());
+					}
 				}
-			}
 #endif // METAL_DEBUG_OPTIONS
 
 #else // PLATFORM_MAC
-			CGSize Size;
-			IOSAppDelegate* AppDelegate = [IOSAppDelegate GetDelegate];
-			do
-			{
-				Drawable = [AppDelegate.IOSView MakeDrawable];
-				if (Drawable != nil)
+				CGSize Size;
+				IOSAppDelegate* AppDelegate = [IOSAppDelegate GetDelegate];
+				do
 				{
-					Size.width = Drawable.texture.width;
-					Size.height = Drawable.texture.height;
-				}
-				else
-				{
-					FPlatformProcess::SleepNoStats(0.001f);
+					Drawable = [AppDelegate.IOSView MakeDrawable];
+					if (Drawable != nil)
+					{
+						Size.width = Drawable.texture.width;
+						Size.height = Drawable.texture.height;
+					}
+					else
+					{
+						FPlatformProcess::SleepNoStats(0.001f);
+					}
 				}
-			}
-			while (Drawable == nil || Size.width != BackBuffer[GetViewportIndex(Accessor)]->GetSizeX() || Size.height != BackBuffer[GetViewportIndex(Accessor)]->GetSizeY());
+				while (Drawable == nil || Size.width != BackBuffer[GetViewportIndex(Accessor)]->GetSizeX() || Size.height != BackBuffer[GetViewportIndex(Accessor)]->GetSizeY());
 
 #endif // PLATFORM_MAC
-
-			if (IsInRHIThread())
-			{
-				GWorkingRHIThreadStallTime += FPlatformTime::Cycles() - IdleStart;
-			}
-			else
-			{
-				GRenderThreadIdle[ERenderThreadIdleTypes::WaitingForGPUPresent] += FPlatformTime::Cycles() - IdleStart;
-				GRenderThreadNumIdle[ERenderThreadIdleTypes::WaitingForGPUPresent]++;
 			}
 
 			// Retain the drawable here or it will be released when the
diff --git a/Engine/Source/Runtime/D3D12RHI/Private/D3D12CommandContext.cpp b/Engine/Source/Runtime/D3D12RHI/Private/D3D12CommandContext.cpp
--- a/Engine/Source/Runtime/D3D12RHI/Private/D3D12CommandContext.cpp
+++ b/Engine/Source/Runtime/D3D12RHI/Private/D3D12CommandContext.cpp
@@ -575,8 +575,6 @@ void FD3D12ContextCommon::FlushCommands(ED3D12FlushFlags FlushFlags)
 	if (SubmissionEvent && !SubmissionEvent->IsComplete())
 	{
 		SCOPED_NAMED_EVENT_TEXT("Submission_Wait", FColor::Turquoise);
-
-		FRenderThreadIdleScope Scope(ERenderThreadIdleTypes::WaitingForAllOtherSleep);
 		SubmissionEvent->Wait();
 	}
 }
diff --git a/Engine/Source/Runtime/D3D12RHI/Private/D3D12Submission.cpp b/Engine/Source/Runtime/D3D12RHI/Private/D3D12Submission.cpp
--- a/Engine/Source/Runtime/D3D12RHI/Private/D3D12Submission.cpp
+++ b/Engine/Source/Runtime/D3D12RHI/Private/D3D12Submission.cpp
@@ -853,7 +853,6 @@ void FD3D12DynamicRHI::ProcessInterruptQueueUntil(FGraphEvent* GraphEvent)
 
 	if (GraphEvent && !GraphEvent->IsComplete())
 	{
-		FRenderThreadIdleScope Scope(ERenderThreadIdleTypes::WaitingForGPUQuery);
 		GraphEvent->Wait();
 	}
 
@@ -864,7 +863,7 @@ void FD3D12DynamicRHI::ProcessInterruptQueueUntil(FGraphEvent* GraphEvent)
 	if (!GraphEvent || !GraphEvent->IsComplete())
 	{
 		// If we're waiting for a sync point, accumulate the idle time
-		FRenderThreadIdleScope Scope(ERenderThreadIdleTypes::WaitingForGPUQuery, GraphEvent != nullptr);
+		FThreadIdleStats::FScopeIdle IdleScope(/* bIgnore = */GraphEvent == nullptr);
 
 	Retry:
 		if (InterruptCS.TryLock())
@@ -883,7 +882,7 @@ void FD3D12DynamicRHI::ProcessInterruptQueueUntil(FGraphEvent* GraphEvent)
 		else if (GraphEvent && !GraphEvent->IsComplete())
 		{
 			// Failed to get the lock. Another thread is processing the interrupt queue. Try again...
-			FPlatformProcess::Sleep(0);
+			FPlatformProcess::SleepNoStats(0);
 			goto Retry;
 		}
 	}
diff --git a/Engine/Source/Runtime/D3D12RHI/Private/D3D12Viewport.cpp b/Engine/Source/Runtime/D3D12RHI/Private/D3D12Viewport.cpp
--- a/Engine/Source/Runtime/D3D12RHI/Private/D3D12Viewport.cpp
+++ b/Engine/Source/Runtime/D3D12RHI/Private/D3D12Viewport.cpp
@@ -771,8 +771,6 @@ void FD3D12Viewport::WaitForFrameEventCompletion()
 {
 	if (FrameSyncPoints.Num())
 	{
-		FRenderThreadIdleScope Scope(ERenderThreadIdleTypes::WaitingForGPUPresent);
-
 		for (auto& SyncPoint : FrameSyncPoints)
 		{
 			if (SyncPoint)
diff --git a/Engine/Source/Runtime/Engine/Private/UnrealClient.cpp b/Engine/Source/Runtime/Engine/Private/UnrealClient.cpp
--- a/Engine/Source/Runtime/Engine/Private/UnrealClient.cpp
+++ b/Engine/Source/Runtime/Engine/Private/UnrealClient.cpp
@@ -1488,17 +1488,15 @@ void FViewport::EndRenderFrame(FRHICommandListImmediate& RHICmdList, bool bPrese
 		GEngine->SetPresentLatencyMarkerStart(CurrentFrameCounter);
 	});
 
-	uint32 StartTime = FPlatformTime::Cycles();
-	RHICmdList.EndDrawingViewport(GetViewportRHI(), bPresent, bLockToVsync);
-	uint32 EndTime = FPlatformTime::Cycles();
+	{
+		FRenderThreadIdleScope IdleScope(ERenderThreadIdleTypes::WaitingForGPUPresent);
+		RHICmdList.EndDrawingViewport(GetViewportRHI(), bPresent, bLockToVsync);
+	}
 
 	RHICmdList.EnqueueLambda([CurrentFrameCounter = GFrameCounterRenderThread](FRHICommandListImmediate& InRHICmdList)
 	{
 		GEngine->SetPresentLatencyMarkerEnd(CurrentFrameCounter);
 	});
-
-	GRenderThreadIdle[ERenderThreadIdleTypes::WaitingForGPUPresent] += EndTime - StartTime;
-	GRenderThreadNumIdle[ERenderThreadIdleTypes::WaitingForGPUPresent]++;
 }
 
 FRHIGPUMask FViewport::GetGPUMask(FRHICommandListImmediate& RHICmdList) const
diff --git a/Engine/Source/Runtime/OpenGLDrv/Private/Android/AndroidOpenGL.cpp b/Engine/Source/Runtime/OpenGLDrv/Private/Android/AndroidOpenGL.cpp
--- a/Engine/Source/Runtime/OpenGLDrv/Private/Android/AndroidOpenGL.cpp
+++ b/Engine/Source/Runtime/OpenGLDrv/Private/Android/AndroidOpenGL.cpp
@@ -677,7 +677,8 @@ void FAndroidOpenGL::GetQueryObject(GLuint QueryId, EQueryMode QueryMode, GLuint
 {
 	GLenum QueryName = (QueryMode == QM_Result) ? GL_QUERY_RESULT : GL_QUERY_RESULT_AVAILABLE;
 	VERIFY_GL_SCOPE();
-	uint32 IdleStart = (QueryName == GL_QUERY_RESULT) ? FPlatformTime::Cycles() : 0;
+
+	FRenderThreadIdleScope IdleScope(ERenderThreadIdleTypes::WaitingForGPUQuery, QueryName == GL_QUERY_RESULT);
 
 #if !VIRTUALIZE_QUERIES
 #if CHECK_QUERY_ERRORS
@@ -730,19 +731,6 @@ void FAndroidOpenGL::GetQueryObject(GLuint QueryId, EQueryMode QueryMode, GLuint
 		VirtualResults[QueryId] = *OutResult;
 	}
 #endif
-	if (QueryName == GL_QUERY_RESULT)
-	{
-		uint32 ThisCycles = FPlatformTime::Cycles() - IdleStart;
-		if (IsInRHIThread())
-		{
-			GWorkingRHIThreadStallTime += ThisCycles;
-		}
-		else
-		{
-			GRenderThreadIdle[ERenderThreadIdleTypes::WaitingForGPUQuery] += ThisCycles;
-			GRenderThreadNumIdle[ERenderThreadIdleTypes::WaitingForGPUQuery]++;
-		}
-	}
 
 #if CHECK_QUERY_ERRORS
 	Err = glGetError();
diff --git a/Engine/Source/Runtime/OpenGLDrv/Private/OpenGLQuery.cpp b/Engine/Source/Runtime/OpenGLDrv/Private/OpenGLQuery.cpp
--- a/Engine/Source/Runtime/OpenGLDrv/Private/OpenGLQuery.cpp
+++ b/Engine/Source/Runtime/OpenGLDrv/Private/OpenGLQuery.cpp
@@ -429,7 +429,8 @@ void FOpenGLDynamicRHI::GetRenderQueryResult_OnThisThread(FOpenGLRenderQuery* Qu
 		else if (bWait) // Isn't the query finished yet, and can we wait for it?
 		{
 			SCOPE_CYCLE_COUNTER(STAT_RenderQueryResultTime);
-			uint32 IdleStart = FPlatformTime::Cycles();
+
+			FRenderThreadIdleScope IdleScope(ERenderThreadIdleTypes::WaitingForGPUQuery);
 			GBatcher.Waited();
 			
 			if (GOpenGLPollRenderQueryResult == 0)
@@ -479,17 +480,6 @@ void FOpenGLDynamicRHI::GetRenderQueryResult_OnThisThread(FOpenGLRenderQuery* Qu
 					Query->TotalResults.Increment();
 				}
 			}
-
-			uint32 ThisCycles = FPlatformTime::Cycles() - IdleStart;
-			if (IsInRHIThread())
-			{
-				GWorkingRHIThreadStallTime += ThisCycles;
-			}
-			else
-			{
-				GRenderThreadIdle[ERenderThreadIdleTypes::WaitingForGPUQuery] += ThisCycles;
-				GRenderThreadNumIdle[ERenderThreadIdleTypes::WaitingForGPUQuery]++;
-			}
 		}
 	}
 }
@@ -995,46 +985,43 @@ uint64 FOpenGLBufferedGPUTiming::GetTiming(bool bGetCurrentResultsAndBlock)
 			// This really only happens if occlusion and frame sync event queries are disabled, otherwise those will block until the GPU catches up to 1 frame behind
 			const bool bBlocking = ( NumIssuedTimestamps == BufferSize ) || bGetCurrentResultsAndBlock;
 
-			uint32 IdleStart = FPlatformTime::Cycles();
-			double StartTimeoutTime = FPlatformTime::Seconds();
-
 			GLuint EndAvailable = GL_FALSE;
-
-			SCOPE_CYCLE_COUNTER( STAT_RenderQueryResultTime );
-			// If we are blocking, retry until the GPU processes the time stamp command
-			do
 			{
-				FOpenGL::GetQueryObject(EndTimestamps[TimestampIndex]->Resource, FOpenGL::QM_ResultAvailable, &EndAvailable);
+				FRenderThreadIdleScope IdleScope(ERenderThreadIdleTypes::WaitingForGPUQuery);
+				double StartTimeoutTime = FPlatformTime::Seconds();
 
-				if ((FPlatformTime::Seconds() - StartTimeoutTime) > 0.5)
+				SCOPE_CYCLE_COUNTER( STAT_RenderQueryResultTime );
+				// If we are blocking, retry until the GPU processes the time stamp command
+				do
 				{
-					UE_LOG(LogRHI, Log, TEXT("Timed out while waiting for GPU to catch up. (500 ms) EndTimeStamp"));
-					return 0;
-				}
-			} while ( EndAvailable == GL_FALSE && bBlocking );
+					FOpenGL::GetQueryObject(EndTimestamps[TimestampIndex]->Resource, FOpenGL::QM_ResultAvailable, &EndAvailable);
 
-			GRenderThreadIdle[ERenderThreadIdleTypes::WaitingForGPUQuery] += FPlatformTime::Cycles() - IdleStart;
-			GRenderThreadNumIdle[ERenderThreadIdleTypes::WaitingForGPUQuery]++;
+					if ((FPlatformTime::Seconds() - StartTimeoutTime) > 0.5)
+					{
+						UE_LOG(LogRHI, Log, TEXT("Timed out while waiting for GPU to catch up. (500 ms) EndTimeStamp"));
+						return 0;
+					}
+				} while ( EndAvailable == GL_FALSE && bBlocking );
+			}
 
 			if ( EndAvailable == GL_TRUE )
 			{
-				IdleStart = FPlatformTime::Cycles();
-				StartTimeoutTime = FPlatformTime::Seconds();
-
 				GLuint StartAvailable = GL_FALSE;
-
-				do
 				{
-					FOpenGL::GetQueryObject(StartTimestamps[TimestampIndex]->Resource, FOpenGL::QM_ResultAvailable, &StartAvailable);
+					FRenderThreadIdleScope IdleScope(ERenderThreadIdleTypes::WaitingForGPUQuery);
+					double StartTimeoutTime = FPlatformTime::Seconds();
 
-					if ((FPlatformTime::Seconds() - StartTimeoutTime) > 0.5)
+					do
 					{
-						UE_LOG(LogRHI, Log, TEXT("Timed out while waiting for GPU to catch up. (500 ms) StartTimeStamp"));
-						return 0;
-					}
-				} while ( StartAvailable == GL_FALSE && bBlocking );
+						FOpenGL::GetQueryObject(StartTimestamps[TimestampIndex]->Resource, FOpenGL::QM_ResultAvailable, &StartAvailable);
 
-				GRenderThreadIdle[ERenderThreadIdleTypes::WaitingForGPUQuery] += FPlatformTime::Cycles() - IdleStart;
+						if ((FPlatformTime::Seconds() - StartTimeoutTime) > 0.5)
+						{
+							UE_LOG(LogRHI, Log, TEXT("Timed out while waiting for GPU to catch up. (500 ms) StartTimeStamp"));
+							return 0;
+						}
+					} while ( StartAvailable == GL_FALSE && bBlocking );
+				}
 
 				if(StartAvailable == GL_TRUE)
 				{
diff --git a/Engine/Source/Runtime/OpenGLDrv/Private/OpenGLViewport.cpp b/Engine/Source/Runtime/OpenGLDrv/Private/OpenGLViewport.cpp
--- a/Engine/Source/Runtime/OpenGLDrv/Private/OpenGLViewport.cpp
+++ b/Engine/Source/Runtime/OpenGLDrv/Private/OpenGLViewport.cpp
@@ -125,81 +125,72 @@ void FOpenGLDynamicRHI::RHIEndDrawingViewport(FRHIViewport* ViewportRHI,bool bPr
 	FOpenGLViewport* Viewport = ResourceCast(ViewportRHI);
 
 	SCOPE_CYCLE_COUNTER(STAT_OpenGLPresentTime);
-	uint32 IdleStart = FPlatformTime::Cycles();
-
+	{
+		FRenderThreadIdleScope IdleScope(ERenderThreadIdleTypes::WaitingForGPUPresent);
 
-	check(DrawingViewport.GetReference() == Viewport);
+		check(DrawingViewport.GetReference() == Viewport);
 
-	FOpenGLTexture* BackBuffer = Viewport->GetBackBuffer();
+		FOpenGLTexture* BackBuffer = Viewport->GetBackBuffer();
 
-	FOpenGLContextState& ContextState = GetContextStateForCurrentContext();
+		FOpenGLContextState& ContextState = GetContextStateForCurrentContext();
 
-	if (ContextState.bScissorEnabled)
-	{
-		ContextState.bScissorEnabled = false;
-		glDisable(GL_SCISSOR_TEST);
-	}
+		if (ContextState.bScissorEnabled)
+		{
+			ContextState.bScissorEnabled = false;
+			glDisable(GL_SCISSOR_TEST);
+		}
 
-	bool bNeedFinishFrame = PlatformBlitToViewport(PlatformDevice,
-		*Viewport, 
-		BackBuffer->GetSizeX(),
-		BackBuffer->GetSizeY(),
-		bPresent,
-		bLockToVsync
-	);
+		bool bNeedFinishFrame = PlatformBlitToViewport(PlatformDevice,
+			*Viewport, 
+			BackBuffer->GetSizeX(),
+			BackBuffer->GetSizeY(),
+			bPresent,
+			bLockToVsync
+		);
 
-	// Always consider the Framebuffer in the rendering context dirty after the blit
-	RenderingContextState.Framebuffer = -1;
+		// Always consider the Framebuffer in the rendering context dirty after the blit
+		RenderingContextState.Framebuffer = -1;
 
-	DrawingViewport = NULL;
+		DrawingViewport = NULL;
 
-	// Don't wait on the GPU when using SLI, let the driver determine how many frames behind the GPU should be allowed to get
-	if (GNumAlternateFrameRenderingGroups == 1)
-	{
-		if (bNeedFinishFrame)
+		// Don't wait on the GPU when using SLI, let the driver determine how many frames behind the GPU should be allowed to get
+		if (GNumAlternateFrameRenderingGroups == 1)
 		{
-			static const auto CFinishFrameVar = IConsoleManager::Get().FindTConsoleVariableDataInt(TEXT("r.FinishCurrentFrame"));
-			if (!CFinishFrameVar->GetValueOnRenderThread())
+			if (bNeedFinishFrame)
 			{
-				// Wait for the GPU to finish rendering the previous frame before finishing this frame.
-				Viewport->WaitForFrameEventCompletion();
-				Viewport->IssueFrameEvent();
+				static const auto CFinishFrameVar = IConsoleManager::Get().FindTConsoleVariableDataInt(TEXT("r.FinishCurrentFrame"));
+				if (!CFinishFrameVar->GetValueOnRenderThread())
+				{
+					// Wait for the GPU to finish rendering the previous frame before finishing this frame.
+					Viewport->WaitForFrameEventCompletion();
+					Viewport->IssueFrameEvent();
+				}
+				else
+				{
+					// Finish current frame immediately to reduce latency
+					Viewport->IssueFrameEvent();
+					Viewport->WaitForFrameEventCompletion();
+				}
 			}
-			else
+		
+			// If the input latency timer has been triggered, block until the GPU is completely
+			// finished displaying this frame and calculate the delta time.
+			if ( GInputLatencyTimer.RenderThreadTrigger )
 			{
-				// Finish current frame immediately to reduce latency
-				Viewport->IssueFrameEvent();
 				Viewport->WaitForFrameEventCompletion();
+				uint32 EndTime = FPlatformTime::Cycles();
+				GInputLatencyTimer.DeltaTime = EndTime - GInputLatencyTimer.StartTime;
+				GInputLatencyTimer.RenderThreadTrigger = false;
 			}
 		}
-		
-		// If the input latency timer has been triggered, block until the GPU is completely
-		// finished displaying this frame and calculate the delta time.
-		if ( GInputLatencyTimer.RenderThreadTrigger )
+
+		if (bRevertToSharedContextAfterDrawingViewport)
 		{
-			Viewport->WaitForFrameEventCompletion();
-			uint32 EndTime = FPlatformTime::Cycles();
-			GInputLatencyTimer.DeltaTime = EndTime - GInputLatencyTimer.StartTime;
-			GInputLatencyTimer.RenderThreadTrigger = false;
+			PlatformSharedContextSetup(PlatformDevice);
+			bRevertToSharedContextAfterDrawingViewport = false;
 		}
 	}
 
-	if (bRevertToSharedContextAfterDrawingViewport)
-	{
-		PlatformSharedContextSetup(PlatformDevice);
-		bRevertToSharedContextAfterDrawingViewport = false;
-	}
-	uint32 ThisCycles = FPlatformTime::Cycles() - IdleStart;
-	if (IsInRHIThread())
-	{
-		GWorkingRHIThreadStallTime += ThisCycles;
-	}
-	else if (IsInActualRenderingThread())
-	{
-		GRenderThreadIdle[ERenderThreadIdleTypes::WaitingForGPUPresent] += ThisCycles;
-		GRenderThreadNumIdle[ERenderThreadIdleTypes::WaitingForGPUPresent]++;
-	}
-
 	EndFrameTick();
 }
 
diff --git a/Engine/Source/Runtime/RHI/Private/RHICommandList.cpp b/Engine/Source/Runtime/RHI/Private/RHICommandList.cpp
--- a/Engine/Source/Runtime/RHI/Private/RHICommandList.cpp
+++ b/Engine/Source/Runtime/RHI/Private/RHICommandList.cpp
@@ -112,11 +112,8 @@ bool GIsRunningRHIInTaskThread_InternalUseOnly = false;
 
 /** Accumulates how many cycles the renderthread has been idle. */
 uint32 GRenderThreadIdle[ERenderThreadIdleTypes::Num] = { 0 };
-/** Accumulates how times renderthread was idle. */
-uint32 GRenderThreadNumIdle[ERenderThreadIdleTypes::Num] = { 0 };
 
 uint32 GWorkingRHIThreadTime = 0;
-uint32 GWorkingRHIThreadStallTime = 0;
 uint32 GWorkingRHIThreadStartCycles = 0;
 
 /** How many cycles the from sampling input to the frame being flipped. */
@@ -281,8 +278,6 @@ void FRHICommandListBase::WaitForDispatchEvent()
 	if (!DispatchEvent->IsComplete())
 	{
 		SCOPE_CYCLE_COUNTER(STAT_WaitForCmdListDispatch);
-
-		FRenderThreadIdleScope Scope(ERenderThreadIdleTypes::WaitingForAllOtherSleep);
 		DispatchEvent->Wait();
 	}
 }
@@ -640,8 +635,6 @@ void FRHICommandListImmediate::QueueAsyncCommandListSubmit(TArrayView<FQueuedCom
 				if (!Task.Event->IsComplete())
 				{
 					SCOPE_CYCLE_COUNTER(STAT_ParallelTranslateWait);
-
-					FRenderThreadIdleScope Scope(ERenderThreadIdleTypes::WaitingForAllOtherSleep);
 					Task.Event->Wait();
 				}
 
@@ -1046,8 +1039,6 @@ void FRHICommandListImmediate::WaitForRHIThreadTasks()
 	if (RHIThreadTask && !RHIThreadTask->IsComplete())
 	{
 		SCOPE_CYCLE_COUNTER(STAT_ExplicitWaitRHIThread);
-
-		FRenderThreadIdleScope Scope(ERenderThreadIdleTypes::WaitingForAllOtherSleep);
 		RHIThreadTask->Wait();
 	}
 }
diff --git a/Engine/Source/Runtime/RHI/Public/RHICommandList.h b/Engine/Source/Runtime/RHI/Public/RHICommandList.h
--- a/Engine/Source/Runtime/RHI/Public/RHICommandList.h
+++ b/Engine/Source/Runtime/RHI/Public/RHICommandList.h
@@ -91,42 +91,32 @@ namespace ERenderThreadIdleTypes
 
 /** Accumulates how many cycles the renderthread has been idle. */
 extern RHI_API uint32 GRenderThreadIdle[ERenderThreadIdleTypes::Num];
-/** Accumulates how times renderthread was idle. */
-extern RHI_API uint32 GRenderThreadNumIdle[ERenderThreadIdleTypes::Num];
 
 /** private accumulator for the RHI thread. */
 extern RHI_API uint32 GWorkingRHIThreadTime;
-extern RHI_API uint32 GWorkingRHIThreadStallTime;
 extern RHI_API uint32 GWorkingRHIThreadStartCycles;
 
 /** Helper to mark scopes as idle time on the render or RHI threads. */
 struct FRenderThreadIdleScope
 {
+	FThreadIdleStats::FScopeIdle RHIThreadIdleScope;
+
 	const ERenderThreadIdleTypes::Type Type;
-	const uint32 Start;
 	const bool bCondition;
+	const uint32 Start;
 
-	FRenderThreadIdleScope(ERenderThreadIdleTypes::Type Type, bool bCondition = true)
-		: Type(Type)
-		, Start(FPlatformTime::Cycles())
-		, bCondition(bCondition)
+	FRenderThreadIdleScope(ERenderThreadIdleTypes::Type Type, bool bInCondition = true)
+		: RHIThreadIdleScope(!(bInCondition && IsInRHIThread()))
+		, Type(Type)
+		, bCondition(bInCondition && IsInRenderingThread())
+		, Start(bCondition ? FPlatformTime::Cycles() : 0)
 	{}
+
 	~FRenderThreadIdleScope()
 	{
 		if (bCondition)
 		{
-			uint32 End = FPlatformTime::Cycles();
-			uint32 IdleCycles = End - Start;
-
-			if (IsInRHIThread())
-			{
-				GWorkingRHIThreadStallTime += IdleCycles;
-			}
-			else if (IsInRenderingThread())
-			{
-				GRenderThreadIdle[Type] += IdleCycles;
-				GRenderThreadNumIdle[ERenderThreadIdleTypes::WaitingForAllOtherSleep]++;
-			}
+			GRenderThreadIdle[Type] += FPlatformTime::Cycles() - Start;
 		}
 	}
 };
diff --git a/Engine/Source/Runtime/Renderer/Private/SceneOcclusion.cpp b/Engine/Source/Runtime/Renderer/Private/SceneOcclusion.cpp
--- a/Engine/Source/Runtime/Renderer/Private/SceneOcclusion.cpp
+++ b/Engine/Source/Runtime/Renderer/Private/SceneOcclusion.cpp
@@ -824,7 +824,8 @@ void FHZBOcclusionTester::MapResults(FRHICommandListImmediate& RHICmdList)
 
 	if (!IsInvalidFrame() )
 	{
-		uint32 IdleStart = FPlatformTime::Cycles();
+		// RHIMapStagingSurface will block until the results are ready (from the previous frame) so we need to consider this RT idle time
+		FRenderThreadIdleScope IdleScope(ERenderThreadIdleTypes::WaitingForGPUQuery);
 
 		SCOPED_GPU_MASK(RHICmdList, ResultsReadback->GetLastCopyGPUMask());
 
@@ -836,10 +837,6 @@ void FHZBOcclusionTester::MapResults(FRHICommandListImmediate& RHICmdList)
 			check(ResultsBufferRowPitch >= SizeX);
 			check(ResultBufferHeight >= SizeY);
 		}
-
-		// RHIMapStagingSurface will block until the results are ready (from the previous frame) so we need to consider this RT idle time
-		GRenderThreadIdle[ERenderThreadIdleTypes::WaitingForGPUQuery] += FPlatformTime::Cycles() - IdleStart;
-		GRenderThreadNumIdle[ERenderThreadIdleTypes::WaitingForGPUQuery]++;
 	}
 	
 	// Can happen because of device removed, we might crash later but this occlusion culling system can behave gracefully.
diff --git a/Engine/Source/Runtime/SlateRHIRenderer/Private/SlateRHIRenderer.cpp b/Engine/Source/Runtime/SlateRHIRenderer/Private/SlateRHIRenderer.cpp
--- a/Engine/Source/Runtime/SlateRHIRenderer/Private/SlateRHIRenderer.cpp
+++ b/Engine/Source/Runtime/SlateRHIRenderer/Private/SlateRHIRenderer.cpp
@@ -1502,24 +1502,25 @@ void FSlateRHIRenderer::DrawWindow_RenderThread(FRHICommandListImmediate& RHICmd
 	FThreadIdleStats& RenderThread = FThreadIdleStats::Get();
 	GRenderThreadIdle[ERenderThreadIdleTypes::WaitingForAllOtherSleep] = RenderThread.Waits;
 	GRenderThreadIdle[ERenderThreadIdleTypes::WaitingForGPUPresent] += GSwapBufferTime;
-	GRenderThreadNumIdle[ERenderThreadIdleTypes::WaitingForGPUPresent]++;
 
-	SET_CYCLE_COUNTER(STAT_RenderingIdleTime_RenderThreadSleepTime, GRenderThreadIdle[0]);
-	SET_CYCLE_COUNTER(STAT_RenderingIdleTime_WaitingForGPUQuery, GRenderThreadIdle[1]);
-	SET_CYCLE_COUNTER(STAT_RenderingIdleTime_WaitingForGPUPresent, GRenderThreadIdle[2]);
+	SET_CYCLE_COUNTER(STAT_RenderingIdleTime_RenderThreadSleepTime, GRenderThreadIdle[ERenderThreadIdleTypes::WaitingForAllOtherSleep]);
+	SET_CYCLE_COUNTER(STAT_RenderingIdleTime_WaitingForGPUQuery   , GRenderThreadIdle[ERenderThreadIdleTypes::WaitingForGPUQuery     ]);
+	SET_CYCLE_COUNTER(STAT_RenderingIdleTime_WaitingForGPUPresent , GRenderThreadIdle[ERenderThreadIdleTypes::WaitingForGPUPresent   ]);
+
+	const uint32 RenderThreadNonCriticalWaits = RenderThread.Waits - RenderThread.WaitsCriticalPath;
+	const uint32 RenderThreadWaitingForGPUQuery = GRenderThreadIdle[ERenderThreadIdleTypes::WaitingForGPUQuery];
 
 	// Set the RenderThreadIdle CSV stats
-	CSV_CUSTOM_STAT(RenderThreadIdle, Total, FPlatformTime::ToMilliseconds(RenderThread.Waits), ECsvCustomStatOp::Set);
-	CSV_CUSTOM_STAT(RenderThreadIdle, CriticalPath, FPlatformTime::ToMilliseconds(RenderThread.WaitsCriticalPath), ECsvCustomStatOp::Set);
-	CSV_CUSTOM_STAT(RenderThreadIdle, SwapBuffer, FPlatformTime::ToMilliseconds(GSwapBufferTime), ECsvCustomStatOp::Set);
-	CSV_CUSTOM_STAT(RenderThreadIdle, NonCriticalPath, FPlatformTime::ToMilliseconds(RenderThread.Waits - RenderThread.WaitsCriticalPath), ECsvCustomStatOp::Set);
-	CSV_CUSTOM_STAT(RenderThreadIdle, GPUQuery, FPlatformTime::ToMilliseconds(GRenderThreadIdle[1]), ECsvCustomStatOp::Set);
+	CSV_CUSTOM_STAT(RenderThreadIdle, Total          , FPlatformTime::ToMilliseconds(RenderThread.Waits            ), ECsvCustomStatOp::Set);
+	CSV_CUSTOM_STAT(RenderThreadIdle, CriticalPath   , FPlatformTime::ToMilliseconds(RenderThread.WaitsCriticalPath), ECsvCustomStatOp::Set);
+	CSV_CUSTOM_STAT(RenderThreadIdle, SwapBuffer     , FPlatformTime::ToMilliseconds(GSwapBufferTime               ), ECsvCustomStatOp::Set);
+	CSV_CUSTOM_STAT(RenderThreadIdle, NonCriticalPath, FPlatformTime::ToMilliseconds(RenderThreadNonCriticalWaits  ), ECsvCustomStatOp::Set);
+	CSV_CUSTOM_STAT(RenderThreadIdle, GPUQuery       , FPlatformTime::ToMilliseconds(RenderThreadWaitingForGPUQuery), ECsvCustomStatOp::Set);
 
 	for (int32 Index = 0; Index < ERenderThreadIdleTypes::Num; Index++)
 	{
 		RenderThreadIdle += GRenderThreadIdle[Index];
 		GRenderThreadIdle[Index] = 0;
-		GRenderThreadNumIdle[Index] = 0;
 	}
 
 	SET_CYCLE_COUNTER(STAT_RenderingIdleTime, RenderThreadIdle);
@@ -1549,10 +1550,17 @@ void FSlateRHIRenderer::DrawWindow_RenderThread(FRHICommandListImmediate& RHICmd
 			GWorkingRHIThreadTime += (ThisCycles - GWorkingRHIThreadStartCycles);
 			GWorkingRHIThreadStartCycles = ThisCycles;
 
-			uint32 NewVal = GWorkingRHIThreadTime - GWorkingRHIThreadStallTime;
+			FThreadIdleStats& RHIThreadStats = FThreadIdleStats::Get();
+
+			uint32 NewVal = GWorkingRHIThreadTime;
+			if (NewVal > RHIThreadStats.Waits)
+			{
+				NewVal -= RHIThreadStats.Waits;
+			}
+
 			FPlatformAtomics::AtomicStore((int32*)&GRHIThreadTime, (int32)NewVal);
 			GWorkingRHIThreadTime = 0;
-			GWorkingRHIThreadStallTime = 0;
+			RHIThreadStats.Reset();
 		});
 	}
 }
diff --git a/Engine/Source/Runtime/VulkanRHI/Private/VulkanQuery.cpp b/Engine/Source/Runtime/VulkanRHI/Private/VulkanQuery.cpp
--- a/Engine/Source/Runtime/VulkanRHI/Private/VulkanQuery.cpp
+++ b/Engine/Source/Runtime/VulkanRHI/Private/VulkanQuery.cpp
@@ -114,8 +114,7 @@ bool FVulkanOcclusionQueryPool::InternalTryGetResults(bool bWait)
 	{
 		if (bWait)
 		{
-			uint32 IdleStart = FPlatformTime::Cycles();
-
+			FRenderThreadIdleScope IdleScope(ERenderThreadIdleTypes::WaitingForGPUQuery);
 			SCOPE_CYCLE_COUNTER(STAT_VulkanWaitQuery);
 
 			// We'll do manual wait
@@ -173,12 +172,6 @@ bool FVulkanOcclusionQueryPool::InternalTryGetResults(bool bWait)
 				++NumLoops;
 			}
 
-			if (IsInActualRenderingThread())
-			{
-				GRenderThreadIdle[ERenderThreadIdleTypes::WaitingForGPUQuery] += FPlatformTime::Cycles() - IdleStart;
-				GRenderThreadNumIdle[ERenderThreadIdleTypes::WaitingForGPUQuery]++;
-			}
-
 			State = EState::RT_PostGetResults;
 			return true;
 		}
@@ -518,32 +511,31 @@ bool FVulkanDynamicRHI::RHIGetRenderQueryResult(FRHIRenderQuery* QueryRHI, uint6
 			// This really only happens if occlusion and frame sync event queries are disabled, otherwise those will block until the GPU catches up to 1 frame behind
 
 			const bool bBlocking = (Query->Pool->NumIssuedTimestamps == Query->Pool->BufferSize) || bWait;
-			const uint32 IdleStart = FPlatformTime::Cycles();
-
-			SCOPE_CYCLE_COUNTER(STAT_RenderQueryResultTime);
-
-			if (bBlocking)
 			{
-				const FVulkanTimingQueryPool::FCmdBufferFence& StartQuerySyncPoint = Query->Pool->TimestampListHandles[TimestampIndex];
-				bool bWaitForStart = StartQuerySyncPoint.FenceCounter == StartQuerySyncPoint.CmdBuffer->GetFenceSignaledCounter();
-				if (bWaitForStart)
+				FRenderThreadIdleScope IdleScope(ERenderThreadIdleTypes::WaitingForGPUQuery);
+				SCOPE_CYCLE_COUNTER(STAT_RenderQueryResultTime);
+
+				if (bBlocking)
 				{
-					FRHICommandListExecutor::GetImmediateCommandList().ImmediateFlush(EImmediateFlushType::FlushRHIThread);
+					const FVulkanTimingQueryPool::FCmdBufferFence& StartQuerySyncPoint = Query->Pool->TimestampListHandles[TimestampIndex];
+					bool bWaitForStart = StartQuerySyncPoint.FenceCounter == StartQuerySyncPoint.CmdBuffer->GetFenceSignaledCounter();
+					if (bWaitForStart)
+					{
+						FRHICommandListExecutor::GetImmediateCommandList().ImmediateFlush(EImmediateFlushType::FlushRHIThread);
 
-					// Need to submit the open command lists.
-					Device->SubmitCommandsAndFlushGPU();
-				}
+						// Need to submit the open command lists.
+						Device->SubmitCommandsAndFlushGPU();
+					}
 
-				// CPU wait for query results to be ready.
-				if (bWaitForStart && StartQuerySyncPoint.FenceCounter == StartQuerySyncPoint.CmdBuffer->GetFenceSignaledCounter())
-				{
-					Device->GetImmediateContext().GetCommandBufferManager()->WaitForCmdBuffer(StartQuerySyncPoint.CmdBuffer);
+					// CPU wait for query results to be ready.
+					if (bWaitForStart && StartQuerySyncPoint.FenceCounter == StartQuerySyncPoint.CmdBuffer->GetFenceSignaledCounter())
+					{
+						Device->GetImmediateContext().GetCommandBufferManager()->WaitForCmdBuffer(StartQuerySyncPoint.CmdBuffer);
+					}
 				}
-			}
 
-			Query->Pool->ResultsBuffer->InvalidateMappedMemory();
-			GRenderThreadIdle[ERenderThreadIdleTypes::WaitingForGPUQuery] += FPlatformTime::Cycles() - IdleStart;
-			GRenderThreadNumIdle[ERenderThreadIdleTypes::WaitingForGPUQuery]++;
+				Query->Pool->ResultsBuffer->InvalidateMappedMemory();
+			}
 
 			const uint64* Data = Query->Pool->MappedPointer;
 			OutNumPixels = ToMicroseconds(Data[TimestampIndex]);
diff --git a/Engine/Source/Runtime/VulkanRHI/Private/VulkanSwapChain.cpp b/Engine/Source/Runtime/VulkanRHI/Private/VulkanSwapChain.cpp
--- a/Engine/Source/Runtime/VulkanRHI/Private/VulkanSwapChain.cpp
+++ b/Engine/Source/Runtime/VulkanRHI/Private/VulkanSwapChain.cpp
@@ -671,7 +671,8 @@ int32 FVulkanSwapChain::AcquireImageIndex(VulkanRHI::FSemaphore** OutSemaphore)
 		const uint32 MaxImageIndex = ImageAcquiredSemaphore.Num() - 1;
 
 		SCOPE_CYCLE_COUNTER(STAT_VulkanAcquireBackBuffer);
-		uint32 IdleStart = FPlatformTime::Cycles();
+		FRenderThreadIdleScope IdleScope(ERenderThreadIdleTypes::WaitingForGPUPresent);
+
 		Result = VulkanRHI::vkAcquireNextImageKHR(
 			Device.GetInstanceHandle(),
 			SwapChain,
@@ -691,17 +692,6 @@ int32 FVulkanSwapChain::AcquireImageIndex(VulkanRHI::FSemaphore** OutSemaphore)
 				AcquiredFence,
 				&ImageIndex);
 		}
-
-		uint32 ThisCycles = FPlatformTime::Cycles() - IdleStart;
-		if (IsInRHIThread())
-		{
-			GWorkingRHIThreadStallTime += ThisCycles;
-		}
-		else if (IsInActualRenderingThread())
-		{
-			GRenderThreadIdle[ERenderThreadIdleTypes::WaitingForGPUPresent] += ThisCycles;
-			GRenderThreadNumIdle[ERenderThreadIdleTypes::WaitingForGPUPresent]++;
-		}
 	}
 
 	if (Result == VK_ERROR_OUT_OF_DATE_KHR)
@@ -772,17 +762,14 @@ void FVulkanSwapChain::RenderThreadPacing()
 
 		if (SampledDeltaMS < (TargetIntervalWithEpsilonMS))
 		{
-			uint32 IdleStart = FPlatformTime::Cycles();
+			FRenderThreadIdleScope IdleScope(ERenderThreadIdleTypes::WaitingForGPUPresent);
 			QUICK_SCOPE_CYCLE_COUNTER(STAT_StallForEmulatedSyncInterval);
+
 			FPlatformProcess::SleepNoStats((TargetIntervalWithEpsilonMS - SampledDeltaMS) * 0.001f);
 			if (GPrintVulkanVsyncDebug)
 			{
 				UE_LOG(LogVulkanRHI, Log, TEXT("CPU RT delta: %f, TargetWEps: %f, sleepTime: %f "), SampledDeltaMS, TargetIntervalWithEpsilonMS, TargetIntervalWithEpsilonMS - DeltaCPUPresentTimeMS);
 			}
-
-			uint32 ThisCycles = FPlatformTime::Cycles() - IdleStart;
-			GRenderThreadIdle[ERenderThreadIdleTypes::WaitingForGPUPresent] += ThisCycles;
-			GRenderThreadNumIdle[ERenderThreadIdleTypes::WaitingForGPUPresent]++;
 		}
 		else
 		{
@@ -830,24 +817,14 @@ FVulkanSwapChain::EStatus FVulkanSwapChain::Present(FVulkanQueue* GfxQueue, FVul
 
 			if (TimeToSleep > 0.0)
 			{
-				uint32 IdleStart = FPlatformTime::Cycles();
+				FRenderThreadIdleScope IdleScope(ERenderThreadIdleTypes::WaitingForGPUPresent);
+
 				QUICK_SCOPE_CYCLE_COUNTER(STAT_StallForEmulatedSyncInterval);
 				FPlatformProcess::SleepNoStats(static_cast<float>(TimeToSleep));
 				if (GPrintVulkanVsyncDebug)
 				{
 					UE_LOG(LogVulkanRHI, Log, TEXT("CurrentID: %i, CPU TimeToSleep: %f, TargetWEps: %f"), PresentID, TimeToSleep * 1000.0, TargetIntervalWithEpsilon * 1000.0);
 				}
-
-				uint32 ThisCycles = FPlatformTime::Cycles() - IdleStart;
-				if (IsInRHIThread())
-				{
-					GWorkingRHIThreadStallTime += ThisCycles;
-				}
-				else if (IsInActualRenderingThread())
-				{
-					GRenderThreadIdle[ERenderThreadIdleTypes::WaitingForGPUPresent] += ThisCycles;
-					GRenderThreadNumIdle[ERenderThreadIdleTypes::WaitingForGPUPresent]++;
-				}
 			}
 			else
 			{
@@ -863,17 +840,11 @@ FVulkanSwapChain::EStatus FVulkanSwapChain::Present(FVulkanQueue* GfxQueue, FVul
 
 	{
 		SCOPE_CYCLE_COUNTER(STAT_VulkanQueuePresent);
-		uint32 IdleStart = FPlatformTime::Cycles();
-		VkResult PresentResult = FVulkanPlatform::Present(PresentQueue->GetHandle(), Info);
-		uint32 ThisCycles = FPlatformTime::Cycles() - IdleStart;
-		if (IsInRHIThread())
-		{
-			GWorkingRHIThreadStallTime += ThisCycles;
-		}
-		else if (IsInActualRenderingThread())
+
+		VkResult PresentResult;
 		{
-			GRenderThreadIdle[ERenderThreadIdleTypes::WaitingForGPUPresent] += ThisCycles;
-			GRenderThreadNumIdle[ERenderThreadIdleTypes::WaitingForGPUPresent]++;
+			FRenderThreadIdleScope IdleScope(ERenderThreadIdleTypes::WaitingForGPUPresent);
+			PresentResult = FVulkanPlatform::Present(PresentQueue->GetHandle(), Info);
 		}
 
 		CurrentImageIndex = -1;
diff --git a/Engine/Source/Runtime/VulkanRHI/Private/VulkanUtil.cpp b/Engine/Source/Runtime/VulkanRHI/Private/VulkanUtil.cpp
--- a/Engine/Source/Runtime/VulkanRHI/Private/VulkanUtil.cpp
+++ b/Engine/Source/Runtime/VulkanRHI/Private/VulkanUtil.cpp
@@ -421,30 +421,28 @@ uint64 FVulkanGPUTiming::GetTiming(bool bGetCurrentResultsAndBlock)
 
 			if (bBlocking)
 			{
-				const uint32 IdleStart = FPlatformTime::Cycles();
-
 				SCOPE_CYCLE_COUNTER(STAT_RenderQueryResultTime);
-
-				const bool bWaitForStart = StartQuerySyncPoint.FenceCounter >= StartQuerySyncPoint.CmdBuffer->GetFenceSignaledCounter();
-				const bool bWaitForEnd = EndQuerySyncPoint.FenceCounter >= EndQuerySyncPoint.CmdBuffer->GetFenceSignaledCounter();
-				if (bWaitForEnd || bWaitForStart)
 				{
-					// Need to submit the open command lists.
-					Device->SubmitCommandsAndFlushGPU();
-				}
+					FRenderThreadIdleScope IdleScope(ERenderThreadIdleTypes::WaitingForGPUQuery);
 
-				// CPU wait for query results to be ready.
-				if (bWaitForStart && StartQuerySyncPoint.FenceCounter >= StartQuerySyncPoint.CmdBuffer->GetFenceSignaledCounter())
-				{
-					CmdContext->GetCommandBufferManager()->WaitForCmdBuffer(StartQuerySyncPoint.CmdBuffer);
-				}
-				if (bWaitForEnd && EndQuerySyncPoint.FenceCounter >= EndQuerySyncPoint.CmdBuffer->GetFenceSignaledCounter())
-				{
-					CmdContext->GetCommandBufferManager()->WaitForCmdBuffer(EndQuerySyncPoint.CmdBuffer);
-				}
+					const bool bWaitForStart = StartQuerySyncPoint.FenceCounter >= StartQuerySyncPoint.CmdBuffer->GetFenceSignaledCounter();
+					const bool bWaitForEnd = EndQuerySyncPoint.FenceCounter >= EndQuerySyncPoint.CmdBuffer->GetFenceSignaledCounter();
+					if (bWaitForEnd || bWaitForStart)
+					{
+						// Need to submit the open command lists.
+						Device->SubmitCommandsAndFlushGPU();
+					}
 
-				GRenderThreadIdle[ERenderThreadIdleTypes::WaitingForGPUQuery] += FPlatformTime::Cycles() - IdleStart;
-				GRenderThreadNumIdle[ERenderThreadIdleTypes::WaitingForGPUQuery]++;
+					// CPU wait for query results to be ready.
+					if (bWaitForStart && StartQuerySyncPoint.FenceCounter >= StartQuerySyncPoint.CmdBuffer->GetFenceSignaledCounter())
+					{
+						CmdContext->GetCommandBufferManager()->WaitForCmdBuffer(StartQuerySyncPoint.CmdBuffer);
+					}
+					if (bWaitForEnd && EndQuerySyncPoint.FenceCounter >= EndQuerySyncPoint.CmdBuffer->GetFenceSignaledCounter())
+					{
+						CmdContext->GetCommandBufferManager()->WaitForCmdBuffer(EndQuerySyncPoint.CmdBuffer);
+					}
+				}
 
 				Pool->ResultsBuffer->InvalidateMappedMemory();
 				StartTimeAvailability = Pool->MappedPointer[QueryStartIndex * 2 + 1];
diff --git a/Engine/Source/Runtime/Windows/D3D11RHI/Private/D3D11Query.cpp b/Engine/Source/Runtime/Windows/D3D11RHI/Private/D3D11Query.cpp
--- a/Engine/Source/Runtime/Windows/D3D11RHI/Private/D3D11Query.cpp
+++ b/Engine/Source/Runtime/Windows/D3D11RHI/Private/D3D11Query.cpp
@@ -302,7 +302,8 @@ bool FD3D11DynamicRHI::GetQueryData(ID3D11Query* Query, void* Data, SIZE_T DataS
 	if ( Result == S_FALSE && bWait )
 	{
 		SCOPE_CYCLE_COUNTER( STAT_RenderQueryResultTime );
-		uint32 IdleStart = FPlatformTime::Cycles();
+		FRenderThreadIdleScope IdleScope(ERenderThreadIdleTypes::WaitingForGPUQuery);
+
 		double StartTime = FPlatformTime::Seconds();
 		double TimeoutWarningLimit = 5.0;
 		// timer queries are used for Benchmarks which can stall a bit more
@@ -343,8 +344,6 @@ bool FD3D11DynamicRHI::GetQueryData(ID3D11Query* Query, void* Data, SIZE_T DataS
 				return false;
 			}
 		} while ( Result == S_FALSE );
-		GRenderThreadIdle[ERenderThreadIdleTypes::WaitingForGPUQuery] += FPlatformTime::Cycles() - IdleStart;
-		GRenderThreadNumIdle[ERenderThreadIdleTypes::WaitingForGPUQuery]++;
 	}
 
 	if( Result == S_OK )
@@ -726,30 +725,15 @@ uint64 FD3D11BufferedGPUTiming::GetTiming(bool bGetCurrentResultsAndBlock)
 			// This really only happens if occlusion and frame sync event queries are disabled, otherwise those will block until the GPU catches up to 1 frame behind
 			const bool bBlocking = ( NumIssuedTimestamps == BufferSize ) || bGetCurrentResultsAndBlock;
 			const uint32 AsyncFlags = bBlocking ? 0 : D3D11_ASYNC_GETDATA_DONOTFLUSH;
-			uint32 IdleStart = FPlatformTime::Cycles();
-			double StartTimeoutTime = FPlatformTime::Seconds();
-
-			SCOPE_CYCLE_COUNTER( STAT_RenderQueryResultTime );
-			// If we are blocking, retry until the GPU processes the time stamp command
-			do 
 			{
-				D3DResult = D3DRHI->GetDeviceContext()->GetData( EndTimestamps[TimestampIndex], &EndTime, sizeof(EndTime), AsyncFlags );
+				FRenderThreadIdleScope IdleScope(ERenderThreadIdleTypes::WaitingForGPUQuery);
+				double StartTimeoutTime = FPlatformTime::Seconds();
 
-				if ((FPlatformTime::Seconds() - StartTimeoutTime) > 0.5)
-				{
-					UE_LOG(LogD3D11RHI, Log, TEXT("Timed out while waiting for GPU to catch up. (500 ms)"));
-					return 0;
-				}
-			} while ( D3DResult == S_FALSE && bBlocking );
-			GRenderThreadIdle[ERenderThreadIdleTypes::WaitingForGPUQuery] += FPlatformTime::Cycles() - IdleStart;
-			GRenderThreadNumIdle[ERenderThreadIdleTypes::WaitingForGPUQuery]++;
-			if ( D3DResult == S_OK )
-			{
-				IdleStart = FPlatformTime::Cycles();
-				StartTimeoutTime = FPlatformTime::Seconds();
+				SCOPE_CYCLE_COUNTER( STAT_RenderQueryResultTime );
+				// If we are blocking, retry until the GPU processes the time stamp command
 				do 
 				{
-					D3DResult = D3DRHI->GetDeviceContext()->GetData( StartTimestamps[TimestampIndex], &StartTime, sizeof(StartTime), AsyncFlags );
+					D3DResult = D3DRHI->GetDeviceContext()->GetData( EndTimestamps[TimestampIndex], &EndTime, sizeof(EndTime), AsyncFlags );
 
 					if ((FPlatformTime::Seconds() - StartTimeoutTime) > 0.5)
 					{
@@ -757,7 +741,24 @@ uint64 FD3D11BufferedGPUTiming::GetTiming(bool bGetCurrentResultsAndBlock)
 						return 0;
 					}
 				} while ( D3DResult == S_FALSE && bBlocking );
-				GRenderThreadIdle[ERenderThreadIdleTypes::WaitingForGPUQuery] += FPlatformTime::Cycles() - IdleStart;
+			}
+
+			if ( D3DResult == S_OK )
+			{
+				{
+					FRenderThreadIdleScope IdleScope(ERenderThreadIdleTypes::WaitingForGPUQuery);
+					double StartTimeoutTime = FPlatformTime::Seconds();
+					do 
+					{
+						D3DResult = D3DRHI->GetDeviceContext()->GetData( StartTimestamps[TimestampIndex], &StartTime, sizeof(StartTime), AsyncFlags );
+
+						if ((FPlatformTime::Seconds() - StartTimeoutTime) > 0.5)
+						{
+							UE_LOG(LogD3D11RHI, Log, TEXT("Timed out while waiting for GPU to catch up. (500 ms)"));
+							return 0;
+						}
+					} while ( D3DResult == S_FALSE && bBlocking );
+				}
 
 				if ( D3DResult == S_OK && EndTime > StartTime )
 				{
