Add a GL cvar to set a  minimum occupancy for GL texture cache.
Can reduce texture LRU based hitching after a static scene resumes.

diff --git a/Engine/Source/Runtime/OpenGLDrv/Private/OpenGLTexture.cpp b/Engine/Source/Runtime/OpenGLDrv/Private/OpenGLTexture.cpp
--- a/Engine/Source/Runtime/OpenGLDrv/Private/OpenGLTexture.cpp
+++ b/Engine/Source/Runtime/OpenGLDrv/Private/OpenGLTexture.cpp
@@ -68,6 +68,17 @@ static FAutoConsoleVariableRef CVarTextureEvictionLogging(
 	ECVF_RenderThreadSafe
 );
 
+static int32 GOGLTextureMinLRUCapacity = 0;
+static FAutoConsoleVariableRef CVarTextureEvictionMinLRUCapacity(
+	TEXT("r.OpenGL.TextureEvictionMinLRUCapacity"),
+	GOGLTextureMinLRUCapacity,
+	TEXT("Keep a minimum number of textures resident in GL When using the texture LRU\n")
+	TEXT("This can reduce LRU restore times when resuming from static scenes.\n")
+	TEXT("0: (default)")
+	,
+	ECVF_RenderThreadSafe
+);
+
 /*-----------------------------------------------------------------------------
 	Texture allocator support.
 -----------------------------------------------------------------------------*/
@@ -2482,8 +2493,9 @@ void FTextureEvictionLRU::TickEviction()
 
 	FScopeLock Lock(&TextureLRULock);
 	FOpenGLTextureLRUContainer& TextureLRU = GetLRUContainer();
+
 	for (int32 EvictCount = 0; 
-		TextureLRU.Num() && (TextureLRU.GetLeastRecent()->EvictionParamsPtr->FrameLastRendered + GOGLTextureEvictFramesToLive) < GFrameNumberRenderThread && EvictCount < GOGLTexturesToEvictPerFrame
+		TextureLRU.Num() > FMath::Max(0, GOGLTextureMinLRUCapacity) && (TextureLRU.GetLeastRecent()->EvictionParamsPtr->FrameLastRendered + GOGLTextureEvictFramesToLive) < GFrameNumberRenderThread && EvictCount < GOGLTexturesToEvictPerFrame
 		;EvictCount++)
 	{
 		FOpenGLTexture* RemovedFromLRU = TextureLRU.RemoveLeastRecent();
