Add a system to insert GL fences after N draw calls. This acts as a hint to the driver to begin processing geometry.
cvar driven, disabled by default.

diff --git a/Engine/Source/Runtime/OpenGLDrv/Private/OpenGLCommands.cpp b/Engine/Source/Runtime/OpenGLDrv/Private/OpenGLCommands.cpp
--- a/Engine/Source/Runtime/OpenGLDrv/Private/OpenGLCommands.cpp
+++ b/Engine/Source/Runtime/OpenGLDrv/Private/OpenGLCommands.cpp
@@ -107,6 +107,17 @@ namespace OpenGLConsoleVariables
 		TEXT("1: enabled."),
 		ECVF_ReadOnly | ECVF_RenderThreadSafe
 	);
+
+	int32 GOpenGLFenceKickPerDrawCount = 0;
+	static FAutoConsoleVariableRef CVarOpenGLFenceKickPerDrawCount(
+		TEXT("r.OpenGL.FenceKickPerDrawCount"),
+		GOpenGLFenceKickPerDrawCount,
+		TEXT("Insert a GL fence after the specified number of draws has been issued.\n")
+		TEXT("This hint can encourage some drivers to begin processing geometry as soon as a sufficient workload has built up.\n")
+		TEXT("0: disabled (default)"),
+		ECVF_RenderThreadSafe
+	);
+
 };
 
 #if PLATFORM_64BITS
@@ -2267,6 +2278,63 @@ void FOpenGLDynamicRHI::CommitComputeResourceTables(FOpenGLComputeShaderProxy* C
 	PendingState.DirtyUniformBuffers[SF_Compute] = 0;
 }
 
+class FOpengGLFenceKick
+{
+public:
+
+	~FOpengGLFenceKick() { Reset(); }
+
+	void OnDrawCall(const FOpenGLContextState& ContextState)
+	{
+		if(OpenGLConsoleVariables::GOpenGLFenceKickPerDrawCount)
+		{
+			// we kick every GOpenGLFenceKickPerDrawCount draws within a render pass.
+			// counter is reset when the FBO changes.
+			if (ContextState.Framebuffer != LastSeenFramebuffer)
+			{
+				LastSeenFramebuffer = ContextState.Framebuffer;
+				DrawCounter = 0;
+			}
+			else if (++DrawCounter >= OpenGLConsoleVariables::GOpenGLFenceKickPerDrawCount)
+			{
+				InsertKick();
+				DrawCounter = 0;
+			}
+		}
+	}
+
+	void Reset()
+	{
+		for (UGLsync Sync : Syncs)
+		{
+			FOpenGL::DeleteSync(Sync);
+		}
+		Syncs.Reset();
+		DrawCounter = 0;
+		LastSeenFramebuffer = 0;
+	}
+
+private:
+
+	void InsertKick()
+	{
+		// record syncs incase a create/destroy pair is optimised by the driver.
+		UGLsync Sync = FOpenGL::FenceSync(GL_SYNC_GPU_COMMANDS_COMPLETE, 0);
+		Syncs.Add(Sync);
+		check(Syncs.Num() < 10000); // arbitrary check to ensure that we do eventually reset the array.
+	}
+
+	TArray<UGLsync> Syncs;
+	int32 DrawCounter = 0;
+	GLuint LastSeenFramebuffer = 0;
+};
+static FOpengGLFenceKick GOpenGLKickHint;
+
+void OpenGLCommands_OnEndFrame()
+{
+	GOpenGLKickHint.Reset();
+}
+
 void FOpenGLDynamicRHI::RHIDrawPrimitive(uint32 BaseVertexIndex,uint32 NumPrimitives,uint32 NumInstances)
 {
 	SCOPE_CYCLE_COUNTER_DETAILED(STAT_OpenGLDrawPrimitiveTime);
@@ -2306,6 +2374,7 @@ void FOpenGLDynamicRHI::RHIDrawPrimitive(uint32 BaseVertexIndex,uint32 NumPrimit
 		CONDITIONAL_SCOPE_CYCLE_COUNTER(STAT_OpenGLShaderFirstDrawTime, PendingState.BoundShaderState->RequiresDriverInstantiation());
 		FOpenGL::DrawArraysInstanced(DrawMode, 0, NumElements, NumInstances);
 	}
+	GOpenGLKickHint.OnDrawCall(ContextState);
 }
 
 void FOpenGLDynamicRHI::RHIDrawPrimitiveIndirect(FRHIBuffer* ArgumentBufferRHI, uint32 ArgumentOffset)
@@ -2315,7 +2384,7 @@ void FOpenGLDynamicRHI::RHIDrawPrimitiveIndirect(FRHIBuffer* ArgumentBufferRHI,
 		VERIFY_GL_SCOPE();
 
 		check(ArgumentBufferRHI);
-	GPUProfilingData.RegisterGPUWork(0);
+		GPUProfilingData.RegisterGPUWork(0);
 
 		FOpenGLContextState& ContextState = GetContextStateForCurrentContext();
 		BindPendingFramebuffer(ContextState);
@@ -2347,12 +2416,12 @@ void FOpenGLDynamicRHI::RHIDrawPrimitiveIndirect(FRHIBuffer* ArgumentBufferRHI,
 			FOpenGL::DrawArraysIndirect( DrawMode, INDEX_TO_VOID(ArgumentOffset));
 		}
 		glBindBuffer( GL_DRAW_INDIRECT_BUFFER, 0);
+		GOpenGLKickHint.OnDrawCall(ContextState);
 	}
 	else
 	{
 		UE_LOG(LogRHI, Fatal,TEXT("OpenGL RHI does not yet support indirect draw calls."));
 	}
-
 }
 
 void FOpenGLDynamicRHI::RHIDrawIndexedIndirect(FRHIBuffer* IndexBufferRHI, FRHIBuffer* ArgumentsBufferRHI, int32 DrawArgumentsIndex, uint32 NumInstances)
@@ -2403,6 +2472,8 @@ void FOpenGLDynamicRHI::RHIDrawIndexedIndirect(FRHIBuffer* IndexBufferRHI, FRHIB
 			FOpenGL::DrawElementsIndirect( DrawMode, IndexType, INDEX_TO_VOID(DrawArgumentsIndex * 5 *sizeof(uint32)));
 		}
 		glBindBuffer( GL_DRAW_INDIRECT_BUFFER, 0);
+
+		GOpenGLKickHint.OnDrawCall(ContextState);
 	}
 	else
 	{
@@ -2501,6 +2572,7 @@ void FOpenGLDynamicRHI::RHIDrawIndexedPrimitive(FRHIBuffer* IndexBufferRHI, int3
 			glDrawElements(DrawMode, NumElements, IndexType, INDEX_TO_VOID(StartIndex));
 		}
 	}
+	GOpenGLKickHint.OnDrawCall(ContextState);
 }
 
 void FOpenGLDynamicRHI::RHIDrawIndexedPrimitiveIndirect(FRHIBuffer* IndexBufferRHI, FRHIBuffer* ArgumentBufferRHI, uint32 ArgumentOffset)
@@ -2548,6 +2620,8 @@ void FOpenGLDynamicRHI::RHIDrawIndexedPrimitiveIndirect(FRHIBuffer* IndexBufferR
 			FOpenGL::DrawElementsIndirect( DrawMode, IndexType, INDEX_TO_VOID(ArgumentOffset));
 		}
 		glBindBuffer( GL_DRAW_INDIRECT_BUFFER, 0);
+
+		GOpenGLKickHint.OnDrawCall(ContextState);
 	}
 	else
 	{
diff --git a/Engine/Source/Runtime/OpenGLDrv/Private/OpenGLDevice.cpp b/Engine/Source/Runtime/OpenGLDrv/Private/OpenGLDevice.cpp
--- a/Engine/Source/Runtime/OpenGLDrv/Private/OpenGLDevice.cpp
+++ b/Engine/Source/Runtime/OpenGLDrv/Private/OpenGLDevice.cpp
@@ -196,11 +196,15 @@ void FOpenGLDynamicRHI::RHIBeginFrame()
 	OpenGL_PollAllFences();
 }
 
+extern void OpenGLCommands_OnEndFrame();
+
 void FOpenGLDynamicRHI::RHIEndFrame()
 {
 	GPUProfilingData.EndFrame();
 
 	OpenGL_PollAllFences();
+
+	OpenGLCommands_OnEndFrame();
 }
 
 void FOpenGLDynamicRHI::RHIPerFrameRHIFlushComplete()
