Fixed: GPUScene breaks movable objects on some mobile devices

diff --git a/Engine/Source/Runtime/OpenGLDrv/Private/Android/AndroidEGL.cpp b/Engine/Source/Runtime/OpenGLDrv/Private/Android/AndroidEGL.cpp
--- a/Engine/Source/Runtime/OpenGLDrv/Private/Android/AndroidEGL.cpp
+++ b/Engine/Source/Runtime/OpenGLDrv/Private/Android/AndroidEGL.cpp
@@ -31,9 +31,10 @@ DEFINE_LOG_CATEGORY(LogEGL);
 typedef int32(*PFN_ANativeWindow_setBuffersTransform)(struct ANativeWindow* window, int32 transform);
 static PFN_ANativeWindow_setBuffersTransform ANativeWindow_setBuffersTransform_API = nullptr;
 
+// Use blit by default as setBuffersTransform is broken on random devices
 static TAutoConsoleVariable<int32> CVarAndroidGLESFlipYMethod(
 	TEXT("r.Android.GLESFlipYMethod"),
-	0,
+	2,
 	TEXT(" 0: Flip Y method detected automatically by GPU vendor.\n"
 		 " 1: Force flip Y by native window setBuffersTransform.\n"
 		 " 2: Force flip Y by BlitFrameBuffer."),
diff --git a/Engine/Source/Runtime/OpenGLDrv/Private/Android/AndroidOpenGL.cpp b/Engine/Source/Runtime/OpenGLDrv/Private/Android/AndroidOpenGL.cpp
--- a/Engine/Source/Runtime/OpenGLDrv/Private/Android/AndroidOpenGL.cpp
+++ b/Engine/Source/Runtime/OpenGLDrv/Private/Android/AndroidOpenGL.cpp
@@ -252,6 +252,9 @@ bool PlatformBlitToViewport(FPlatformOpenGLDevice* Device, const FOpenGLViewport
 			);
 
 			glEnable(GL_FRAMEBUFFER_SRGB);
+
+			// Bind viewport FBO so driver knows we don't need backbuffer image anymore
+			glBindFramebuffer(GL_FRAMEBUFFER, Context->ViewportFramebuffer);
 		}
 	}
 
@@ -269,10 +272,6 @@ bool PlatformBlitToViewport(FPlatformOpenGLDevice* Device, const FOpenGLViewport
 	static IConsoleVariable* CVar = IConsoleManager::Get().FindConsoleVariable(TEXT("a.UseFrameTimeStampsForPacing"));
 	const bool bForceGPUFence = CVar ? CVar->GetInt() != 0 : false;
 
-	// Bind a dummy FBO so driver knows we don't need backbuffer anymore
-	check(Context->DummyFrameBuffer != 0);
-	glBindFramebuffer(GL_FRAMEBUFFER, Context->DummyFrameBuffer);
-
 	return bPresent && ShouldUseGPUFencesToLimitLatency();
 }
 
