Fix for OpenGL build error.

diff --git a/Engine/Source/Runtime/OpenGLDrv/Private/OpenGLCommands.cpp b/Engine/Source/Runtime/OpenGLDrv/Private/OpenGLCommands.cpp
--- a/Engine/Source/Runtime/OpenGLDrv/Private/OpenGLCommands.cpp
+++ b/Engine/Source/Runtime/OpenGLDrv/Private/OpenGLCommands.cpp
@@ -2181,7 +2181,7 @@ FORCEINLINE void FOpenGLDynamicRHI::SetResourcesFromTables(ShaderType* RESTRICT
 				, GetNumTextureUnits<Frequency>()
 			);
 
-			auto SRV = FOpenGLDynamicRHI::ResourceCast(RHISRV);
+			FOpenGLShaderResourceView* SRV = FOpenGLDynamicRHI::ResourceCast(RHISRV);
 			if (SRV->Target == GL_SHADER_STORAGE_BUFFER)
 			{
 				RHI.InternalSetShaderBufferUAV(GetFirstUAVUnit<Frequency>() + Index, SRV->Resource);
