Removed TEnableIf from APIs in OpenGLDrv.

diff --git a/Engine/Source/Runtime/OpenGLDrv/Public/OpenGLDrv.h b/Engine/Source/Runtime/OpenGLDrv/Public/OpenGLDrv.h
--- a/Engine/Source/Runtime/OpenGLDrv/Public/OpenGLDrv.h
+++ b/Engine/Source/Runtime/OpenGLDrv/Public/OpenGLDrv.h
@@ -332,31 +332,34 @@ public:
 	virtual void Shutdown();
 	virtual const TCHAR* GetName() override { return TEXT("OpenGL"); }
 
-	// If using a Proxy object return the contained GL object rather than the proxy itself.
 	template<typename TRHIType>
-	static FORCEINLINE typename TEnableIf<TIsGLProxyObject<typename TOpenGLResourceTraits<TRHIType>::TConcreteType>::Value, typename TOpenGLResourceTraits<TRHIType>::TConcreteType::ContainedGLType*>::Type ResourceCast(TRHIType* Resource)
-	{	
-		auto GLProxy = static_cast<typename TOpenGLResourceTraits<TRHIType>::TConcreteType*>(Resource);
-		// rhi resource can be null.
-		return GLProxy ? GLProxy->GetGLResourceObject() : nullptr;
-	}
-
-	template<typename TRHIType>
-	static FORCEINLINE typename TEnableIf<TIsGLProxyObject<typename TOpenGLResourceTraits<TRHIType>::TConcreteType>::Value, typename TOpenGLResourceTraits<TRHIType>::TConcreteType*>::Type ResourceCastProxy(TRHIType* Resource)
+	static FORCEINLINE typename TOpenGLResourceTraits<TRHIType>::TConcreteType* ResourceCastProxy(TRHIType* Resource)
 	{
+		static_assert(TIsGLProxyObject_V<typename TOpenGLResourceTraits<TRHIType>::TConcreteType>, "Wrong type past");
 		return static_cast<typename TOpenGLResourceTraits<TRHIType>::TConcreteType*>(Resource);
 	}
 
+	// If using a Proxy object return the contained GL object rather than the proxy itself.
 	template<typename TRHIType>
-	static FORCEINLINE typename TEnableIf<!TIsGLProxyObject<typename TOpenGLResourceTraits<TRHIType>::TConcreteType>::Value, typename TOpenGLResourceTraits<TRHIType>::TConcreteType*>::Type ResourceCast(TRHIType* Resource)
+	static FORCEINLINE auto* ResourceCast(TRHIType* Resource)
 	{
-		CheckRHITFence(static_cast<typename TOpenGLResourceTraits<TRHIType>::TConcreteType*>(Resource));
-		return static_cast<typename TOpenGLResourceTraits<TRHIType>::TConcreteType*>(Resource);
+		if constexpr (TIsGLProxyObject_V<typename TOpenGLResourceTraits<TRHIType>::TConcreteType>)
+		{
+			auto GLProxy = static_cast<typename TOpenGLResourceTraits<TRHIType>::TConcreteType*>(Resource);
+			// rhi resource can be null.
+			return GLProxy ? GLProxy->GetGLResourceObject() : nullptr;
+		}
+		else
+		{
+			CheckRHITFence(static_cast<typename TOpenGLResourceTraits<TRHIType>::TConcreteType*>(Resource));
+			return static_cast<typename TOpenGLResourceTraits<TRHIType>::TConcreteType*>(Resource);
+		}
 	}
 
 	template<typename TRHIType>
-	static FORCEINLINE typename TEnableIf<!TIsGLProxyObject<typename TOpenGLResourceTraits<TRHIType>::TConcreteType>::Value, typename TOpenGLResourceTraits<TRHIType>::TConcreteType*>::Type ResourceCast_Unfenced(TRHIType* Resource)
+	static FORCEINLINE typename TOpenGLResourceTraits<TRHIType>::TConcreteType* ResourceCast_Unfenced(TRHIType* Resource)
 	{
+		static_assert(!TIsGLProxyObject_V<typename TOpenGLResourceTraits<TRHIType>::TConcreteType>, "Wrong type passed");
 		return static_cast<typename TOpenGLResourceTraits<TRHIType>::TConcreteType*>(Resource);
 	}
 
diff --git a/Engine/Source/Runtime/OpenGLDrv/Public/OpenGLResources.h b/Engine/Source/Runtime/OpenGLDrv/Public/OpenGLResources.h
--- a/Engine/Source/Runtime/OpenGLDrv/Public/OpenGLResources.h
+++ b/Engine/Source/Runtime/OpenGLDrv/Public/OpenGLResources.h
@@ -357,22 +357,13 @@ typedef TOpenGLShaderProxy<FRHIComputeShader, FOpenGLComputeShader> FOpenGLCompu
 
 
 template <typename T>
-struct TIsGLProxyObject
-{
-	enum { Value = false };
-};
+constexpr bool TIsGLProxyObject_V = false;
 
 template<typename TRHIType, typename TOGLResourceType>
-struct TIsGLProxyObject<TOpenGLResourceProxy<TRHIType, TOGLResourceType>>
-{
-	enum { Value = true };
-};
+constexpr bool TIsGLProxyObject_V<TOpenGLResourceProxy<TRHIType, TOGLResourceType>> = true;
 
 template<typename TRHIType, typename TOGLResourceType>
-struct TIsGLProxyObject<TOpenGLShaderProxy<TRHIType, TOGLResourceType>>
-{
-	enum { Value = true };
-};
+constexpr bool TIsGLProxyObject_V<TOpenGLShaderProxy<TRHIType, TOGLResourceType>> = true;
 
 typedef void (*BufferBindFunction)( GLenum Type, GLuint Buffer );
 
@@ -1478,12 +1469,12 @@ struct TIsGLResourceWithFence
 };
 
 template<typename T>
-static typename TEnableIf<!TIsGLResourceWithFence<T>::Value>::Type CheckRHITFence(T* Resource) {}
-
-template<typename T>
-static typename TEnableIf<TIsGLResourceWithFence<T>::Value>::Type CheckRHITFence(T* Resource)
+static void CheckRHITFence(T* Resource)
 {
-	Resource->CreationFence.WaitFenceRenderThreadOnly();
+	if constexpr (TIsGLResourceWithFence<T>::Value)
+	{
+		Resource->CreationFence.WaitFenceRenderThreadOnly();
+	}
 }
 
 /** Given a pointer to a RHI texture that was created by the OpenGL RHI, returns a pointer to the FOpenGLTexture it encapsulates. */
