Android GLES binary program cache refactor:
 - Initial support for opening multiple caches.
 - Android GLES's precompiled binary cache is now mmapped rather than using async IO.
   - the OS can page in/out the binary data as required, it should not count as used mem for the Android LMK.
 - GLES no longer needs to restart after generating the binary cache.
 - Disabled binary program compression (it should no longer be needed when data is mmapped and adds unnecessary mem pressure.)
 - Prebuilt binary cache loads faster and requires far less time per tick to populate when using lru and resident program binaries.
 - GL programs unused for a set time are reclaimed/evicted from the driver.

diff --git a/Engine/Source/Runtime/OpenGLDrv/Private/Android/AndroidOpenGL.cpp b/Engine/Source/Runtime/OpenGLDrv/Private/Android/AndroidOpenGL.cpp
--- a/Engine/Source/Runtime/OpenGLDrv/Private/Android/AndroidOpenGL.cpp
+++ b/Engine/Source/Runtime/OpenGLDrv/Private/Android/AndroidOpenGL.cpp
@@ -160,6 +160,7 @@ static bool AreAndroidOpenGLRemoteCompileServicesAvailable()
 
 		RemoteCompileService = !bConfigRulesDisableProgramCompileServices && OpenGLRemoteGLProgramCompileJNI.bAllFound && (CVarProgramLRU->GetInt() != 0) && (CVarNumRemoteProgramCompileServices->GetInt() > 0);
 		FGenericCrashContext::SetEngineData(TEXT("Android.PSOService"), RemoteCompileService == 0? TEXT("disabled") : TEXT("enabled"));
+		UE_CLOG(!RemoteCompileService, LogRHI, Log, TEXT("Remote PSO services disabled: (%d, %d, %d, %d)"), bConfigRulesDisableProgramCompileServices, OpenGLRemoteGLProgramCompileJNI.bAllFound, CVarProgramLRU->GetInt(), CVarNumRemoteProgramCompileServices->GetInt());
 	}
 	return RemoteCompileService;
 }
diff --git a/Engine/Source/Runtime/OpenGLDrv/Private/OpenGLBinaryProgramUtils.cpp b/Engine/Source/Runtime/OpenGLDrv/Private/OpenGLBinaryProgramUtils.cpp
new file mode 100644
--- /dev/null
+++ b/Engine/Source/Runtime/OpenGLDrv/Private/OpenGLBinaryProgramUtils.cpp
@@ -0,0 +1,123 @@
+// Copyright Epic Games, Inc. All Rights Reserved.
+
+/*=============================================================================
+	OpenGLBinaryProgramUtils.cpp
+=============================================================================*/
+
+#include "OpenGLBinaryProgramUtils.h"
+
+namespace UE
+{
+	namespace OpenGL
+	{
+		static TAutoConsoleVariable<int32> CVarStoreCompressedBinaries(
+			TEXT("r.OpenGL.StoreCompressedProgramBinaries"),
+			0,
+			TEXT(""),
+			ECVF_ReadOnly | ECVF_RenderThreadSafe
+		);
+
+		bool IsStoringCompressedBinaryPrograms()
+		{
+			return CVarStoreCompressedBinaries.GetValueOnAnyThread() != 0;
+		}
+
+		bool UncompressCompressedBinaryProgram(const TArrayView<const uint8>& CompressedProgramBinary, TArray<uint8>& UncompressedProgramBinaryOUT)
+		{
+			if (ensure(CompressedProgramBinary.Num() > sizeof(FCompressedProgramBinaryHeader)))
+			{
+				FCompressedProgramBinaryHeader* Header = (FCompressedProgramBinaryHeader*)CompressedProgramBinary.GetData();
+
+				if (Header->UncompressedSize == FCompressedProgramBinaryHeader::NotCompressed)
+				{
+					const uint32 ProgramSize = CompressedProgramBinary.Num() - sizeof(FCompressedProgramBinaryHeader);
+					UncompressedProgramBinaryOUT.SetNumUninitialized(ProgramSize);
+					FMemory::Memcpy(UncompressedProgramBinaryOUT.GetData(), CompressedProgramBinary.GetData() + sizeof(FCompressedProgramBinaryHeader), ProgramSize);
+					return true;
+				}
+				else
+				{
+					UncompressedProgramBinaryOUT.AddUninitialized(Header->UncompressedSize);
+
+					if (Header->UncompressedSize > 0
+						&& FCompression::UncompressMemory(NAME_Zlib, UncompressedProgramBinaryOUT.GetData(), UncompressedProgramBinaryOUT.Num(), CompressedProgramBinary.GetData() + sizeof(FCompressedProgramBinaryHeader), CompressedProgramBinary.Num() - sizeof(FCompressedProgramBinaryHeader)))
+					{
+						return true;
+					}
+				}
+			}
+			return false;
+		}
+
+		bool GetUncompressedProgramBinaryFromGLProgram(GLuint Program, TArray<uint8>& ProgramBinaryOUT)
+		{
+			VERIFY_GL_SCOPE();
+			check(ProgramBinaryOUT.IsEmpty());
+			// pull binary from linked program
+			GLint BinaryLength = -1;
+			glGetProgramiv(Program, GL_PROGRAM_BINARY_LENGTH, &BinaryLength);
+			if (BinaryLength > 0)
+			{
+				ProgramBinaryOUT.SetNumUninitialized(BinaryLength + sizeof(GLenum));
+				uint8* ProgramBinaryPtr = ProgramBinaryOUT.GetData();
+				// BinaryFormat is stored at the start of ProgramBinary array
+				FOpenGL::GetProgramBinary(Program, BinaryLength, &BinaryLength, (GLenum*)ProgramBinaryPtr, ProgramBinaryPtr + sizeof(GLenum));
+				return true;
+			}
+			return false;
+		}
+
+		void CompressProgramBinary(const TArrayView<const uint8>& UncompressedProgramBinary, TArray<uint8>& ProgramBinaryOUT)
+		{
+			check(IsStoringCompressedBinaryPrograms());
+			check(ProgramBinaryOUT.IsEmpty());
+
+			int32 CompressedSize = FCompression::CompressMemoryBound(NAME_Zlib, UncompressedProgramBinary.Num());
+			uint32 CompressedHeaderSize = sizeof(FCompressedProgramBinaryHeader);
+			ProgramBinaryOUT.AddUninitialized(CompressedSize + CompressedHeaderSize);
+			bool bSuccess = FCompression::CompressMemory(NAME_Zlib, ProgramBinaryOUT.GetData() + CompressedHeaderSize, CompressedSize, UncompressedProgramBinary.GetData(), UncompressedProgramBinary.Num());
+			if (bSuccess)
+			{
+				ProgramBinaryOUT.SetNum(CompressedSize + CompressedHeaderSize);
+				ProgramBinaryOUT.Shrink();
+				FCompressedProgramBinaryHeader* Header = (FCompressedProgramBinaryHeader*)ProgramBinaryOUT.GetData();
+				Header->UncompressedSize = UncompressedProgramBinary.Num();
+			}
+			else
+			{
+				// failed, store the uncompressed version.
+				UE_LOG(LogRHI, Log, TEXT("Storing binary program uncompressed (%d, %d, %d)"), UncompressedProgramBinary.Num(), ProgramBinaryOUT.Num(), CompressedSize);
+				ProgramBinaryOUT.SetNumUninitialized(UncompressedProgramBinary.Num() + CompressedHeaderSize);
+				FCompressedProgramBinaryHeader* Header = (FCompressedProgramBinaryHeader*)ProgramBinaryOUT.GetData();
+				Header->UncompressedSize = FCompressedProgramBinaryHeader::NotCompressed;
+				FMemory::Memcpy(ProgramBinaryOUT.GetData() + sizeof(FCompressedProgramBinaryHeader), UncompressedProgramBinary.GetData(), UncompressedProgramBinary.Num());
+			}
+		}
+
+		bool GetCompressedProgramBinaryFromGLProgram(GLuint Program, TArray<uint8>& ProgramBinaryOUT)
+		{
+			// get uncompressed binary
+			TArray<uint8> UncompressedProgramBinary;
+			if (GetUncompressedProgramBinaryFromGLProgram(Program, UncompressedProgramBinary))
+			{
+				CompressProgramBinary(UncompressedProgramBinary, ProgramBinaryOUT);
+				return true;
+			}
+			return false;
+		}
+
+		FOpenGLProgramBinary GetProgramBinaryFromGLProgram(GLuint Program)
+		{
+			TArray<uint8> ProgramBinaryData;
+			if (IsStoringCompressedBinaryPrograms())
+			{
+				GetCompressedProgramBinaryFromGLProgram(Program, ProgramBinaryData);
+			}
+			else
+			{
+				GetUncompressedProgramBinaryFromGLProgram(Program, ProgramBinaryData);
+			}
+			return FOpenGLProgramBinary(MoveTemp(ProgramBinaryData));
+		}
+	}
+}
\ No newline at end of file
diff --git a/Engine/Source/Runtime/OpenGLDrv/Private/OpenGLBinaryProgramUtils.h b/Engine/Source/Runtime/OpenGLDrv/Private/OpenGLBinaryProgramUtils.h
new file mode 100644
--- /dev/null
+++ b/Engine/Source/Runtime/OpenGLDrv/Private/OpenGLBinaryProgramUtils.h
@@ -0,0 +1,84 @@
+// Copyright Epic Games, Inc. All Rights Reserved.
+
+/*=============================================================================
+	OpenGLBinaryProgramUtils.h
+=============================================================================*/
+
+#pragma once
+#include "OpenGLDrvPrivate.h"
+
+class FOpenGLProgramBinary
+{
+public:
+	FOpenGLProgramBinary() {};
+	FOpenGLProgramBinary(TArray<uint8>&& ProgramMemory) : DataView(ProgramMemory.GetData(), ProgramMemory.Num()), OwnedData(MoveTemp(ProgramMemory))
+	{
+		INC_MEMORY_STAT_BY(STAT_OpenGLShaderLRUProgramMemory, DataView.Num());
+		check(IsOwned());
+		check(IsValid());
+	}
+
+	FOpenGLProgramBinary(FOpenGLProgramBinary&& Src)
+	{
+		*this = MoveTemp(Src);
+	}
+
+	// No ownership, just a view to a region within the mmapped PSO cache.
+	FOpenGLProgramBinary(TArrayView<const uint8> ProgramMemory) : DataView(ProgramMemory) 
+	{
+		INC_MEMORY_STAT_BY(STAT_OpenGLShaderLRUProgramMemoryMapped, DataView.Num());
+		check(!IsOwned());
+		check(IsValid());
+	}
+
+	FOpenGLProgramBinary& operator= (FOpenGLProgramBinary&& rhs)
+	{
+		Swap(DataView, rhs.DataView);
+		Swap(OwnedData, rhs.OwnedData);
+		return *this;
+	}
+
+	~FOpenGLProgramBinary()
+	{
+		if(IsOwned())
+		{
+			DEC_MEMORY_STAT_BY(STAT_OpenGLShaderLRUProgramMemory, DataView.Num());
+		}
+		else
+		{
+			// This should only change when shutting down the PSO cache, i.e. DataView.IsEmpty() during normal use.
+			DEC_MEMORY_STAT_BY(STAT_OpenGLShaderLRUProgramMemoryMapped, DataView.Num());
+		}
+	}
+
+	bool IsValid() const { return DataView.Num() > 0; }
+	bool IsOwned() const { return OwnedData.Num() > 0; }
+	const TArrayView<const uint8> GetDataView() const { return DataView; }
+private:
+	// a view to the raw program binary, either points to OwnedData or an array view within the mmapped prebuilt PSO cache.
+	TArrayView<const uint8> DataView;
+	// We may own the data where programs are not 
+	TArray<uint8> OwnedData;
+};
+
+namespace UE
+{
+	namespace OpenGL
+	{
+		// Program Binary helpers.
+
+		bool IsStoringCompressedBinaryPrograms();
+
+		struct FCompressedProgramBinaryHeader
+		{
+			static const uint32 NotCompressed = 0xFFFFFFFF;
+			uint32 UncompressedSize;
+		};
+
+		FOpenGLProgramBinary GetProgramBinaryFromGLProgram(GLuint Program);
+		bool GetUncompressedProgramBinaryFromGLProgram(GLuint Program, TArray<uint8>& ProgramBinaryOUT);
+		bool GetCompressedProgramBinaryFromGLProgram(GLuint Program, TArray<uint8>& ProgramBinaryOUT);
+		void CompressProgramBinary(const TArrayView<const uint8>& UncompressedProgramBinary, TArray<uint8>& ProgramBinaryOUT);
+		bool UncompressCompressedBinaryProgram(const TArrayView<const uint8>& CompressedProgramBinary, TArray<uint8>& UncompressedProgramBinaryOUT);
+	}
+}
\ No newline at end of file
diff --git a/Engine/Source/Runtime/OpenGLDrv/Private/OpenGLCommands.cpp b/Engine/Source/Runtime/OpenGLDrv/Private/OpenGLCommands.cpp
--- a/Engine/Source/Runtime/OpenGLDrv/Private/OpenGLCommands.cpp
+++ b/Engine/Source/Runtime/OpenGLDrv/Private/OpenGLCommands.cpp
@@ -2813,6 +2813,19 @@ void FOpenGLDynamicRHI::RHISetGraphicsPipelineState(FRHIGraphicsPipelineState* G
 
 	auto& PsoInit = FallbackGraphicsState->Initializer;
 
+	if (PsoInit.bFromPSOFileCache)
+	{
+		// If we're from the PSO cache we're just preparing the PSO and do not need to set the state.
+		RHICreateBoundShaderState_internal(
+			PsoInit.BoundShaderState.VertexDeclarationRHI,
+			PsoInit.BoundShaderState.VertexShaderRHI,
+			PsoInit.BoundShaderState.PixelShaderRHI,
+			PsoInit.BoundShaderState.GetGeometryShader(),
+			PsoInit.bFromPSOFileCache
+		);
+		return;
+	}
+
 	RHISetBoundShaderState(
 		RHICreateBoundShaderState_internal(
 			PsoInit.BoundShaderState.VertexDeclarationRHI,
diff --git a/Engine/Source/Runtime/OpenGLDrv/Private/OpenGLDevice.cpp b/Engine/Source/Runtime/OpenGLDrv/Private/OpenGLDevice.cpp
--- a/Engine/Source/Runtime/OpenGLDrv/Private/OpenGLDevice.cpp
+++ b/Engine/Source/Runtime/OpenGLDrv/Private/OpenGLDevice.cpp
@@ -1351,7 +1351,7 @@ void FOpenGLDynamicRHI::Init()
 	GRHISupportsMultithreadedShaderCreation = false;
 
 	FOpenGLProgramBinaryCache::Initialize();
-	RegisterSharedShaderCodeDelegates();
+
 	InitializeStateResources();
 
 	// Create a default point sampler state for internal use.
@@ -1422,8 +1422,6 @@ void FOpenGLDynamicRHI::Shutdown()
 	DestroyShadersAndPrograms();
 	PlatformDestroyOpenGLDevice(PlatformDevice);
 
-	UnregisterSharedShaderCodeDelegates();
-
 	PrivateOpenGLDevicePtr = NULL;
 }
 
diff --git a/Engine/Source/Runtime/OpenGLDrv/Private/OpenGLDrvPrivate.h b/Engine/Source/Runtime/OpenGLDrv/Private/OpenGLDrvPrivate.h
--- a/Engine/Source/Runtime/OpenGLDrv/Private/OpenGLDrvPrivate.h
+++ b/Engine/Source/Runtime/OpenGLDrv/Private/OpenGLDrvPrivate.h
@@ -94,7 +94,7 @@ DECLARE_MEMORY_STAT_EXTERN(TEXT("Uniform buffer pool memory"),STAT_OpenGLFreeUni
 DECLARE_CYCLE_STAT_EXTERN(TEXT("Emulated Uniform buffer time"), STAT_OpenGLEmulatedUniformBufferTime,STATGROUP_OpenGLRHI, );
 DECLARE_DWORD_ACCUMULATOR_STAT_EXTERN(TEXT("Uniform buffer pool num free"),STAT_OpenGLNumFreeUniformBuffers,STATGROUP_OpenGLRHI, );
 DECLARE_CYCLE_STAT_EXTERN(TEXT("Time for first draw of shader programs"), STAT_OpenGLShaderFirstDrawTime,STATGROUP_OpenGLRHI, );
-DECLARE_MEMORY_STAT_EXTERN(TEXT("Program binary memory"), STAT_OpenGLProgramBinaryMemory, STATGROUP_OpenGLRHI, );
+DECLARE_MEMORY_STAT_EXTERN(TEXT("Active Program binary memory (estimate driver use)"), STAT_OpenGLProgramBinaryMemory, STATGROUP_OpenGLRHI, );
 DECLARE_DWORD_ACCUMULATOR_STAT_EXTERN(TEXT("GL Program count"), STAT_OpenGLProgramCount, STATGROUP_OpenGLRHI, );
 DECLARE_CYCLE_STAT_EXTERN(TEXT("Program get from cache time"),STAT_OpenGLUseCachedProgramTime,STATGROUP_OpenGLRHI, );
 DECLARE_CYCLE_STAT_EXTERN(TEXT("Program create from binary time"),STAT_OpenGLCreateProgramFromBinaryTime,STATGROUP_OpenGLRHI, );
@@ -103,10 +103,9 @@ DECLARE_CYCLE_STAT_EXTERN(TEXT("Program LRU cache eviction time"), STAT_OpenGLSh
 DECLARE_CYCLE_STAT_EXTERN(TEXT("Program LRU cache miss time"), STAT_OpenGLShaderLRUMissTime, STATGROUP_OpenGLRHI, );
 DECLARE_DWORD_ACCUMULATOR_STAT_EXTERN(TEXT("Program LRU count"), STAT_OpenGLShaderLRUProgramCount, STATGROUP_OpenGLRHI, );
 DECLARE_DWORD_ACCUMULATOR_STAT_EXTERN(TEXT("Program LRU evicted count"), STAT_OpenGLShaderLRUEvictedProgramCount, STATGROUP_OpenGLRHI, );
-DECLARE_DWORD_ACCUMULATOR_STAT_EXTERN(TEXT("Program LRU evicted by scope"), STAT_OpenGLShaderLRUScopeEvictedProgramCount, STATGROUP_OpenGLRHI, );
 DECLARE_DWORD_ACCUMULATOR_STAT_EXTERN(TEXT("Program LRU miss count"), STAT_OpenGLShaderLRUMissCount, STATGROUP_OpenGLRHI, );
-DECLARE_DWORD_ACCUMULATOR_STAT_EXTERN(TEXT("Program LRU evictions saved by latency"), STAT_OpenGLShaderLRUEvictionDelaySavedCount, STATGROUP_OpenGLRHI, );
-DECLARE_MEMORY_STAT_EXTERN(TEXT("Program LRU binary memory"), STAT_OpenGLShaderLRUProgramMemory, STATGROUP_OpenGLRHI, );
+DECLARE_MEMORY_STAT_EXTERN(TEXT("Program LRU memory (evicted, heap)"), STAT_OpenGLShaderLRUProgramMemory, STATGROUP_OpenGLRHI, );
+DECLARE_MEMORY_STAT_EXTERN(TEXT("Program LRU mem mapped (evicted, filemapped)"), STAT_OpenGLShaderLRUProgramMemoryMapped, STATGROUP_OpenGLRHI, );
 
 #if OPENGLRHI_DETAILED_STATS
 DECLARE_CYCLE_STAT_EXTERN(TEXT("DrawPrimitive Time"),STAT_OpenGLDrawPrimitiveTime,STATGROUP_OpenGLRHI, );
diff --git a/Engine/Source/Runtime/OpenGLDrv/Private/OpenGLProgramBinaryFileCache.cpp b/Engine/Source/Runtime/OpenGLDrv/Private/OpenGLProgramBinaryFileCache.cpp
--- a/Engine/Source/Runtime/OpenGLDrv/Private/OpenGLProgramBinaryFileCache.cpp
+++ b/Engine/Source/Runtime/OpenGLDrv/Private/OpenGLProgramBinaryFileCache.cpp
@@ -16,13 +16,26 @@
 #include "GlobalShader.h"
 #include "SceneUtils.h"
 #include "PsoLruCache.h"
+#include "OpenGLBinaryProgramUtils.h"
+#include <Serialization/StaticMemoryReader.h>
+#include "ProfilingDebugging/ScopedTimers.h"
+#include <Async/MappedFileHandle.h>
 
 #if PLATFORM_ANDROID
 #include "Android/AndroidPlatformMisc.h"
 #endif
-
-static const uint32 GBinaryProgramFileVersion = 4;
  
+static bool GMemoryMapGLProgramCache = true;
+static FAutoConsoleVariableRef CVarMemoryMapGLProgramCache(
+	TEXT("r.OpenGL.MemoryMapGLProgramCache"),
+	GMemoryMapGLProgramCache,
+	TEXT("If true enabled memory mapping of the GL program binary cache. (default)\n")
+	TEXT("If false then upon opening the binary cache all programs are loaded into memory.\n")
+	TEXT("When enabled this can reduce RSS pressure when combined with program LRU. (see r.OpenGL.EnableProgramLRUCache).")
+	,
+	ECVF_ReadOnly | ECVF_RenderThreadSafe
+);
+
 TAutoConsoleVariable<int32> FOpenGLProgramBinaryCache::CVarPBCEnable(
 	TEXT("r.ProgramBinaryCache.Enable"),
 #if PLATFORM_ANDROID
@@ -36,25 +49,17 @@ TAutoConsoleVariable<int32> FOpenGLProgramBinaryCache::CVarPBCEnable(
 
 TAutoConsoleVariable<int32> FOpenGLProgramBinaryCache::CVarRestartAndroidAfterPrecompile(
 	TEXT("r.ProgramBinaryCache.RestartAndroidAfterPrecompile"),
-	1,	// Enabled by default on Android.
-	TEXT("If true, Android apps will restart after precompiling the binary program cache. Enabled by default only on Android"),
+	0,
+	TEXT("If true, Android apps will restart after precompiling the binary program cache."),
 	ECVF_ReadOnly | ECVF_RenderThreadSafe
 );
 
-static TAutoConsoleVariable<int32> CVarUseExistingBinaryFileCache(
-	TEXT("r.OpenGL.UseExistingBinaryFileCache"),
-	1,
-	TEXT("When generating a new binary cache (such as when Shader Pipeline Cache Version Guid changes) use the existing binary file cache to speed up generation of the new cache.\n")
-	TEXT("0: Always rebuild binary file cache when Pipeline Cache Version Guid changes.\n")
-	TEXT("1: When Pipeline Cache Version Guid changes re-use programs from the existing binary cache where possible (default)."),
-	ECVF_RenderThreadSafe);
-
-static int32 GMaxShaderLibProcessingTimeMS = 10;
-static FAutoConsoleVariableRef CVarMaxShaderLibProcessingTime(
-	TEXT("r.OpenGL.MaxShaderLibProcessingTime"),
-	GMaxShaderLibProcessingTimeMS,
-	TEXT("The maximum time per frame to process shader library requests in milliseconds.\n")
-	TEXT("default 10ms. Note: Driver compile time for a single program may exceed this limit."),
+static int32 GMaxBinaryProgramLoadTimeMS = 3;
+static FAutoConsoleVariableRef CVarMaxBinaryProgramLoadTime(
+	TEXT("r.OpenGL.MaxBinaryProgramLoadTime"),
+	GMaxBinaryProgramLoadTimeMS,
+	TEXT("The maximum time per frame to transfer programs from the binary program cache to the GL RHI. in milliseconds.\n")
+	TEXT("default 3ms. Note: Driver compile time for programs may exceed this limit if you're not using the LRU."),
 	ECVF_RenderThreadSafe
 );
 
@@ -62,26 +67,90 @@ namespace UE
 {
 	namespace OpenGL
 	{
+		bool CanMemoryMapGLProgramCache()
+		{
+			return FPlatformProperties::SupportsMemoryMappedFiles() && GMemoryMapGLProgramCache;
+		}
+
+		extern void OnGLProgramLoadedFromBinaryCache(const FOpenGLProgramKey& ProgramKey, TUniqueObj<FOpenGLProgramBinary>&& ProgramBinaryData);
+
 		bool AreBinaryProgramsCompressed()
 		{
-			static const auto CVarStoreCompressedBinaries = IConsoleManager::Get().FindConsoleVariable(TEXT("r.OpenGL.StoreCompressedProgramBinaries"));
-			return CVarStoreCompressedBinaries->GetInt() != 0;
-		}		
+			static const auto StoreCompressedBinariesCVar = IConsoleManager::Get().FindConsoleVariable(TEXT("r.OpenGL.StoreCompressedProgramBinaries"));
+			return StoreCompressedBinariesCVar->GetInt() != 0;
+		}
+
+		static const uint32 GBinaryProgramFileVersion = 5;
+		struct FBinaryCacheFileHeader
+		{
+			uint32 Version = 0;
+			FGuid BinaryCacheGuid;
+			bool bCacheUsesCompressedBinaries;
+			uint32 ProgramCount;
+
+			static FBinaryCacheFileHeader CreateHeader(const FGuid& BinaryCacheGuidIn, uint32 NumPrograms)
+			{
+				FBinaryCacheFileHeader NewHeader;
+				NewHeader.Version = GBinaryProgramFileVersion;
+				NewHeader.BinaryCacheGuid = BinaryCacheGuidIn;
+				NewHeader.bCacheUsesCompressedBinaries = UE::OpenGL::AreBinaryProgramsCompressed();
+				NewHeader.ProgramCount = NumPrograms;
+				return NewHeader;
+			}
+
+			friend FArchive& operator<<(FArchive& Ar, FBinaryCacheFileHeader& Header)
+			{
+				Ar << Header.Version;
+				check(Ar.IsLoading() || Header.Version == GBinaryProgramFileVersion); // This should always be correct when saving.
+				if (Header.Version == GBinaryProgramFileVersion)
+				{
+					Ar << Header.BinaryCacheGuid;
+					Ar << Header.bCacheUsesCompressedBinaries;
+					Ar << Header.ProgramCount;
+				}
+
+				return Ar;
+			}
+
+			bool IsValid() const
+			{
+				return Version == GBinaryProgramFileVersion;
+			}
+		};
 	}
 }
 
+// This contains the mapping for a binary program cache file.
+// It also contains a list of programs that the cache contains.
+class FOpenGLProgramBinaryMapping : public FThreadSafeRefCountedObject
+{
+public:
+	FOpenGLProgramBinaryMapping(TUniquePtr<IMappedFileHandle> MappedCacheFileIn, TUniquePtr<IMappedFileRegion> MappedRegionIn, uint32 ProgramCountIfKnown) : MappedCacheFile(MoveTemp(MappedCacheFileIn)), MappedRegion(MoveTemp(MappedRegionIn)) { Content.Reserve(ProgramCountIfKnown); }
+
+	TArrayView<const uint8> GetView(uint32 FileOffset, uint32 NumBytes) const
+	{
+		check(FileOffset + NumBytes <= MappedRegion->GetMappedSize());
+		return TArrayView<const uint8>(MappedRegion->GetMappedPtr() + FileOffset, NumBytes);
+	}
+
+	void AddProgramKey(const class FOpenGLProgramKey& KeyIn) { check(!Content.Contains(KeyIn)); Content.Add(KeyIn); }
+	bool HasValidMapping() const { return MappedRegion.IsValid() && MappedCacheFile.IsValid(); }
+private:
+	TUniquePtr<IMappedFileHandle> MappedCacheFile;
+	TUniquePtr<IMappedFileRegion> MappedRegion;
+	TSet<FOpenGLProgramKey> Content;
+};
+
+
 static FCriticalSection GProgramBinaryFileCacheCS;
 
-FOpenGLProgramBinaryCache::FPreviousGLProgramBinaryCacheInfo::FPreviousGLProgramBinaryCacheInfo() : NumberOfOldEntriesReused(0) {}
-FOpenGLProgramBinaryCache::FPreviousGLProgramBinaryCacheInfo::FPreviousGLProgramBinaryCacheInfo(FOpenGLProgramBinaryCache::FPreviousGLProgramBinaryCacheInfo&&) = default;
-FOpenGLProgramBinaryCache::FPreviousGLProgramBinaryCacheInfo& FOpenGLProgramBinaryCache::FPreviousGLProgramBinaryCacheInfo::operator = (FOpenGLProgramBinaryCache::FPreviousGLProgramBinaryCacheInfo&&) = default;
-FOpenGLProgramBinaryCache::FPreviousGLProgramBinaryCacheInfo::~FPreviousGLProgramBinaryCacheInfo() = default;
+// guards the container that collects scanned programs and send to RHIT
+static FCriticalSection GPendingGLProgramCreateRequestsCS;
 
 FOpenGLProgramBinaryCache* FOpenGLProgramBinaryCache::CachePtr = nullptr;
 
 FOpenGLProgramBinaryCache::FOpenGLProgramBinaryCache(const FString& InCachePath)
 	: CachePath(InCachePath)
-	, BinaryCacheAsyncReadFileHandle(nullptr)
 	, BinaryCacheWriteFileHandle(nullptr)
 	, BinaryFileState(EBinaryFileState::Uninitialized)
 {
@@ -97,6 +166,13 @@ FOpenGLProgramBinaryCache::FOpenGLProgramBinaryCache(const FString& InCachePath)
 	// Some devices report binary compatibility errors after minor OS updates even though the GL driver version has not changed.
 	const FString BuildNumber = FAndroidMisc::GetDeviceBuildNumber();
 	HashString.Append(BuildNumber);
+
+	// Optional configrule variable for triggering a rebuild of the cache.
+	const FString* ConfigRulesGLProgramKey = FAndroidMisc::GetConfigRulesVariable(TEXT("OpenGLProgramCacheKey"));
+	if (ConfigRulesGLProgramKey && !ConfigRulesGLProgramKey->IsEmpty())
+	{
+		HashString.Append(*ConfigRulesGLProgramKey);
+	}
 #endif
 
 	FSHAHash VersionHash;
@@ -107,10 +183,6 @@ FOpenGLProgramBinaryCache::FOpenGLProgramBinaryCache(const FString& InCachePath)
 
 FOpenGLProgramBinaryCache::~FOpenGLProgramBinaryCache()
 {
-	if (BinaryCacheAsyncReadFileHandle)
-	{
-		delete BinaryCacheAsyncReadFileHandle;
-	}
 	if (BinaryCacheWriteFileHandle)
 	{
 		delete BinaryCacheWriteFileHandle;
@@ -211,7 +283,11 @@ static FAutoConsoleVariableRef CVarNumRemoteProgramCompileServices(
 void FOpenGLProgramBinaryCache::OnShaderPipelineCacheOpened(FString const& Name, EShaderPlatform Platform, uint32 Count, const FGuid& VersionGuid, FShaderPipelineCache::FShaderCachePrecompileContext& ShaderCachePrecompileContext)
 {
 	UE_LOG(LogRHI, Log, TEXT("Scanning Binary program cache, using Shader Pipeline Cache version %s"), *VersionGuid.ToString());
-	ScanProgramCacheFile(VersionGuid);
+
+	FScopeLock Lock(&GProgramBinaryFileCacheCS);
+	check(CurrentShaderPipelineCacheVersionGuid == FGuid());
+	CurrentShaderPipelineCacheVersionGuid = VersionGuid;
+	ScanProgramCacheFile();
 	if (IsBuildingCache_internal())
 	{
 #if PLATFORM_ANDROID
@@ -224,97 +300,60 @@ void FOpenGLProgramBinaryCache::OnShaderPipelineCacheOpened(FString const& Name,
 	}
 }
 
+void FOpenGLProgramBinaryCache::Reset()
+{
+	check(!BinaryCacheWriteFileHandle);
+	BinaryFileState = EBinaryFileState::Uninitialized;
+ 	ProgramsWrittenToOutputCache.Empty();
+}
+
 void FOpenGLProgramBinaryCache::OnShaderPipelineCachePrecompilationComplete(uint32 Count, double Seconds, const FShaderPipelineCache::FShaderCachePrecompileContext& ShaderCachePrecompileContext)
 {
+	FScopeLock Lock(&GProgramBinaryFileCacheCS);
+
 #if PLATFORM_ANDROID
 	FAndroidOpenGL::StopRemoteCompileServices();
 #endif
 
 	UE_LOG(LogRHI, Log, TEXT("OnShaderPipelineCachePrecompilationComplete: %d shaders"), Count);
 
-	// Want to ignore any subsequent Shader Pipeline Cache opening/closing, eg when loading modules
-	FShaderPipelineCache::GetCacheOpenedDelegate().Remove(OnShaderPipelineCacheOpenedDelegate);
-	FShaderPipelineCache::GetPrecompilationCompleteDelegate().Remove(OnShaderPipelineCachePrecompilationCompleteDelegate);
-	OnShaderPipelineCacheOpenedDelegate.Reset();
-	OnShaderPipelineCachePrecompilationCompleteDelegate.Reset();
-
-	check(IsBuildingCache_internal() || BinaryFileState == EBinaryFileState::ValidCacheFile);
+	const bool bFinishedBuildingCache = IsBuildingCache_internal();
+	check(bFinishedBuildingCache || BinaryFileState == EBinaryFileState::ValidCacheFile);
 
-	if (IsBuildingCache_internal())
+	if (bFinishedBuildingCache)
 	{
 		CloseWriteHandle();
 
-#if PLATFORM_ANDROID
-		FAndroidMisc::bNeedsRestartAfterPSOPrecompile = true;
-		if (CVarRestartAndroidAfterPrecompile.GetValueOnAnyThread() == 1)
+		if(!ProgramsWrittenToOutputCache.IsEmpty())
 		{
+#if PLATFORM_ANDROID
+			if (CVarRestartAndroidAfterPrecompile.GetValueOnAnyThread() == 1)
+			{
+				FAndroidMisc::bNeedsRestartAfterPSOPrecompile = true;
 #if USE_ANDROID_JNI
-			extern void AndroidThunkCpp_RestartApplication(const FString & IntentString);
-			AndroidThunkCpp_RestartApplication(TEXT(""));
+				extern void AndroidThunkCpp_RestartApplication(const FString & IntentString);
+				AndroidThunkCpp_RestartApplication(TEXT(""));
 #endif
-		}
+			}
 #endif
-		OpenAsyncReadHandle();
-		BinaryFileState = EBinaryFileState::ValidCacheFile;
-	}
-}
-
-// contains runtime and file information for a single program entry in the cache file.
-struct FGLProgramBinaryFileCacheEntry
-{
-	struct FGLProgramBinaryFileCacheFileInfo
-	{
-		// contains location info for a program in the binary cache file
-		FOpenGLProgramKey ShaderHasheSet;
-		uint32 ProgramOffset;
-		uint32 ProgramSize;
-
-		FGLProgramBinaryFileCacheFileInfo() : ProgramOffset(0), ProgramSize(0) { }
-
-		friend bool operator == (const FGLProgramBinaryFileCacheFileInfo& A, const FGLProgramBinaryFileCacheFileInfo& B)
-		{
-			return A.ShaderHasheSet == B.ShaderHasheSet && A.ProgramOffset == B.ProgramOffset && A.ProgramSize == B.ProgramSize;
 		}
-	} FileInfo;
-
-	// program read request.
-	TWeakPtr<IAsyncReadRequest, ESPMode::ThreadSafe> ReadRequest;
-	// read data
-	TArray<uint8> ProgramBinaryData;
-	// debugging use only, index of program as encountered during scan, -1 if new.
-	int32 ProgramIndex;
-
-	enum class EGLProgramState : uint8
-	{
-		Unset,
-		ProgramStored, // program exists in binary cache but has not yet been loaded
-		ProgramLoading,	// program has started async loading.
-		ProgramLoaded,	// program has loaded and is ready for GL object creation.
-		ProgramAvailable, // program has loaded from binary cache and is available for use with GL.
-		ProgramComplete // program has been either added by rhi or handed over to rhi after being made available to it.
-	};
-	EGLProgramState GLProgramState;
-
-	// if != 0 then prepared runtime GL program name:
-	GLuint GLProgramId;
-
-	FGLProgramBinaryFileCacheEntry() : ProgramIndex(-1), GLProgramState(EGLProgramState::Unset), GLProgramId(0) { }
+		// Reset state and scan the file we just produced.
+		Reset();
+		ScanProgramCacheFile();
 
-	friend uint32 GetTypeHash(const FGLProgramBinaryFileCacheEntry& Key)
-	{
-		return FCrc::MemCrc32(&Key, sizeof(Key));
+		// unset current cache.
+		CurrentShaderPipelineCacheVersionGuid = FGuid();
 	}
-};
-
-static FCriticalSection GPendingGLProgramCreateRequestsCS;
+}
 
 // Scan the binary cache file and build a record of all programs.
-void FOpenGLProgramBinaryCache::ScanProgramCacheFile(const FGuid& ShaderPipelineCacheVersionGuid)
+void FOpenGLProgramBinaryCache::ScanProgramCacheFile()
 {
+	//FScopedDurationTimeLogger Timer(TEXT("ScanProgramCacheFile"));
+
 	UE_LOG(LogRHI, Log, TEXT("OnShaderScanProgramCacheFile"));
-	FScopeLock Lock(&GProgramBinaryFileCacheCS);
 	FString ProgramCacheFilename = GetProgramBinaryCacheFilePath();
-	FString ProgramCacheFilenameTemp = GetProgramBinaryCacheFilePath() + TEXT(".scan");
+	FString ProgramCacheFilenameTemp = ProgramCacheFilename + TEXT(".scan");
 
 	IPlatformFile& PlatformFile = FPlatformFileManager::Get().GetPlatformFile();
 
@@ -327,184 +366,121 @@ void FOpenGLProgramBinaryCache::ScanProgramCacheFile(const FGuid& ShaderPipeline
 	PlatformFile.DeleteFile(*ProgramCacheFilenameTemp);
 	PlatformFile.MoveFile(*ProgramCacheFilenameTemp, *ProgramCacheFilename);
 
-	TUniquePtr<FArchive> FileReader(IFileManager::Get().CreateFileReader(*ProgramCacheFilenameTemp));
-	if (FileReader)
+	TUniquePtr<FArchive> BinaryProgramReader = nullptr;
+	BinaryProgramReader = TUniquePtr<FArchive>( IFileManager::Get().CreateFileReader(*ProgramCacheFilenameTemp) );
+
+	if (BinaryProgramReader)
 	{
-		UE_LOG(LogRHI, Log, TEXT("OnShaderScanProgramCacheFile : Opened %s"), *ProgramCacheFilenameTemp);
-		FArchive& Ar = *FileReader;
-		uint32 Version = 0;
-		Ar << Version;
-		if (Version == GBinaryProgramFileVersion)
+		FArchive& Ar = *BinaryProgramReader;
+		UE::OpenGL::FBinaryCacheFileHeader BinaryCacheHeader;
+		Ar << BinaryCacheHeader;
+		if (BinaryCacheHeader.IsValid())
 		{
-			FGuid BinaryCacheGuid;
-			Ar << BinaryCacheGuid;
-			bool bCacheUsesCompressedBinaries;
-			Ar << bCacheUsesCompressedBinaries;
-
 			const bool bUseCompressedProgramBinaries = UE::OpenGL::AreBinaryProgramsCompressed();
-			bBinaryFileIsValid = (bUseCompressedProgramBinaries == bCacheUsesCompressedBinaries);
-			bBinaryFileIsValidAndGuidMatch = bBinaryFileIsValid && (!ShaderPipelineCacheVersionGuid.IsValid() || ShaderPipelineCacheVersionGuid == BinaryCacheGuid);
+			bBinaryFileIsValid = (bUseCompressedProgramBinaries == BinaryCacheHeader.bCacheUsesCompressedBinaries);
+			bBinaryFileIsValidAndGuidMatch = bBinaryFileIsValid && (!CurrentShaderPipelineCacheVersionGuid.IsValid() || CurrentShaderPipelineCacheVersionGuid == BinaryCacheHeader.BinaryCacheGuid);
 
-			if (CVarUseExistingBinaryFileCache.GetValueOnAnyThread() == 0 && bBinaryFileIsValidAndGuidMatch == false)
+			if (bBinaryFileIsValidAndGuidMatch)
 			{
-				// If we dont want to use the existing binary cache and the guids have changed then rebuild the binary file.
-				bBinaryFileIsValid = false;
-			}
-		}
-
-		if (bBinaryFileIsValid)
-		{
-			const uint32 ProgramBinaryStart = Ar.Tell();
-
-			// Search the file for the end record.
-			bool bFoundEndRecord = false;
-			int32 ProgramIndex = 0;
-			while (!Ar.AtEnd())
-			{
-				check(bFoundEndRecord == false); // There should be no additional data after the eof record.
-
-				FOpenGLProgramKey ProgramKey;
-				uint32 ProgramBinarySize = 0;
-				Ar << ProgramKey;
-				Ar << ProgramBinarySize;
-				uint32 ProgramBinaryOffset = Ar.Tell();
-				if (ProgramBinarySize == 0)
+				int32 ProgramIndex = 0;
+				FScopeLock Lock(&GPendingGLProgramCreateRequestsCS);
+				ScannedBinaryPrograms.Reserve(BinaryCacheHeader.ProgramCount);
+				if (BinaryCacheHeader.ProgramCount > 0)
 				{
-					if (ProgramKey == FOpenGLProgramKey())
-					{
-						bFoundEndRecord = true;
-					}
-					else
+					TRefCountPtr<FOpenGLProgramBinaryMapping> CurrentMapping;
+					// success
 					{
-						// Note: This should not happen with new code. We can no longer write out records with 0 program size. see AppendProgramBinaryFile.
-						UE_LOG(LogRHI, Warning, TEXT("FOpenGLProgramBinaryCache::ScanProgramCacheFile : encountered 0 sized program during binary program cache scan"));
+						TUniquePtr<IMappedFileHandle> MappedCacheFile;
+						TUniquePtr<IMappedFileRegion> MappedRegion;
+						if(UE::OpenGL::CanMemoryMapGLProgramCache())
+						{
+							MappedCacheFile = TUniquePtr<IMappedFileHandle>(FPlatformFileManager::Get().GetPlatformFile().OpenMapped(*ProgramCacheFilenameTemp));
+							if (ensure(MappedCacheFile.IsValid()))
+							{
+								MappedRegion = TUniquePtr<IMappedFileRegion>(MappedCacheFile->MapRegion());
+								check(MappedRegion.IsValid());
+							}
+						}
+						CurrentMapping = new FOpenGLProgramBinaryMapping(MoveTemp(MappedCacheFile), MoveTemp(MappedRegion), BinaryCacheHeader.ProgramCount);
 					}
-				}
-				Ar.Seek(ProgramBinaryOffset + ProgramBinarySize);
-			}
 
-			if (bFoundEndRecord)
-			{
-				Ar.Seek(ProgramBinaryStart);
-				while (!Ar.AtEnd())
-				{
-					FOpenGLProgramKey ProgramKey;
-					uint32 ProgramBinarySize = 0;
-					Ar << ProgramKey;
-					Ar << ProgramBinarySize;
+					UE_LOG(LogRHI, Log, TEXT("OnShaderScanProgramCacheFile : %s %s"), CurrentMapping->HasValidMapping() ? TEXT("mapped") : TEXT("opened"), *ProgramCacheFilenameTemp);
 
-					if (ProgramBinarySize > 0)
+					MappedCacheFiles.Add(CurrentShaderPipelineCacheVersionGuid, CurrentMapping);
+					while (!Ar.AtEnd())
 					{
-						FGLProgramBinaryFileCacheEntry* NewEntry = new FGLProgramBinaryFileCacheEntry();
-						NewEntry->FileInfo.ShaderHasheSet = ProgramKey;
-						NewEntry->ProgramIndex = ProgramIndex++;
-
-						uint32 ProgramBinaryOffset = Ar.Tell();
-						NewEntry->FileInfo.ProgramSize = ProgramBinarySize;
-						NewEntry->FileInfo.ProgramOffset = ProgramBinaryOffset;
-
-						if (bBinaryFileIsValidAndGuidMatch)
+						FOpenGLProgramKey ProgramKey;
+						uint32 ProgramBinarySize = 0;
+						Ar << ProgramKey;
+						Ar << ProgramBinarySize;
+						check(ProgramKey != FOpenGLProgramKey());
+						if (ensure(ProgramBinarySize > 0))
 						{
-							ProgramEntryContainer.Emplace(TUniquePtr<FGLProgramBinaryFileCacheEntry>(NewEntry));
+							ProgramIndex++;
+							uint32 ProgramBinaryOffset = Ar.Tell();
+							CurrentMapping->AddProgramKey(ProgramKey);
 
-							// check to see if any of the shaders are already loaded and so we should serialize the binary
-							bool bAllShadersLoaded = true;
-							for (int32 i = 0; i < CrossCompiler::NUM_SHADER_STAGES && bAllShadersLoaded; i++)
-							{
-								bAllShadersLoaded = ProgramKey.ShaderHashes[i] == FSHAHash() || ShaderIsLoaded(ProgramKey.ShaderHashes[i]);
-							}
-							if (bAllShadersLoaded)
+							UE_LOG(LogRHI, VeryVerbose, TEXT(" scan found PSO %s - %d"), *ProgramKey.ToString(), ProgramBinarySize);
+
+							if(CurrentMapping->HasValidMapping())
 							{
-								FPlatformMisc::LowLevelOutputDebugStringf(TEXT("*** All shaders for %s already loaded\n"), *ProgramKey.ToString());
-								NewEntry->ProgramBinaryData.AddUninitialized(ProgramBinarySize);
-								Ar.Serialize(NewEntry->ProgramBinaryData.GetData(), ProgramBinarySize);
-								NewEntry->GLProgramState = FGLProgramBinaryFileCacheEntry::EGLProgramState::ProgramLoaded;
-								CompleteLoadedGLProgramRequest_internal(NewEntry);
+								ScannedBinaryPrograms.Emplace(ProgramKey, TUniqueObj<FOpenGLProgramBinary>(CurrentMapping->GetView(ProgramBinaryOffset, ProgramBinarySize)));
+								Ar.Seek(ProgramBinaryOffset + ProgramBinarySize);
 							}
 							else
 							{
-								NewEntry->GLProgramState = FGLProgramBinaryFileCacheEntry::EGLProgramState::ProgramStored;
+								check(!UE::OpenGL::CanMemoryMapGLProgramCache());
+								TArray<uint8> ProgramBytes;
+								ProgramBytes.SetNumUninitialized(ProgramBinarySize);
+								Ar.Serialize(ProgramBytes.GetData(), ProgramBinarySize);
+								ScannedBinaryPrograms.Emplace(ProgramKey, TUniqueObj<FOpenGLProgramBinary>(MoveTemp(ProgramBytes)));
 							}
-							AddProgramFileEntryToMap(NewEntry);
-						}
-						else
-						{
-							check(!PreviousBinaryCacheInfo.ProgramToOldBinaryCacheMap.Contains(ProgramKey));
-							PreviousBinaryCacheInfo.ProgramToOldBinaryCacheMap.Emplace(ProgramKey, TUniquePtr<FGLProgramBinaryFileCacheEntry>(NewEntry));
 						}
-						Ar.Seek(ProgramBinaryOffset + ProgramBinarySize);
 					}
-				}
 
-				if (bBinaryFileIsValidAndGuidMatch)
-				{
-					UE_LOG(LogRHI, Log, TEXT("Program Binary cache: Found %d cached programs, end record found: %d"), ProgramIndex, (uint32)bFoundEndRecord);
-					FileReader->Close();
+					UE_LOG(LogRHI, VeryVerbose, TEXT("Program Binary cache: Found %d cached programs"), ProgramIndex);
+					UE_CLOG(ProgramIndex != BinaryCacheHeader.ProgramCount, LogRHI, Error, TEXT("Program Binary cache: Mismatched program count! expected: %d"), BinaryCacheHeader.ProgramCount);
+
+					BinaryProgramReader->Close();
 					// Rename the file back after a successful scan.
 					PlatformFile.MoveFile(*ProgramCacheFilename, *ProgramCacheFilenameTemp);
+					BinaryFileState = EBinaryFileState::ValidCacheFile;
 				}
 				else
 				{
-					UE_LOG(LogRHI, Log, TEXT("Program Binary cache: ShaderPipelineCache changed, regenerating for new pipeline cache. Existing cache contains %d programs, using it to populate."), PreviousBinaryCacheInfo.ProgramToOldBinaryCacheMap.Num());
-					// Not closing the scan source file, we're using it to move shaders from the old cache.
-					PreviousBinaryCacheInfo.OldCacheArchive = MoveTemp(FileReader);
-					PreviousBinaryCacheInfo.OldCacheFilename = ProgramCacheFilenameTemp;
+					// failed to find sentinel record, the file was not finalized.
+					UE_LOG(LogRHI, Warning, TEXT("ScanProgramCacheFile - incomplete or empty binary cache file encountered. Rebuilding binary program cache."));
+					BinaryProgramReader->Close();
+					bBinaryFileIsValid = false;
+					bBinaryFileIsValidAndGuidMatch = false;
+					PlatformFile.DeleteFile(*ProgramCacheFilenameTemp);
 				}
 			}
 			else
 			{
-				// failed to find sentinel record, the file was not finalized.
-				UE_LOG(LogRHI, Warning, TEXT("ScanProgramCacheFile - incomplete binary cache file encountered. Rebuilding binary program cache."));
-				FileReader->Close();
-				bBinaryFileIsValid = false;
-				bBinaryFileIsValidAndGuidMatch = false;
+				UE_LOG(LogRHI, Log, TEXT("OnShaderScanProgramCacheFile : binary file found but is invalid (%s, %s)"), *CurrentShaderPipelineCacheVersionGuid.ToString(), *BinaryCacheHeader.BinaryCacheGuid.ToString());
 			}
 		}
-
-		if (!bBinaryFileIsValid)
+		else
 		{
 			UE_LOG(LogRHI, Log, TEXT("OnShaderScanProgramCacheFile : binary file version invalid"));
 		}
-
-		if (bBinaryFileIsValidAndGuidMatch)
-		{
-			OpenAsyncReadHandle();
-			BinaryFileState = EBinaryFileState::ValidCacheFile;
-		}
 	}
 	else
 	{
 		UE_LOG(LogRHI, Log, TEXT("OnShaderScanProgramCacheFile : Failed to open %s"), *ProgramCacheFilename);
 	}
 
-	if (!bBinaryFileIsValid)
-	{
-		// Attempt to remove any existing binary cache or temp files (eg for different driver version)
-		UE_LOG(LogRHI, Log, TEXT("Deleting binary program cache folder: %s"), *CachePath);
-		PlatformFile.DeleteDirectoryRecursively(*CachePath);
-
-		// Create
-		if (!PlatformFile.CreateDirectoryTree(*CachePath))
-		{
-			UE_LOG(LogRHI, Warning, TEXT("Failed to create directory for a program binary cache. Cache will be disabled: %s"), *CachePath);
-			return;
-		}
-	}
-
-	if (!bBinaryFileIsValid || !bBinaryFileIsValidAndGuidMatch)
+	if (!bBinaryFileIsValidAndGuidMatch)
 	{
 		if (OpenWriteHandle())
 		{
-			BinaryFileState = bBinaryFileIsValid && !bBinaryFileIsValidAndGuidMatch ? EBinaryFileState::BuildingCacheFileWithMove : EBinaryFileState::BuildingCacheFile;
-
-			// save header
+			BinaryFileState = EBinaryFileState::BuildingCacheFile;
+			// save header, 0 program count indicates an unfinished file.
+			// The header is overwritten at the end of the process.
+			UE::OpenGL::FBinaryCacheFileHeader OutHeader = UE::OpenGL::FBinaryCacheFileHeader::CreateHeader(CurrentShaderPipelineCacheVersionGuid, 0);
 			FArchive& Ar = *BinaryCacheWriteFileHandle;
-			uint32 Version = GBinaryProgramFileVersion;
-			Ar << Version;
-			FGuid BinaryCacheGuid = ShaderPipelineCacheVersionGuid;
-			Ar << BinaryCacheGuid;
-			bool bWritingCompressedBinaries = (UE::OpenGL::AreBinaryProgramsCompressed());
-			Ar << bWritingCompressedBinaries;
+			Ar << OutHeader;
 		}
 		else
 		{
@@ -516,41 +492,12 @@ void FOpenGLProgramBinaryCache::ScanProgramCacheFile(const FGuid& ShaderPipeline
 	}
 }
 
-// add the GLProgramBinaryFileCacheEntry into the runtime lookup containers.
-void FOpenGLProgramBinaryCache::AddProgramFileEntryToMap(FGLProgramBinaryFileCacheEntry* NewEntry)
-{
-	const FOpenGLProgramKey& ProgramKey = NewEntry->FileInfo.ShaderHasheSet;
-	check(!ProgramToBinaryMap.Contains(ProgramKey));
-	ProgramToBinaryMap.Add(ProgramKey, NewEntry);
-
-	UE_LOG(LogRHI, Verbose, TEXT("AddProgramFileEntryToMap : Adding program: %s"), *ProgramKey.ToString());
-
-	for (int i = 0; i < CrossCompiler::NUM_NON_COMPUTE_SHADER_STAGES; ++i)
-	{
-		const FSHAHash& ShaderHash = ProgramKey.ShaderHashes[i];
-
-		if (ShaderHash != FSHAHash())
-		{
-			if (ShaderToProgramsMap.Contains(ShaderHash))
-			{
-				ShaderToProgramsMap[ShaderHash].Add(NewEntry);
-			}
-			else
-			{
-				ShaderToProgramsMap.Add(ShaderHash, NewEntry);
-			}
-		}
-	}
-}
-
 bool FOpenGLProgramBinaryCache::OpenWriteHandle()
 {
 	check(BinaryCacheWriteFileHandle == nullptr);
-	check(BinaryCacheAsyncReadFileHandle == nullptr);
 
 	// perform file writing to a file temporary filename so we don't attempt to use the file later if the write session is interrupted
-	FString ProgramCacheFilename = GetProgramBinaryCacheFilePath();
-	FString ProgramCacheFilenameWrite = ProgramCacheFilename + TEXT(".write");
+	const FString ProgramCacheFilenameWrite = GetProgramBinaryCacheFilePath() + TEXT(".write");
 
 	BinaryCacheWriteFileHandle = IFileManager::Get().CreateFileWriter(*ProgramCacheFilenameWrite, EFileWrite::FILEWRITE_None);
 
@@ -561,28 +508,28 @@ bool FOpenGLProgramBinaryCache::OpenWriteHandle()
 
 void FOpenGLProgramBinaryCache::CloseWriteHandle()
 {
-	if (BinaryFileState == EBinaryFileState::BuildingCacheFileWithMove)
-	{
-		UE_LOG(LogRHI, Log, TEXT("FOpenGLProgramBinaryCache: Deleting previous binary program cache (%s), reused %d programs from a total of %d."), *PreviousBinaryCacheInfo.OldCacheFilename, PreviousBinaryCacheInfo.NumberOfOldEntriesReused, ProgramToBinaryMap.Num());
-
-		// clean up references to old cache.
-		PreviousBinaryCacheInfo.OldCacheArchive->Close();
-		PreviousBinaryCacheInfo.OldCacheArchive = nullptr;
-		IPlatformFile& PlatformFile = FPlatformFileManager::Get().GetPlatformFile();
-		PlatformFile.DeleteFile(*PreviousBinaryCacheInfo.OldCacheFilename);
-		PreviousBinaryCacheInfo.OldCacheFilename.Empty();
-		PreviousBinaryCacheInfo.ProgramToOldBinaryCacheMap.Empty();
-	}
-
 	check(BinaryCacheWriteFileHandle != nullptr);
 
-	AppendProgramBinaryFileEofEntry(*BinaryCacheWriteFileHandle);
 	bool bArchiveFailed = BinaryCacheWriteFileHandle->IsError() || BinaryCacheWriteFileHandle->IsCriticalError();
 
+	// Overwrite the header with the final program count. This indicates a successful write.
+	if(!bArchiveFailed)
+	{
+		FArchive& Ar = *BinaryCacheWriteFileHandle;
+		Ar.Seek(0);
+		UE::OpenGL::FBinaryCacheFileHeader OutHeader = UE::OpenGL::FBinaryCacheFileHeader::CreateHeader(CurrentShaderPipelineCacheVersionGuid, ProgramsWrittenToOutputCache.Num());
+		Ar << OutHeader;
+		bArchiveFailed = BinaryCacheWriteFileHandle->IsError() || BinaryCacheWriteFileHandle->IsCriticalError();
+	}
+
 	BinaryCacheWriteFileHandle->Close();
 	delete BinaryCacheWriteFileHandle;
 	BinaryCacheWriteFileHandle = nullptr;
 
+	const FString ProgramCacheFilename = GetProgramBinaryCacheFilePath();
+	const FString ProgramCacheFilenameWrite = ProgramCacheFilename + TEXT(".write");
+	IPlatformFile& PlatformFile = FPlatformFileManager::Get().GetPlatformFile();
+
 	if (bArchiveFailed)
 	{
 		RHIGetPanicDelegate().ExecuteIfBound(FName("FailedBinaryProgramArchiveWrite"));
@@ -590,71 +537,24 @@ void FOpenGLProgramBinaryCache::CloseWriteHandle()
 	}
 
 	// rename the temp filename back to the final filename
-	FString ProgramCacheFilename = GetProgramBinaryCacheFilePath();
-	FString ProgramCacheFilenameWrite = ProgramCacheFilename + TEXT(".write");
-	IPlatformFile& PlatformFile = FPlatformFileManager::Get().GetPlatformFile();
 	PlatformFile.DeleteFile(*ProgramCacheFilename); // file should never exist, but for safety
 	PlatformFile.MoveFile(*ProgramCacheFilename, *ProgramCacheFilenameWrite);
 }
 
-void FOpenGLProgramBinaryCache::OpenAsyncReadHandle()
-{
-	check(BinaryCacheAsyncReadFileHandle == nullptr);
-
-	FString ProgramCacheFilename = GetProgramBinaryCacheFilePath();
-	BinaryCacheAsyncReadFileHandle = FPlatformFileManager::Get().GetPlatformFile().OpenAsyncRead(*ProgramCacheFilename);
-	checkf(BinaryCacheAsyncReadFileHandle, TEXT("Could not opan an async file")); // this generally cannot fail because it is async
-}
-
-/* dead code, needs removal
-void FOpenGLProgramBinaryCache::CloseAsyncReadHandle()
-{
-	// wait for any pending reads.
-	{
-		FScopeLock Lock(&GPendingGLProgramCreateRequestsCS);
-		for (FGLProgramBinaryFileCacheEntry* CreateRequest : PendingGLProgramCreateRequests)
-		{
-			TSharedPtr<IAsyncReadRequest, ESPMode::ThreadSafe> ReadRequest = CreateRequest->ReadRequest.Pin();
-			if (ReadRequest.IsValid())
-			{
-				ReadRequest->WaitCompletion();
-				CreateRequest->ReadRequest = nullptr;
-			}
-		}
-	}
-
-	delete BinaryCacheAsyncReadFileHandle;
-	BinaryCacheAsyncReadFileHandle = nullptr;
-}*/
-
 // Called when a new program has been created by OGL RHI, creates the binary cache if it's invalid and then appends the new program details to the file and runtime containers.
-void FOpenGLProgramBinaryCache::AppendGLProgramToBinaryCache(const FOpenGLProgramKey& ProgramKey, GLuint Program, TArray<uint8>& CachedProgramBinaryOUT)
-{
-	if (IsBuildingCache_internal() == false)
-	{
-		return;
-	}
-
-	FScopeLock Lock(&GProgramBinaryFileCacheCS);
-
-	AddUniqueGLProgramToBinaryCache(ProgramKey, Program, CachedProgramBinaryOUT);
-}
-
-// Add the program to the binary cache if it does not already exist.
-void FOpenGLProgramBinaryCache::AddUniqueGLProgramToBinaryCache(const FOpenGLProgramKey& ProgramKey, GLuint Program, TArray<uint8>& CachedProgramBinaryOUT)
+FOpenGLProgramBinary FOpenGLProgramBinaryCache::CreateAndCacheProgramBinary_Internal(const FOpenGLProgramKey& ProgramKey, GLuint Program)
 {
+	FOpenGLProgramBinary CachedProgramBinary;
 	// Add to runtime and disk.
-	const FOpenGLProgramKey& ProgramHash = ProgramKey;
-
-	// Check we dont already have this: Something could be in the cache but still reach this point if OnSharedShaderCodeRequest(s) have not occurred.
-	if (!ProgramToBinaryMap.Contains(ProgramHash))
+	if (!ProgramsWrittenToOutputCache.Contains(ProgramKey))
 	{
 		uint32 ProgramBinaryOffset = 0, ProgramBinarySize = 0;
 
 		FOpenGLProgramKey SerializedProgramKey = ProgramKey;
-		if (ensure(UE::OpenGL::GetProgramBinaryFromGLProgram(Program, CachedProgramBinaryOUT)))
+		CachedProgramBinary = UE::OpenGL::GetProgramBinaryFromGLProgram(Program);
+		if (ensure(CachedProgramBinary.IsValid()))
 		{
-			AddProgramBinaryDataToBinaryCache(CachedProgramBinaryOUT, ProgramKey);
+			AddProgramBinaryDataToBinaryCache(ProgramKey, CachedProgramBinary);
 		}
 		else
 		{
@@ -666,69 +566,49 @@ void FOpenGLProgramBinaryCache::AddUniqueGLProgramToBinaryCache(const FOpenGLPro
 			// Panic!
 		}
 	}
+	return CachedProgramBinary;
 }
 
-void FOpenGLProgramBinaryCache::CacheProgramBinary(const FOpenGLProgramKey& ProgramKey, TArray<uint8>& CachedProgramBinary)
+void FOpenGLProgramBinaryCache::CacheProgramBinary(const FOpenGLProgramKey& ProgramKey, const FOpenGLProgramBinary& BinaryProgramData)
 {
 	if (CachePtr)
 	{
-		if (CachePtr->IsBuildingCache_internal() == false)
-		{
-			return;
-		}
-
 		FScopeLock Lock(&GProgramBinaryFileCacheCS);
 
-		if (!CachePtr->ProgramToBinaryMap.Contains(ProgramKey))
+		if (!CachePtr->ProgramsWrittenToOutputCache.Contains(ProgramKey))
 		{
-			CachePtr->AddProgramBinaryDataToBinaryCache(CachedProgramBinary, ProgramKey);
+			CachePtr->AddProgramBinaryDataToBinaryCache(ProgramKey, BinaryProgramData);
 		}
 	}
 }
 
 // Serialize out the program binary data and add to runtime structures.
-void FOpenGLProgramBinaryCache::AddProgramBinaryDataToBinaryCache(TArray<uint8>& BinaryProgramData, const FOpenGLProgramKey& ProgramKey)
+void FOpenGLProgramBinaryCache::AddProgramBinaryDataToBinaryCache(const FOpenGLProgramKey& ProgramKey, const FOpenGLProgramBinary& BinaryProgramData)
 {
-	FScopeLock Lock(&GProgramBinaryFileCacheCS);
+	check(IsBuildingCache_internal());
+
 	FArchive& Ar = *BinaryCacheWriteFileHandle;
 	// Serialize to output file:
 	FOpenGLProgramKey SerializedProgramKey = ProgramKey;
-	uint32 ProgramBinarySize = (uint32)BinaryProgramData.Num();
+	const TArrayView<const uint8> BinaryProgramDataView = BinaryProgramData.GetDataView();
+	uint32 ProgramBinarySize = (uint32)BinaryProgramDataView.Num();
+	const uint8* ProgramBinaryBytes = BinaryProgramDataView.GetData();
 	Ar << SerializedProgramKey;
 	uint32 ProgramBinaryOffset = Ar.Tell();
 	Ar << ProgramBinarySize;
-	Ar.Serialize(BinaryProgramData.GetData(), ProgramBinarySize);
+	Ar.Serialize(const_cast<uint8*>(ProgramBinaryBytes), ProgramBinarySize);
 	if (UE::OpenGL::AreBinaryProgramsCompressed())
 	{
 		static uint32 TotalUncompressed = 0;
 		static uint32 TotalCompressed = 0;
 
-		UE::OpenGL::FCompressedProgramBinaryHeader* Header = (UE::OpenGL::FCompressedProgramBinaryHeader*)BinaryProgramData.GetData();
+		const UE::OpenGL::FCompressedProgramBinaryHeader* Header = (UE::OpenGL::FCompressedProgramBinaryHeader*)ProgramBinaryBytes;
 		TotalUncompressed += Header->UncompressedSize;
-		TotalCompressed += BinaryProgramData.Num();
+		TotalCompressed += ProgramBinarySize;
 
 		UE_LOG(LogRHI, Verbose, TEXT("AppendProgramBinaryFile: total Uncompressed: %d, total Compressed %d, Total saved so far: %d"), TotalUncompressed, TotalCompressed, TotalUncompressed - TotalCompressed);
 	}
-
-	FGLProgramBinaryFileCacheEntry* NewIndexEntry = new FGLProgramBinaryFileCacheEntry();
-	ProgramEntryContainer.Emplace(TUniquePtr<FGLProgramBinaryFileCacheEntry>(NewIndexEntry));
-
-	// Store the program file descriptor in the runtime program/shader container:
-	NewIndexEntry->GLProgramState = FGLProgramBinaryFileCacheEntry::EGLProgramState::ProgramStored;
-	NewIndexEntry->FileInfo.ProgramOffset = ProgramBinaryOffset;
-	NewIndexEntry->FileInfo.ProgramSize = ProgramBinarySize;
-	NewIndexEntry->ProgramIndex = ProgramToBinaryMap.Num();
-	NewIndexEntry->FileInfo.ShaderHasheSet = ProgramKey;
-	AddProgramFileEntryToMap(NewIndexEntry);
-}
-
-void FOpenGLProgramBinaryCache::AppendProgramBinaryFileEofEntry(FArchive& Ar)
-{
-	// write out an all zero record that signifies eof.
-	FOpenGLProgramKey SerializedProgramKey;
-	Ar << SerializedProgramKey;
-	uint32 ProgramBinarySize = 0;
-	Ar << ProgramBinarySize;
+	ProgramsWrittenToOutputCache.Add(ProgramKey);
 }
 
 void FOpenGLProgramBinaryCache::Shutdown()
@@ -740,121 +620,35 @@ void FOpenGLProgramBinaryCache::Shutdown()
 	}
 }
 
-void FOpenGLProgramBinaryCache::CacheProgram(GLuint Program, const FOpenGLProgramKey& ProgramKey, TArray<uint8>& CachedProgramBinaryOUT)
-{
-	if (CachePtr)
-	{
-		CachePtr->AppendGLProgramToBinaryCache(ProgramKey, Program, CachedProgramBinaryOUT);
-	}
-}
-
-bool FOpenGLProgramBinaryCache::UseCachedProgram(GLuint& ProgramOUT, const FOpenGLProgramKey& ProgramKey, TArray<uint8>& CachedProgramBinaryOUT)
+bool FOpenGLProgramBinaryCache::DoesOutputCacheContain(const FOpenGLProgramKey& ProgramKey)
 {
 	if (CachePtr)
 	{
-		return CachePtr->UseCachedProgram_internal(ProgramOUT, ProgramKey, CachedProgramBinaryOUT);
+		FScopeLock Lock(&GProgramBinaryFileCacheCS);
+		return CachePtr->DoesOutputCacheContain_Internal(ProgramKey);
 	}
 	return false;
 }
 
-
-namespace UE
+bool FOpenGLProgramBinaryCache::DoesOutputCacheContain_Internal(const FOpenGLProgramKey& ProgramKey) const
 {
-	namespace OpenGL
-	{
-		extern bool CreateGLProgramFromBinary(GLuint& GLProgramIDOUT, const TArray<uint8>& ProgramBinaryData);
-		extern void OnGLProgramLoadedFromBinaryCache(const FOpenGLProgramKey& ProgramKey, TArray<uint8>&& ProgramBinaryData);
-	}
+	return ProgramsWrittenToOutputCache.Contains(ProgramKey);
 }
 
-bool FOpenGLProgramBinaryCache::UseCachedProgram_internal(GLuint& ProgramOUT, const FOpenGLProgramKey& ProgramKey, TArray<uint8>& CachedProgramBinaryOUT)
+FOpenGLProgramBinary FOpenGLProgramBinaryCache::CreateAndCacheProgramBinary(const FOpenGLProgramKey& ProgramKey, GLuint Program )
 {
-	SCOPE_CYCLE_COUNTER(STAT_OpenGLUseCachedProgramTime);
-
-	FGLProgramBinaryFileCacheEntry** ProgramBinRefPtr = nullptr;
-
-	FScopeLock Lock(&GProgramBinaryFileCacheCS);
-
-	ProgramBinRefPtr = ProgramToBinaryMap.Find(ProgramKey);
-
-	if (ProgramBinRefPtr)
-	{
-		FGLProgramBinaryFileCacheEntry* FoundProgram = *ProgramBinRefPtr;
-		check(FoundProgram->FileInfo.ShaderHasheSet == ProgramKey);
-
-		TSharedPtr<IAsyncReadRequest, ESPMode::ThreadSafe> LocalReadRequest = FoundProgram->ReadRequest.Pin();
-		bool bHasReadRequest = LocalReadRequest.IsValid();
-		check(!bHasReadRequest);
-
-		// by this point the program must be either available or no attempt to load from shader library has occurred.
-		checkf(FoundProgram->GLProgramState == FGLProgramBinaryFileCacheEntry::EGLProgramState::ProgramStored
-			|| FoundProgram->GLProgramState == FGLProgramBinaryFileCacheEntry::EGLProgramState::ProgramAvailable,
-			TEXT("Unexpected program state:  (%s) == %d"), *ProgramKey.ToString(), (int32)FoundProgram->GLProgramState);
-
-		if (FoundProgram->GLProgramState == FGLProgramBinaryFileCacheEntry::EGLProgramState::ProgramAvailable)
-		{
-			UE_LOG(LogRHI, Log, TEXT("UseCachedProgram : Program (%s) GLid = %x is ready!"), *ProgramKey.ToString(), FoundProgram->GLProgramId);
-			ProgramOUT = FoundProgram->GLProgramId;
-
-			// GLProgram has been handed over.
-			FoundProgram->GLProgramId = 0;
-			FoundProgram->GLProgramState = FGLProgramBinaryFileCacheEntry::EGLProgramState::ProgramComplete;
-			return true;
-		}
-		else
-		{
-			UE_LOG(LogRHI, Log, TEXT("UseCachedProgram : %s was not ready when needed!! (state %d)"), *ProgramKey.ToString(), (uint32)FoundProgram->GLProgramState);
-		}
-	}
-	else if (BinaryFileState == EBinaryFileState::BuildingCacheFileWithMove)
+	FOpenGLProgramBinary ProgramBinary;
+	if (CachePtr)
 	{
-		// We're building the new cache using the original cache to warm:
-		TUniquePtr<FGLProgramBinaryFileCacheEntry>* FoundExistingBinary = PreviousBinaryCacheInfo.ProgramToOldBinaryCacheMap.Find(ProgramKey);
-		if (FoundExistingBinary)
-		{
-			TUniquePtr<FGLProgramBinaryFileCacheEntry>& ExistingBinary = *FoundExistingBinary;
-			// read old binary:
-			CachedProgramBinaryOUT.SetNumUninitialized(ExistingBinary->FileInfo.ProgramSize);
-			PreviousBinaryCacheInfo.OldCacheArchive->Seek(ExistingBinary->FileInfo.ProgramOffset);
-			PreviousBinaryCacheInfo.OldCacheArchive->Serialize(CachedProgramBinaryOUT.GetData(), ExistingBinary->FileInfo.ProgramSize);
-			bool bSuccess = UE::OpenGL::CreateGLProgramFromBinary(ProgramOUT, CachedProgramBinaryOUT);
-			if (!bSuccess)
-			{
-				UE_LOG(LogRHI, Log, TEXT("[%s, %d, %d]"), *ProgramKey.ToString(), ProgramOUT, CachedProgramBinaryOUT.Num());
-				RHIGetPanicDelegate().ExecuteIfBound(FName("FailedBinaryProgramCreateFromOldCache"));
-				UE_LOG(LogRHI, Fatal, TEXT("UseCachedProgram : Failed to create GL program from binary data while BuildingCacheFileWithMove! [%s]"), *ProgramKey.ToString());
-			}
-			// Now write to new cache, we're returning true here so no attempt will be made to add it back to the cache later.
-			AddProgramBinaryDataToBinaryCache(CachedProgramBinaryOUT, ProgramKey);
-
-			PreviousBinaryCacheInfo.NumberOfOldEntriesReused++;
-			return true;
-		}
+		FScopeLock Lock(&GProgramBinaryFileCacheCS);
+		ProgramBinary = CachePtr->CreateAndCacheProgramBinary_Internal(ProgramKey, Program);
 	}
-	return false;
+	return ProgramBinary;
 }
 
-// void FOpenGLProgramBinaryCache::CompilePendingShaders(const FOpenGLLinkedProgramConfiguration& Config)
-// {
-// 	VERIFY_GL_SCOPE();
-// 
-// 	// Find the existing compiled shader in the cache.
-// 	FScopeLock Lock(&GCompiledShaderCacheCS);
-// 	for (int32 StageIdx = 0; StageIdx < UE_ARRAY_COUNT(Config.Shaders); ++StageIdx)
-// 	{
-// 		TSharedPtr<FOpenGLCompiledShaderValue> FoundShader = GetOpenGLCompiledShaderCache().FindRef(Config.Shaders[StageIdx].ShaderKey);
-// 		if (FoundShader && FoundShader->bHasCompiled == false)
-// 		{
-// 			TArray<ANSICHAR> GlslCode = FoundShader->GetUncompressedShader();
-// 			CompileCurrentShader(Config.Shaders[StageIdx].Resource, GlslCode);
-// 			FoundShader->bHasCompiled = true;
-// 		}
-// 	}
-// }
-
 FString FOpenGLProgramBinaryCache::GetProgramBinaryCacheFilePath() const
 {
-	FString ProgramFilename = CachePath + TEXT("/") + CacheFilename;
+	FString ProgramFilename = CachePath + TEXT("/") + CacheFilename + TEXT("_")+CurrentShaderPipelineCacheVersionGuid.ToString();
 	return ProgramFilename;
 }
 
@@ -863,48 +657,49 @@ void FOpenGLProgramBinaryCache::CheckPendingGLProgramCreateRequests()
 	if (CachePtr)
 	{
 		QUICK_SCOPE_CYCLE_COUNTER(STAT_OpenGLShaderCreateShaderLibRequests);
+		check(IsInRenderingThread() || IsInRHIThread());
+		FScopeLock Lock(&GPendingGLProgramCreateRequestsCS);
 		CachePtr->CheckPendingGLProgramCreateRequests_internal();
 	}
 }
 
+// Move programs encountered during the scan to the GL RHI program container.
+// GMaxBinaryProgramLoadTimeMS attempts to reduce hitching, if we're not using the LRU then we still create GL programs and require more time.
 void FOpenGLProgramBinaryCache::CheckPendingGLProgramCreateRequests_internal()
 {
-	check(IsInRenderingThread() || IsInRHIThread());
-	FScopeLock Lock(&GPendingGLProgramCreateRequestsCS);
-	//UE_LOG(LogRHI, Log, TEXT("CheckPendingGLProgramCreateRequests : PendingGLProgramCreateRequests = %d"), PendingGLProgramCreateRequests.Num());
-
-	if (PendingGLProgramCreateRequests.Num() > 0)
+	if (ScannedBinaryPrograms.Num() > 0)
 	{
-		float TimeRemainingS = (float)GMaxShaderLibProcessingTimeMS / 1000.0f;
+		//FScopedDurationTimeLogger Timer(TEXT("CheckPendingGLProgramCreateRequests"));
+		float TimeRemainingS = (float)GMaxBinaryProgramLoadTimeMS / 1000.0f;
 		double StartTime = FPlatformTime::Seconds();
 		int32 Count = 0;
-		while (PendingGLProgramCreateRequests.Num() && TimeRemainingS > 0.0f)
+
+		for (auto It = ScannedBinaryPrograms.CreateIterator(); It && TimeRemainingS > 0.0f; ++It)
 		{
-			CompleteLoadedGLProgramRequest_internal(PendingGLProgramCreateRequests.Pop());
+			UE::OpenGL::OnGLProgramLoadedFromBinaryCache(It->Key, MoveTemp(It->Value));
 			TimeRemainingS -= (float)(FPlatformTime::Seconds() - StartTime);
 			StartTime = FPlatformTime::Seconds();
 			Count++;
+			It.RemoveCurrent();
+		}
+		float TimeTaken = (float)GMaxBinaryProgramLoadTimeMS - (TimeRemainingS * 1000.0f);
+		if (TimeRemainingS < 0.0f)
+		{
+			UE_LOG(LogRHI, Warning, TEXT("CheckPendingGLProgramCreateRequests : iter count = %d, time taken = %f ms (remaining %d)"), Count, TimeTaken, ScannedBinaryPrograms.Num());
+		}
+		else
+		{
+			UE_LOG(LogRHI, Verbose, TEXT("CheckPendingGLProgramCreateRequests : iter count = %d, time taken = %f ms (remaining %d)"), Count, TimeTaken, ScannedBinaryPrograms.Num());
 		}
-		float TimeTaken = (float)GMaxShaderLibProcessingTimeMS - (TimeRemainingS * 1000.0f);
-		UE_CLOG(TimeTaken > 2.0f, LogRHI, Log, TEXT("CheckPendingGLProgramCreateRequests : iter count = %d, time taken = %f ms (remaining %d)"), Count, TimeTaken, PendingGLProgramCreateRequests.Num());
 	}
 }
 
-
-
-void FOpenGLProgramBinaryCache::CompleteLoadedGLProgramRequest_internal(FGLProgramBinaryFileCacheEntry* PendingGLCreate)
-{
-	check(PendingGLCreate->GLProgramState == FGLProgramBinaryFileCacheEntry::EGLProgramState::ProgramLoaded);
-	FOpenGLProgramKey& ProgramKey = PendingGLCreate->FileInfo.ShaderHasheSet;
-	UE::OpenGL::OnGLProgramLoadedFromBinaryCache(ProgramKey, MoveTemp(PendingGLCreate->ProgramBinaryData));
-	// Ownership transfered to OpenGLProgramsCache.
-	PendingGLCreate->GLProgramState = FGLProgramBinaryFileCacheEntry::EGLProgramState::ProgramComplete;
-}
-
 bool FOpenGLProgramBinaryCache::CheckSinglePendingGLProgramCreateRequest(const FOpenGLProgramKey& ProgramKey)
 {
 	if (CachePtr)
 	{
+		check(IsInRenderingThread() || IsInRHIThread());
+		FScopeLock Lock(&GPendingGLProgramCreateRequestsCS);
 		return CachePtr->CheckSinglePendingGLProgramCreateRequest_internal(ProgramKey);
 	}
 	return false;
@@ -913,139 +708,11 @@ bool FOpenGLProgramBinaryCache::CheckSinglePendingGLProgramCreateRequest(const F
 // Any pending program must complete in this case.
 bool FOpenGLProgramBinaryCache::CheckSinglePendingGLProgramCreateRequest_internal(const FOpenGLProgramKey& ProgramKey)
 {
-	FGLProgramBinaryFileCacheEntry** ProgramBinRefPtr = nullptr;
-	FScopeLock ProgramBinaryCacheLock(&GProgramBinaryFileCacheCS);
-	ProgramBinRefPtr = CachePtr->ProgramToBinaryMap.Find(ProgramKey);
-	if (ProgramBinRefPtr)
-	{
-		FGLProgramBinaryFileCacheEntry* ProgramEntry = *ProgramBinRefPtr;
-		TSharedPtr<IAsyncReadRequest, ESPMode::ThreadSafe> LocalReadRequest = ProgramEntry->ReadRequest.Pin();
-		if (LocalReadRequest.IsValid())
-		{
-			ensure(ProgramEntry->GLProgramState == FGLProgramBinaryFileCacheEntry::EGLProgramState::ProgramLoading);
-			LocalReadRequest->WaitCompletion();
-			ProgramEntry->ReadRequest = nullptr;
-			ProgramEntry->GLProgramState = FGLProgramBinaryFileCacheEntry::EGLProgramState::ProgramLoaded;
-			CompleteLoadedGLProgramRequest_internal(ProgramEntry);
-		}
-		else
-		{
-			FScopeLock Lock(&GPendingGLProgramCreateRequestsCS);
-			if (ProgramEntry->GLProgramState == FGLProgramBinaryFileCacheEntry::EGLProgramState::ProgramLoaded)
-			{
-				int32 PendingRequestIndex = -1;
-				if (ensure(PendingGLProgramCreateRequests.Find(ProgramEntry, PendingRequestIndex)))
-				{
-					CompleteLoadedGLProgramRequest_internal(ProgramEntry);
-					PendingGLProgramCreateRequests.RemoveAtSwap(PendingRequestIndex);
-				}
-			}
-		}
+	TUniqueObj<FOpenGLProgramBinary> ProgFound;
+	if (ScannedBinaryPrograms.RemoveAndCopyValue(ProgramKey, ProgFound))
+	{		
+		UE::OpenGL::OnGLProgramLoadedFromBinaryCache(ProgramKey, MoveTemp(ProgFound));
 		return true;
 	}
 	return false;
 }
-
-bool OnExternalReadCallback(const TSharedPtr<IAsyncReadRequest, ESPMode::ThreadSafe>& AsyncReadRequest, FGLProgramBinaryFileCacheEntry* ProgramBinEntry, TArray<FGLProgramBinaryFileCacheEntry*>& PendingGLProgramCreateRequests, double RemainingTime)
-{
-	if (!AsyncReadRequest->WaitCompletion(RemainingTime))
-	{
-		return false;
-	}
-
-	FScopeLock ProgramBinaryCacheLock(&GProgramBinaryFileCacheCS);
-
-	if (ProgramBinEntry->GLProgramState == FGLProgramBinaryFileCacheEntry::EGLProgramState::ProgramLoading)
-	{
-		// Async load complete.
-		ProgramBinEntry->GLProgramState = FGLProgramBinaryFileCacheEntry::EGLProgramState::ProgramLoaded;
-		FOpenGLProgramKey& ProgramKey = ProgramBinEntry->FileInfo.ShaderHasheSet;
-
-		{
-			// Add this program to the create queue.
-			FScopeLock Lock(&GPendingGLProgramCreateRequestsCS);
-			PendingGLProgramCreateRequests.Add(ProgramBinEntry);
-		}
-	}
-
-	return true;
-}
-
-void FOpenGLProgramBinaryCache::BeginProgramReadRequest(FGLProgramBinaryFileCacheEntry* ProgramBinEntry, FArchive* Ar)
-{
-	check(ProgramBinEntry);
-
-	TSharedPtr<IAsyncReadRequest, ESPMode::ThreadSafe> LocalReadRequest = ProgramBinEntry->ReadRequest.Pin();
-	bool bHasReadRequest = LocalReadRequest.IsValid();
-
-	if (ensure(!bHasReadRequest))
-	{
-		check(ProgramBinEntry->ProgramBinaryData.Num() == 0);
-		check(ProgramBinEntry->GLProgramState == FGLProgramBinaryFileCacheEntry::EGLProgramState::ProgramStored);
-
-		int64 ReadSize = ProgramBinEntry->FileInfo.ProgramSize;
-		int64 ReadOffset = ProgramBinEntry->FileInfo.ProgramOffset;
-
-		if (ensure(ReadSize > 0))
-		{
-			ProgramBinEntry->ProgramBinaryData.SetNumUninitialized(ReadSize);
-			ProgramBinEntry->GLProgramState = FGLProgramBinaryFileCacheEntry::EGLProgramState::ProgramLoading;
-			LocalReadRequest = MakeShareable(BinaryCacheAsyncReadFileHandle->ReadRequest(ReadOffset, ReadSize, AIOP_Normal, nullptr, ProgramBinEntry->ProgramBinaryData.GetData()));
-			ProgramBinEntry->ReadRequest = LocalReadRequest;
-			bHasReadRequest = true;
-
-			FExternalReadCallback ExternalReadCallback = [ProgramBinEntry, LocalReadRequest, this](double ReaminingTime)
-			{
-				return OnExternalReadCallback(LocalReadRequest, ProgramBinEntry, PendingGLProgramCreateRequests, ReaminingTime);
-			};
-
-			if (!Ar || !Ar->AttachExternalReadDependency(ExternalReadCallback))
-			{
-				// Archive does not support async loading
-				// do a blocking load
-				ExternalReadCallback(0.0);
-			}
-		}
-	}
-}
-
-void FOpenGLProgramBinaryCache::OnShaderLibraryRequestShaderCode(const FSHAHash& Hash, FArchive* Ar)
-{
-	if (CachePtr)
-	{
-		CachePtr->OnShaderLibraryRequestShaderCode_internal(Hash, Ar);
-	}
-}
-
-void FOpenGLProgramBinaryCache::OnShaderLibraryRequestShaderCode_internal(const FSHAHash& Hash, FArchive* Ar)
-{
-	FScopeLock Lock(&GProgramBinaryFileCacheCS);
-	FGLShaderToPrograms& FoundShaderToBinary = ShaderToProgramsMap.FindOrAdd(Hash);
-	if (!FoundShaderToBinary.bLoaded)
-	{
-		FoundShaderToBinary.bLoaded = true;
-
-		// if the binary cache is valid, look to see if we now have any complete programs to stream in.
-		// otherwise, we'll do this check bLoaded shaders when the binary cache loads.
-		if (BinaryFileState == EBinaryFileState::ValidCacheFile)
-		{
-			for (struct FGLProgramBinaryFileCacheEntry* ProgramBinEntry : FoundShaderToBinary.AssociatedPrograms)
-			{
-				const FOpenGLProgramKey& ProgramKey = ProgramBinEntry->FileInfo.ShaderHasheSet;
-				if (ProgramBinEntry->GLProgramState == FGLProgramBinaryFileCacheEntry::EGLProgramState::ProgramStored)
-				{
-					bool bAllShadersLoaded = true;
-					for (int32 i = 0; i < CrossCompiler::NUM_NON_COMPUTE_SHADER_STAGES && bAllShadersLoaded; i++)
-					{
-						bAllShadersLoaded = ProgramKey.ShaderHashes[i] == FSHAHash() || ShaderIsLoaded(ProgramKey.ShaderHashes[i]);
-					}
-
-					if (bAllShadersLoaded)
-					{
-						FOpenGLProgramBinaryCache::BeginProgramReadRequest(ProgramBinEntry, Ar);
-					}
-				}
-			}
-		}
-	}
-}
diff --git a/Engine/Source/Runtime/OpenGLDrv/Private/OpenGLProgramBinaryFileCache.h b/Engine/Source/Runtime/OpenGLDrv/Private/OpenGLProgramBinaryFileCache.h
--- a/Engine/Source/Runtime/OpenGLDrv/Private/OpenGLProgramBinaryFileCache.h
+++ b/Engine/Source/Runtime/OpenGLDrv/Private/OpenGLProgramBinaryFileCache.h
@@ -22,9 +22,10 @@
 #include "ShaderPipelineCache.h"
 
 class FOpenGLLinkedProgram;
-
-
+class FOpenGLProgramBinary;
+class FOpenGLProgramBinaryMapping;
 struct FGLProgramBinaryFileCacheEntry;
+
 class FOpenGLProgramBinaryCache
 {
 public:
@@ -33,21 +34,15 @@ public:
 
 	static bool IsEnabled();
 
-	/** Compile required shaders for a program, only in case binary program was not found in the cache   */
-// 	static void CompilePendingShaders(const FOpenGLLinkedProgramConfiguration& Config);
+	/** Extract the program binary from GL and store on disk. */
+	static FOpenGLProgramBinary CreateAndCacheProgramBinary(const FOpenGLProgramKey& ProgramKey, GLuint Program);
 
-	/** Try to find and load program binary from cache */
-	static bool UseCachedProgram(GLuint& ProgramOUT, const FOpenGLProgramKey& ProgramKey, TArray<uint8>& CachedProgramBinaryOUT);
-
-	/** Extract the program binary from GL and store on disk. Only when ProgramBinaryCache is enabled */
-	static void CacheProgram(GLuint Program, const FOpenGLProgramKey& ProgramKey, TArray<uint8>& CachedProgramBinaryOUT);
+	/** Is the program already in the currently building cache */
+	static bool DoesOutputCacheContain(const FOpenGLProgramKey& ProgramKey);
 
 	/** Take an existing program binary and store on disk. Only when ProgramBinaryCache is enabled
-	* CachedProgramBinary should be const, cant be due to serialize api
 	*/
-	static void CacheProgramBinary(const FOpenGLProgramKey& ProgramKey, /*const*/ TArray<uint8>& CachedProgramBinary);
-
-	static void OnShaderLibraryRequestShaderCode(const FSHAHash& Hash, FArchive* Ar);
+	static void CacheProgramBinary(const FOpenGLProgramKey& ProgramKey, const FOpenGLProgramBinary& CachedProgramBinary);
 
 	/** Create any pending GL programs that have come from shader library requests */
 	static void CheckPendingGLProgramCreateRequests();
@@ -81,6 +76,8 @@ private:
 	/* Just the cache filename, without the path */
 	FString CacheFilename;
 
+	// Guid of the PSO cache being procesed.
+	FGuid CurrentShaderPipelineCacheVersionGuid;
 
 	/**
 	* Shaders that were requested for compilation
@@ -91,10 +88,9 @@ private:
 	bool AppendProgramBinaryFile(FArchive& Ar, const FOpenGLProgramKey& ProgramKey, GLuint Program, uint32& ProgramBinaryOffsetOUT, uint32& ProgramBinarySizeOUT);
 	void AppendProgramBinaryFileEofEntry(FArchive& Ar);
 
-	void ScanProgramCacheFile(const FGuid& ShaderPipelineCacheVersionGuid = FGuid());
+	void Reset();
 
-	/* Add a program */
-	void AddProgramFileEntryToMap(FGLProgramBinaryFileCacheEntry* IndexEntry);
+	void ScanProgramCacheFile();
 
 	void OpenAsyncReadHandle();
 	void CloseAsyncReadHandle();
@@ -102,30 +98,16 @@ private:
 	bool OpenWriteHandle();
 	void CloseWriteHandle();
 
-	void AppendGLProgramToBinaryCache(const FOpenGLProgramKey& ProgramKey, GLuint Program, TArray<uint8>& CachedProgramBinaryOUT);
-	void AddUniqueGLProgramToBinaryCache(const FOpenGLProgramKey& ProgramKey, GLuint Program, TArray<uint8>& CachedProgramBinaryOUT);
+	FOpenGLProgramBinary CreateAndCacheProgramBinary_Internal(const FOpenGLProgramKey& ProgramKey, GLuint Program);
+	bool DoesOutputCacheContain_Internal(const FOpenGLProgramKey& ProgramKey) const;
 
-	void AddProgramBinaryDataToBinaryCache(TArray<uint8>& BinaryProgramData, const FOpenGLProgramKey& ProgramKey);
+	void AddProgramBinaryDataToBinaryCache(const FOpenGLProgramKey& ProgramKey, const FOpenGLProgramBinary& BinaryProgramData);
 
 	void ReleaseGLProgram_internal(FOpenGLLinkedProgramConfiguration& Config, GLuint Program);
 
-	FORCEINLINE_DEBUGGABLE bool ShaderIsLoaded(const FSHAHash& Hash)
-	{
-		const FGLShaderToPrograms* FoundShaderToBinary = CachePtr->ShaderToProgramsMap.Find(Hash);
-		return FoundShaderToBinary && FoundShaderToBinary->bLoaded;
-	}
-
-	bool UseCachedProgram_internal(GLuint& ProgramOUT, const FOpenGLProgramKey& ProgramKey, TArray<uint8>& CachedProgramBinaryOUT);
-
-	void OnShaderLibraryRequestShaderCode_internal(const FSHAHash& Hash, FArchive* Ar);
-
-	void BeginProgramReadRequest(FGLProgramBinaryFileCacheEntry* IndexEntry, FArchive* Ar);
-
 	void CheckPendingGLProgramCreateRequests_internal();
 	bool CheckSinglePendingGLProgramCreateRequest_internal(const FOpenGLProgramKey& ProgramKey);
 
-	void CompleteLoadedGLProgramRequest_internal(FGLProgramBinaryFileCacheEntry* PendingGLCreate);
-
 	/** Delegate handlers to track the ShaderPipelineCache precompile. */
 	void OnShaderPipelineCacheOpened(FString const& Name, EShaderPlatform Platform, uint32 Count, const FGuid& VersionGuid, FShaderPipelineCache::FShaderCachePrecompileContext& ShaderCachePrecompileContext);
 	void OnShaderPipelineCachePrecompilationComplete(uint32 Count, double Seconds, const FShaderPipelineCache::FShaderCachePrecompileContext& ShaderCachePrecompileContext);
@@ -135,38 +117,12 @@ private:
 	FDelegateHandle OnShaderPipelineCacheOpenedDelegate;
 	FDelegateHandle OnShaderPipelineCachePrecompilationCompleteDelegate;
 
-	TArray<TUniquePtr<FGLProgramBinaryFileCacheEntry>> ProgramEntryContainer; // this is the owner of all FGLProgramBinaryFileCacheEntry ptrs
-
-	TMap<FOpenGLProgramKey, FGLProgramBinaryFileCacheEntry*> ProgramToBinaryMap; // program key to program entry.
+	TSet<FOpenGLProgramKey> ProgramsWrittenToOutputCache; // Only populated when creating the cache. Used to prevents duplicates.
 
-	struct FGLShaderToPrograms
-	{
-		FGLShaderToPrograms() : bLoaded(false)
-		{
-		}
-
-		FGLShaderToPrograms(FGLProgramBinaryFileCacheEntry* ProgramEntry) : bLoaded(false)
-		{
-			AssociatedPrograms.Add(ProgramEntry);
-		}
-
-		void Add(FGLProgramBinaryFileCacheEntry* ProgramEntry)
-		{
-			checkSlow(!AssociatedPrograms.Contains(ProgramEntry));
-			AssociatedPrograms.Add(ProgramEntry);
-		}
-
-		bool bLoaded;
-		TArray<FGLProgramBinaryFileCacheEntry*> AssociatedPrograms;
-	};
-
-	// Map of shader hash to a list of programs which reference it.
-	TMap<FSHAHash, FGLShaderToPrograms> ShaderToProgramsMap;
+	// All of the mmapped binary programs encountered during the cache open scan.
+	// The contents are moved to OpenGL's program cache during RHI's end of frame tick.
+	TMap<FOpenGLProgramKey, TUniqueObj<FOpenGLProgramBinary>> ScannedBinaryPrograms;
 
-	// programs loaded via async and now ready for creation on GL context owning thread.
-	TArray<FGLProgramBinaryFileCacheEntry*> PendingGLProgramCreateRequests;
-
-	IAsyncReadFileHandle* BinaryCacheAsyncReadFileHandle;
 	FArchive* BinaryCacheWriteFileHandle;
 	bool bShownLoadingScreen;
 
@@ -174,31 +130,13 @@ private:
 	{
 		Uninitialized,					// No binary file is yet established and we should not read or write to it.
 		BuildingCacheFile,				// We are precompiling shaders from the PSO and storing them in a new binary cache. Do not attempt to read.
-		BuildingCacheFileWithMove,		// We are precompiling shaders from the PSO and storing them in a new binary cache, shaders matching from the existing cache are moved to the new file. Do not attempt to read.
-		ValidCacheFile,					// We have a valid cache file we can use for reading. Do not attempt to write.
+		ValidCacheFile,					// We have a valid cache file we can use for reading.
 	};
 
-	bool IsBuildingCache_internal() const
-	{
-		return BinaryFileState == EBinaryFileState::BuildingCacheFile || BinaryFileState == EBinaryFileState::BuildingCacheFileWithMove;
-	}
-
-	struct FPreviousGLProgramBinaryCacheInfo
-	{
-		FPreviousGLProgramBinaryCacheInfo();
-		FPreviousGLProgramBinaryCacheInfo(FPreviousGLProgramBinaryCacheInfo&&);
-		FPreviousGLProgramBinaryCacheInfo& operator=(FPreviousGLProgramBinaryCacheInfo&&);
-		~FPreviousGLProgramBinaryCacheInfo();
-
-
-		FString OldCacheFilename;
-		TUniquePtr<FArchive> OldCacheArchive;
-		TMap<FOpenGLProgramKey, TUniquePtr<FGLProgramBinaryFileCacheEntry> > ProgramToOldBinaryCacheMap; // program key to program entry for old cache used when generating new cache.
-
-		// for logging
-		uint32 NumberOfOldEntriesReused;
-	};
-	FPreviousGLProgramBinaryCacheInfo PreviousBinaryCacheInfo;
+	bool IsBuildingCache_internal() const	{ return BinaryFileState == EBinaryFileState::BuildingCacheFile; }
 
 	EBinaryFileState BinaryFileState;
+
+	// Container of all currently open program cache files.
+	TMap<FGuid, TRefCountPtr<FOpenGLProgramBinaryMapping>> MappedCacheFiles;
 };
diff --git a/Engine/Source/Runtime/OpenGLDrv/Private/OpenGLShaders.cpp b/Engine/Source/Runtime/OpenGLDrv/Private/OpenGLShaders.cpp
--- a/Engine/Source/Runtime/OpenGLDrv/Private/OpenGLShaders.cpp
+++ b/Engine/Source/Runtime/OpenGLDrv/Private/OpenGLShaders.cpp
@@ -17,6 +17,8 @@
 #include "PsoLruCache.h"
 #include "RHICoreShader.h"
 #include "OpenGLProgramBinaryFileCache.h"
+#include "OpenGLBinaryProgramUtils.h"
+#include "HAL/PlatformFramePacer.h"
 
 #define CHECK_FOR_GL_SHADERS_TO_REPLACE 0
 
@@ -77,6 +79,29 @@ static TAutoConsoleVariable<int32> CVarIgnoreShaderCompileFailure(
 	TEXT("1: Ignore Shader compile failures."),
 	ECVF_RenderThreadSafe);
 
+int32 GUntouchedProgramEvictTimeSeconds = 0;
+static FAutoConsoleVariableRef CVarEvictUntouchedProgramSeconds(
+	TEXT("r.OpenGL.ProgramLRUEvictTimeSeconds"),
+	GUntouchedProgramEvictTimeSeconds,
+	TEXT("OpenGL Program LRU, unused program eviction time.\n")
+	TEXT("Time in seconds before an unused program is eligible for eviction from the GL driver.\n")
+	TEXT("0 to disable unused program eviction. (default)")
+	,
+	ECVF_RenderThreadSafe
+);
+
+int32 GProgramLRUResidentCountBeforeEviction = 300;
+static FAutoConsoleVariableRef CVarProgramLRUResidentCountBeforeEviction(
+	TEXT("r.OpenGL.ProgramLRUResidentCountBeforeEviction"),
+	GProgramLRUResidentCountBeforeEviction,
+	TEXT("The number of permanently resident programs allowed before the LRU begins evicting unused programs.\n")
+	TEXT("Can help reduce hitching after a static scene is presented for an extended period of time.\n")
+	TEXT("0 to remove all eligible programs as soon as possible. (max memory savings with higher hitching potential.)\n")
+	TEXT("300: (default) ensures the last 300 programs are retained as GL objects.")
+	,
+	ECVF_RenderThreadSafe
+);
+
 #if PLATFORM_ANDROID
 bool GOpenGLShaderHackLastCompileSuccess = false;
 #endif
@@ -1195,14 +1220,6 @@ static void GetOpenGLProgramUniformBlockBinding(GLuint Program, GLuint UniformBl
 
 // ============================================================================================================================
 
-int32 GEvictOnBSSDestructLatency = 0;
-static FAutoConsoleVariableRef CVarEvictOnBssDestructLatency(
-	TEXT("r.OpenGL.EvictOnBSSDestruct.Latency"),
-	GEvictOnBSSDestructLatency,
-	TEXT(""),
-	ECVF_ReadOnly | ECVF_RenderThreadSafe
-);
-
 class FOpenGLLinkedProgram
 {
 public:
@@ -1245,14 +1262,13 @@ public:
 	class FLRUInfo
 	{
 	public:
-		FLRUInfo() : EvictBucket(-2) {}
 		// ID to LRU (if used) allows quick access when updating LRU status.
 		FSetElementId LRUNode;
 		// cached binary used to create this program.
-		TArray<uint8> CachedProgramBinary;
+		TUniqueObj<FOpenGLProgramBinary> CachedProgramBinary;
 
-		// < 0 if not pending eviction. Bucket index if pending eviction.
-		int32 EvictBucket;
+		void Touch() { LastTouchedFrame = GFrameNumber; }
+ 		uint32 LastTouchedFrame = 0;
 	} LRUInfo;
 
 private:
@@ -1346,56 +1362,6 @@ public:
 	}
 };
 
-static bool bMeasureEviction = false;
-class FDelayedEvictionContainer
-{
-public:
-
-	FDelayedEvictionContainer()
-	{
-		Init();
-	}
-
-	void Add(FOpenGLLinkedProgram* LinkedProgram);
-
-	FORCEINLINE_DEBUGGABLE static void OnProgramTouched(FOpenGLLinkedProgram* LinkedProgram)
-	{
-		if(LinkedProgram->LRUInfo.EvictBucket >=0 )
-		{
-			FDelayedEvictionContainer::Get().Remove(LinkedProgram);
-			INC_DWORD_STAT(STAT_OpenGLShaderLRUEvictionDelaySavedCount);
-		}
-	}
-
-	void Tick();
-
-	void Init();
-
-	void Remove(FOpenGLLinkedProgram* RemoveMe);
-
-	FORCEINLINE_DEBUGGABLE static FDelayedEvictionContainer & Get()
-	{
-		static FDelayedEvictionContainer DelayedEvictionContainer;
-		return DelayedEvictionContainer;
-	}
-private:
-	class FDelayEvictBucket
-	{
-	public:
-		FDelayEvictBucket() : NumToFreePerTick(0) {}
-		int32 NumToFreePerTick;
-		TSet<FOpenGLLinkedProgram*> ProgramsToEvict;
-	};
-
-	TArray<FDelayEvictBucket> Buckets;
-
-	int32 TotalBuckets;
-	int32 TimePerBucket;
-	int32 CurrentBucketTickCount;
-	int32 NewProgramBucket;
-	int32 EvictBucketIndex;
-};
-
 static void ConfigureStageStates(FOpenGLLinkedProgram* LinkedProgram)
 {
 	const FOpenGLLinkedProgramConfiguration &Config = LinkedProgram->Config;
@@ -1444,13 +1410,12 @@ namespace UE
 {
 	namespace OpenGL
 	{
-        static bool CreateGLProgramFromUncompressedBinary(GLuint& ProgramOUT, const TArray<uint8>& ProgramBinary)
+        static bool CreateGLProgramFromUncompressedBinary(GLuint& ProgramOUT, const TArrayView<const uint8>& ProgramBinary)
         {
 	        VERIFY_GL_SCOPE();
 	        GLuint GLProgramName = 0;
 	        FOpenGL::GenProgramPipelines(1, &GLProgramName);
 	        int32 BinarySize = ProgramBinary.Num();
-	        //UE_LOG(LogRHI, Log, TEXT("CreateGLProgramFromBinary : gen program %x, size: %d"), GLProgramName, BinarySize);
         
 	        check(BinarySize);
         
@@ -1464,7 +1429,7 @@ namespace UE
 	        return VerifyLinkedProgram(GLProgramName);
         }
 
-        static bool CreateGLProgramFromCompressedBinary(GLuint& ProgramOUT, const TArray<uint8>& CompressedProgramBinary)
+        static bool CreateGLProgramFromCompressedBinary(GLuint& ProgramOUT, const TArrayView<const uint8>& CompressedProgramBinary)
         {
 	        TArray<uint8> UncompressedProgramBinary;
         
@@ -1483,7 +1448,7 @@ namespace UE
 	        return false;
         }
 
-		bool CreateGLProgramFromBinary(GLuint& ProgramOUT, const TArray<uint8>& ProgramBinary)
+		bool CreateGLProgramFromBinary(GLuint& ProgramOUT, const TArrayView<const uint8>& ProgramBinary)
 		{
 			SCOPE_CYCLE_COUNTER(STAT_OpenGLCreateProgramFromBinaryTime)
 			bool bSuccess;
@@ -1544,69 +1509,75 @@ class FGLProgramCacheLRU
 {
 	class FEvictedGLProgram
 	{
-		FOpenGLLinkedProgram* LinkedProgram;
- 
-		FORCEINLINE_DEBUGGABLE TArray<uint8>& GetProgramBinary()
+		FOpenGLLinkedProgram* LinkedProgram = nullptr;
+
+		// When evicted, the cached binary program is owned here.
+		TUniqueObj<FOpenGLProgramBinary> CachedProgramBinary;
+
+		FORCEINLINE_DEBUGGABLE TArrayView<const uint8> GetProgramBinaryView() const
  		{
-			return LinkedProgram->LRUInfo.CachedProgramBinary;
+			return CachedProgramBinary->GetDataView();
  		}
 
 	public:
 
 		// Create an evicted program with the program binary provided.
-		FEvictedGLProgram(const FOpenGLProgramKey& ProgramKey, TArray<uint8>&& ProgramBinaryIn)
+		FEvictedGLProgram(const FOpenGLProgramKey& ProgramKey, TUniqueObj<FOpenGLProgramBinary>&& ProgramBinaryIn)
 		{
-			LinkedProgram = new FOpenGLLinkedProgram(ProgramKey);
-
-			GetProgramBinary() = MoveTemp(ProgramBinaryIn);
-
-			INC_MEMORY_STAT_BY(STAT_OpenGLShaderLRUProgramMemory, GetProgramBinary().Num());
+			check(ProgramBinaryIn->IsValid());
+			CachedProgramBinary = MoveTemp(ProgramBinaryIn);
 		}
 
 		FEvictedGLProgram(FOpenGLLinkedProgram* InLinkedProgram)
 			: LinkedProgram(InLinkedProgram)
 		{
-			bool bCreateProgramBinary = CVarLRUKeepProgramBinaryResident.GetValueOnAnyThread() == 0 || LinkedProgram->LRUInfo.CachedProgramBinary.Num() == 0;
-
+			// If the binary is not available then we must create it.
+			bool bCreateProgramBinary = LinkedProgram->LRUInfo.CachedProgramBinary->GetDataView().Num() == 0;
 			if( bCreateProgramBinary )
 			{
-			// build offline binary:
-				UE::OpenGL::GetProgramBinaryFromGLProgram(LinkedProgram->Program, GetProgramBinary());
-				INC_MEMORY_STAT_BY(STAT_OpenGLShaderLRUProgramMemory, GetProgramBinary().Num());
-			}
-
-			if(bMeasureEviction)
-			{
-				QUICK_SCOPE_CYCLE_COUNTER(STAT_EvictFromLRU_DeleteGLResource);
-				// Remove existing GL program and associated data.
-				LinkedProgram->DeleteGLResources();
+				// build the binary:
+				CachedProgramBinary = TUniqueObj<FOpenGLProgramBinary>(UE::OpenGL::GetProgramBinaryFromGLProgram(LinkedProgram->Program));
 			}
 			else
 			{
-				LinkedProgram->DeleteGLResources();
+				// transfer ownership from the linked program obj.
+				CachedProgramBinary = MoveTemp(LinkedProgram->LRUInfo.CachedProgramBinary);
+				// The binary should only be available if we're keeping resident, or is mmapped
+				check(CVarLRUKeepProgramBinaryResident.GetValueOnAnyThread() == 1 || !CachedProgramBinary->IsOwned());
 			}
+
+			LinkedProgram->DeleteGLResources();
 		}
 
-		void RestoreGLProgramFromBinary()
+		void RestoreGLProgramFromBinary(const FOpenGLProgramKey& ProgramKey)
 		{
+			if (LinkedProgram == nullptr)
+			{
+				LinkedProgram = new FOpenGLLinkedProgram(ProgramKey);
+			}
+
 			check(LinkedProgram->Program == 0);
-			bool bSuccess = UE::OpenGL::CreateGLProgramFromBinary(LinkedProgram->Program, GetProgramBinary());
+			check(LinkedProgram->LRUInfo.CachedProgramBinary->GetDataView().Num() == 0); // ownership of the binary should have been transferred to us from the linked prog at eviction time.
+
+			bool bSuccess = UE::OpenGL::CreateGLProgramFromBinary(LinkedProgram->Program, GetProgramBinaryView());
+
 			if(bSuccess)
 			{
-				if(CVarLRUKeepProgramBinaryResident.GetValueOnAnyThread() == 0)
+				// Always keep any mmapped data resident.
+				if(CVarLRUKeepProgramBinaryResident.GetValueOnAnyThread() || !CachedProgramBinary->IsOwned())
 				{
-					DEC_MEMORY_STAT_BY(STAT_OpenGLShaderLRUProgramMemory, GetProgramBinary().Num());
-					GetProgramBinary().Empty();
+					// avoid destruction the program binary by passing the binary back to the linked program.
+					LinkedProgram->LRUInfo.CachedProgramBinary = MoveTemp(CachedProgramBinary);
 				}
 			}
 			else
 			{
-				uint32 ProgramCRC = FCrc::MemCrc32(GetProgramBinary().GetData(), GetProgramBinary().Num());
-				UE_LOG(LogRHI, Log, TEXT("[%s, %d, %d, crc 0x%X]"), *LinkedProgram->Config.ProgramKey.ToString(), LinkedProgram->Program, GetProgramBinary().Num(), ProgramCRC );
+				uint32 ProgramCRC = FCrc::MemCrc32(GetProgramBinaryView().GetData(), GetProgramBinaryView().Num());
+				UE_LOG(LogRHI, Log, TEXT("[%s, %d, %d, crc 0x%X]"), *LinkedProgram->Config.ProgramKey.ToString(), LinkedProgram->Program, GetProgramBinaryView().Num(), ProgramCRC );
 				// dump first 32 bytes..
-				if (GetProgramBinary().Num() >= 32)
+				if (GetProgramBinaryView().Num() >= 32)
 				{
-					const uint32* MemPtr = (const uint32*)GetProgramBinary().GetData();
+					const uint32* MemPtr = (const uint32*)GetProgramBinaryView().GetData();
 					for (int32 Dump = 0; Dump < 8; Dump++)
 					{
 						UE_LOG(LogRHI, Log, TEXT("[%d :  0x%08X]"), Dump, *MemPtr++);
@@ -1650,7 +1621,7 @@ class FGLProgramCacheLRU
 			INC_DWORD_STAT(STAT_OpenGLShaderLRUMissCount);
 
 			// UE_LOG(LogRHI, Warning, TEXT("LRU: found and recovered EVICTED program %s"), *ProgramKey.ToString());
-			FoundEvicted->RestoreGLProgramFromBinary();
+			FoundEvicted->RestoreGLProgramFromBinary(ProgramKey);
 			FOpenGLLinkedProgram* LinkedProgram = FoundEvicted->GetLinkedProgram();
 
 			// Remove from the evicted program map.
@@ -1664,8 +1635,6 @@ class FGLProgramCacheLRU
 			// reconfigure the new program:
 			ConfigureGLProgramStageStates(LinkedProgram);
 
-			SetNewProgramStats(LinkedProgram->Program);
-
 			return LinkedProgram;
 		}
 
@@ -1678,27 +1647,13 @@ class FGLProgramCacheLRU
 		SCOPE_CYCLE_COUNTER(STAT_OpenGLShaderLRUEvictTime);
 		LinkedProgram->LRUInfo.LRUNode = FSetElementId();
 
-		if (LinkedProgram->LRUInfo.EvictBucket >= 0)
-		{
-			// remove it from the delayed eviction container since we're evicting now.
-			FDelayedEvictionContainer::Get().Remove(LinkedProgram);
-		}
-
 		DEC_DWORD_STAT(STAT_OpenGLShaderLRUProgramCount);
 
-		if (bMeasureEviction)
-		{
-			QUICK_SCOPE_CYCLE_COUNTER(STAT__EvictFromLRU_GetBinarySize);
-			LRUBinaryMemoryUse -= GetProgramBinarySize(LinkedProgram->Program);
-		}
-		else
-		{
-			LRUBinaryMemoryUse -= GetProgramBinarySize(LinkedProgram->Program);
-		}
+		LRUBinaryMemoryUse -= GetProgramBinarySize(LinkedProgram->Program);
 
 		checkf(!EvictedPrograms.Contains(LinkedProgram->Config.ProgramKey), TEXT("Program is already in the evicted program list: %s"), *LinkedProgram->Config.ProgramKey.ToString());
 		//UE_LOG(LogRHI, Warning, TEXT("LRU: Evicting program %d"), LinkedProgram->Program);
-		FEvictedGLProgram& test = EvictedPrograms.Emplace(LinkedProgram->Config.ProgramKey, FEvictedGLProgram(LinkedProgram));
+		FEvictedGLProgram& test = EvictedPrograms.Emplace(LinkedProgram->Config.ProgramKey, LinkedProgram);
 		INC_DWORD_STAT(STAT_OpenGLShaderLRUEvictedProgramCount);
 	}
 
@@ -1706,7 +1661,7 @@ public:
 
 	bool IsEvicted(const FOpenGLProgramKey& ProgramKey)
 	{
-		return FindEvicted(ProgramKey) != nullptr;
+		return EvictedPrograms.Contains(ProgramKey);
 	}
 
 	void EvictLeastRecentFromLRU()
@@ -1724,13 +1679,21 @@ public:
 		FOpenGLLinkedProgram* RemovedLinkedProgram = nullptr;
 		if(LRU.Remove(ProgramKey, RemovedLinkedProgram))
 		{
-			INC_DWORD_STAT(STAT_OpenGLShaderLRUScopeEvictedProgramCount);
 			EvictFromLRU(RemovedLinkedProgram);
 		}
 	}
 
+	void EvictLeastRecentByPredicate(TUniqueFunction<bool(FOpenGLLinkedProgram*)> Pred)
+	{
+		while (Pred(LRU.GetLeastRecent()))
+		{
+			EvictLeastRecentFromLRU();
+		}
+	}
+
 	FGLProgramCacheLRU() : LRUBinaryMemoryUse(0), LRU(LRUCapacity)
 	{
+		EvictedPrograms.Reserve(10000); // |TODO: establish a reasonable default.
 	}
 
 	bool IsLRUAtCapacity() const
@@ -1738,6 +1701,12 @@ public:
 		return LRU.Num() == CVarLRUMaxProgramCount.GetValueOnAnyThread() || LRU.Num() == LRU.Max() || LRUBinaryMemoryUse > CVarLRUMaxProgramBinarySize.GetValueOnAnyThread();
 	}
 
+	// returns the number of programs currently resident.
+	int32 GetLRUSize() const
+	{
+		return LRU.Num();
+	}
+
 	void Add(const FOpenGLProgramKey& ProgramKey, FOpenGLLinkedProgram* LinkedProgram)
 	{
 		// Remove least recently used programs until we reach our limit.
@@ -1753,17 +1722,18 @@ public:
 		}
 
 		LinkedProgram->LRUInfo.LRUNode = LRU.Add(ProgramKey, LinkedProgram);
-		FDelayedEvictionContainer::OnProgramTouched(LinkedProgram);
+		LinkedProgram->LRUInfo.Touch();
+
 		LRUBinaryMemoryUse += GetProgramBinarySize(LinkedProgram->Program);
 		INC_DWORD_STAT(STAT_OpenGLShaderLRUProgramCount);
 	}
 
-	void AddAsEvicted(const FOpenGLProgramKey& ProgramKey, TArray<uint8>&& ProgramBinary)
+	void AddOrReplaceEvicted(const FOpenGLProgramKey& ProgramKey, TUniqueObj<FOpenGLProgramBinary>&& ProgramBinaryData)
 	{
 		checkf(!LRU.Contains(ProgramKey), TEXT("Program is already in the LRU program list: %s"), *ProgramKey.ToString());
-		checkf(!IsEvicted(ProgramKey), TEXT("Program is already in the evicted program list: %s"), *ProgramKey.ToString());
 
-		FEvictedGLProgram& test = EvictedPrograms.Emplace(ProgramKey, FEvictedGLProgram(ProgramKey, MoveTemp(ProgramBinary)));
+		// TODO: if we're replacing then check we dont replace a mmapped binary with a non-mmapped binary?
+ 		FEvictedGLProgram& test = EvictedPrograms.Emplace(ProgramKey, FEvictedGLProgram(ProgramKey, MoveTemp(ProgramBinaryData)));
 
 		// UE_LOG(LogRHI, Warning, TEXT("LRU: adding EVICTED program %s"), *ProgramKey.ToString());
 
@@ -1791,6 +1761,11 @@ public:
 		}
 	}
 
+	bool Contains(const FOpenGLProgramKey& ProgramKey) const
+	{
+		return LRU.Contains(ProgramKey) || EvictedPrograms.Contains(ProgramKey);
+	}
+
 	FORCEINLINE_DEBUGGABLE void Touch(FOpenGLLinkedProgram* LinkedProgram)
 	{
 		if(LinkedProgram->LRUInfo.LRUNode.IsValidId())
@@ -1802,7 +1777,7 @@ public:
 			// This must find the program.
 			ensure(FindEvictedAndUpdateLRU(LinkedProgram->Config.ProgramKey));
 		}
-		FDelayedEvictionContainer::OnProgramTouched(LinkedProgram);
+		LinkedProgram->LRUInfo.Touch();
 	}
 
 	void Empty()
@@ -1889,6 +1864,19 @@ public:
 		}
 	}
 
+	FORCEINLINE_DEBUGGABLE bool Contains(const FOpenGLProgramKey& ProgramKey) const
+	{
+		if (bUseLRUCache)
+		{
+			return ProgramCacheLRU.Contains(ProgramKey);
+		}
+		else
+		{
+			return ProgramCache.Contains(ProgramKey);
+		}
+	}
+	
+
 	FORCEINLINE_DEBUGGABLE void Add(const FOpenGLProgramKey& ProgramKey, FOpenGLLinkedProgram* LinkedProgram)
 	{
 		if (bUseLRUCache)
@@ -1938,16 +1926,25 @@ public:
 		}
 	}
 
+	void EvictLeastRecentByPredicate(TUniqueFunction<bool(FOpenGLLinkedProgram*)> Pred)
+	{
+		check(IsUsingLRU());
+		if (ProgramCacheLRU.LRU.Num())
+		{
+			ProgramCacheLRU.EvictLeastRecentByPredicate(MoveTemp(Pred));
+		}
+	}
+
 	void EvictProgram(const FOpenGLProgramKey& ProgramKey)
 	{
 		check(IsUsingLRU());
 		ProgramCacheLRU.EvictProgramFromLRU(ProgramKey);
 	}
 
-	void AddAsEvicted(const FOpenGLProgramKey& ProgramKey, TArray<uint8>&& ProgramBinary)
+	void AddOrReplaceEvicted(const FOpenGLProgramKey& ProgramKey, TUniqueObj<FOpenGLProgramBinary>&& ProgramBinary)
 	{
 		check(IsUsingLRU());
-		ProgramCacheLRU.AddAsEvicted(ProgramKey, MoveTemp(ProgramBinary));
+		ProgramCacheLRU.AddOrReplaceEvicted(ProgramKey, MoveTemp(ProgramBinary));
 	}
 
 	bool IsEvicted(const FOpenGLProgramKey& ProgramKey)
@@ -1956,6 +1953,12 @@ public:
 		return ProgramCacheLRU.IsEvicted(ProgramKey);
 	}
 
+	int32 GetLRUSize() const
+	{
+		check(IsUsingLRU());
+		return ProgramCacheLRU.GetLRUSize();
+	}
+
 	void EnumerateLinkedPrograms(TFunction<void(FOpenGLLinkedProgram*)> EnumFunc)
 	{
 		if (bUseLRUCache)
@@ -1979,84 +1982,6 @@ static FGLProgramCache& GetOpenGLProgramsCache()
 	return ProgramsCache;
 }
 
-
-void FDelayedEvictionContainer::Init()
-{
-	const int32 EvictLatencyTicks = GEvictOnBSSDestructLatency;
-	const int32 NumLatencyBuckets = 3;
-	TotalBuckets = NumLatencyBuckets + 1;
-	Buckets.SetNum(TotalBuckets);
-	TimePerBucket = (EvictLatencyTicks)/(NumLatencyBuckets-1);
-	CurrentBucketTickCount = TimePerBucket;
-	NewProgramBucket = 0;
-	EvictBucketIndex = 1;
-}
-
-void FDelayedEvictionContainer::Add(FOpenGLLinkedProgram* LinkedProgram)
-{
-	if (GEvictOnBSSDestructLatency == 0)
-	{
-		GetOpenGLProgramsCache().EvictProgram(LinkedProgram->Config.ProgramKey);
-		return;
-	}
-
-	checkf(!GetOpenGLProgramsCache().IsEvicted(LinkedProgram->Config.ProgramKey), TEXT("FDelayedEvictionContainer::Add is already evicted! [%s], %d"), *LinkedProgram->Config.ProgramKey.ToString(), LinkedProgram->LRUInfo.EvictBucket);
-
-	if (LinkedProgram->LRUInfo.EvictBucket >=0 )
-	{
-		Remove(LinkedProgram);
-	}
-	Buckets[NewProgramBucket].ProgramsToEvict.Add(LinkedProgram);
-	LinkedProgram->LRUInfo.EvictBucket = NewProgramBucket;
-}
-
-void FDelayedEvictionContainer::Remove(FOpenGLLinkedProgram* RemoveMe)
-{
-	if (GEvictOnBSSDestructLatency == 0)
-	{
-		return;
-	}
-	check(RemoveMe->LRUInfo.EvictBucket >= 0);
-	ensure( Buckets[RemoveMe->LRUInfo.EvictBucket].ProgramsToEvict.Remove(RemoveMe) == 1 );
-	RemoveMe->LRUInfo.EvictBucket = -1;
-}
-
-void FDelayedEvictionContainer::Tick()
-{
-	if (GEvictOnBSSDestructLatency == 0)
-	{
-		return;
-	}
-
-	FDelayEvictBucket& EvictionBucket = Buckets[EvictBucketIndex];
-
-	const int32 NumToFree = EvictionBucket.ProgramsToEvict.Num();
-	if (NumToFree)
-	{
-		auto It = EvictionBucket.ProgramsToEvict.CreateIterator();
-		for (int32 i= FMath::Min(EvictionBucket.NumToFreePerTick, NumToFree)-1;i>=0;i--)
-		{
-			FOpenGLLinkedProgram* LinkedProgram = *It;
-			It.RemoveCurrent();
-			++It;
-			bMeasureEviction = true;
-			check(LinkedProgram->LRUInfo.EvictBucket == EvictBucketIndex);
-			LinkedProgram->LRUInfo.EvictBucket = -3; // Mark EvictBucket to indicated evicted from ProgramsToEvict, Prevent EvictProgram from attempting to remove again.
-			GetOpenGLProgramsCache().EvictProgram(LinkedProgram->Config.ProgramKey);
-			bMeasureEviction = false;
-		}
-	}
-
-	if (--CurrentBucketTickCount == 0)
-	{
-		check(EvictionBucket.ProgramsToEvict.Num() == 0);
-		EvictBucketIndex = (EvictBucketIndex+1) % Buckets.Num();
-		NewProgramBucket = (NewProgramBucket+1) % Buckets.Num();
-		CurrentBucketTickCount = TimePerBucket;
-		Buckets[EvictBucketIndex].NumToFreePerTick = (Buckets[EvictBucketIndex].ProgramsToEvict.Num() -1)/TimePerBucket + 1;
-	}
-}
-
 // This short queue preceding released programs cache is here because usually the programs are requested again
 // very shortly after they're released, so looking through recently released programs first provides tangible
 // performance improvement.
@@ -2575,18 +2500,13 @@ static FOpenGLLinkedProgram* LinkProgram( const FOpenGLLinkedProgramConfiguratio
 	check( (Config.Shaders[CrossCompiler::SHADER_STAGE_VERTEX].Resource == 0) != (Config.Shaders[CrossCompiler::SHADER_STAGE_COMPUTE].Resource == 0));
 	check( (Config.Shaders[CrossCompiler::SHADER_STAGE_PIXEL].Resource == 0) != (Config.Shaders[CrossCompiler::SHADER_STAGE_COMPUTE].Resource == 0));
 
-	TArray<uint8> CachedProgramBinary;
 	GLuint Program = 0;
-	bool bShouldLinkProgram = true;
 	if (FOpenGLProgramBinaryCache::IsEnabled())
 	{
-		// Try to create program from a saved binary
-		bShouldLinkProgram = !FOpenGLProgramBinaryCache::UseCachedProgram(Program, Config.ProgramKey, CachedProgramBinary);
-		if (bShouldLinkProgram)
-		{
-			// In case there is no saved binary in the cache, compile required shaders we have deferred before
-			UE::OpenGL::CompilePendingShaders(Config);
-		}
+		UE_CLOG(bFromPSOFileCache && !FOpenGLProgramBinaryCache::IsBuildingCache(), LogRHI, Error, TEXT("precompile PSO %s should be precompiled but is requesting compile."), *Config.ProgramKey.ToString());
+		// We should only see PSO cache programs here during cache creation. Otherwise the PSO should have been found already from the program cache.
+		//check(!bFromPSOFileCache || FOpenGLProgramBinaryCache::IsBuildingCache());
+		UE::OpenGL::CompilePendingShaders(Config);
 	}
 
 	if (Program == 0)
@@ -2595,61 +2515,53 @@ static FOpenGLLinkedProgram* LinkProgram( const FOpenGLLinkedProgramConfiguratio
 		FOpenGL::GenProgramPipelines(1, &Program);
 	}
 
-	if (bShouldLinkProgram)
+	if (Config.Shaders[CrossCompiler::SHADER_STAGE_VERTEX].Resource)
 	{
-		if (Config.Shaders[CrossCompiler::SHADER_STAGE_VERTEX].Resource)
-		{
-			FOpenGL::UseProgramStages(Program, GL_VERTEX_SHADER_BIT, Config.Shaders[CrossCompiler::SHADER_STAGE_VERTEX].Resource);
-		}
-		if (Config.Shaders[CrossCompiler::SHADER_STAGE_PIXEL].Resource)
-		{
-			FOpenGL::UseProgramStages(Program, GL_FRAGMENT_SHADER_BIT, Config.Shaders[CrossCompiler::SHADER_STAGE_PIXEL].Resource);
-		}
-		if (Config.Shaders[CrossCompiler::SHADER_STAGE_GEOMETRY].Resource)
-		{
-			FOpenGL::UseProgramStages(Program, GL_GEOMETRY_SHADER_BIT, Config.Shaders[CrossCompiler::SHADER_STAGE_GEOMETRY].Resource);
-		}
-		if (Config.Shaders[CrossCompiler::SHADER_STAGE_COMPUTE].Resource)
-		{
-			FOpenGL::UseProgramStages(Program, GL_COMPUTE_SHADER_BIT, Config.Shaders[CrossCompiler::SHADER_STAGE_COMPUTE].Resource);
-		}
+		FOpenGL::UseProgramStages(Program, GL_VERTEX_SHADER_BIT, Config.Shaders[CrossCompiler::SHADER_STAGE_VERTEX].Resource);
+	}
+	if (Config.Shaders[CrossCompiler::SHADER_STAGE_PIXEL].Resource)
+	{
+		FOpenGL::UseProgramStages(Program, GL_FRAGMENT_SHADER_BIT, Config.Shaders[CrossCompiler::SHADER_STAGE_PIXEL].Resource);
+	}
+	if (Config.Shaders[CrossCompiler::SHADER_STAGE_GEOMETRY].Resource)
+	{
+		FOpenGL::UseProgramStages(Program, GL_GEOMETRY_SHADER_BIT, Config.Shaders[CrossCompiler::SHADER_STAGE_GEOMETRY].Resource);
+	}
+	if (Config.Shaders[CrossCompiler::SHADER_STAGE_COMPUTE].Resource)
+	{
+		FOpenGL::UseProgramStages(Program, GL_COMPUTE_SHADER_BIT, Config.Shaders[CrossCompiler::SHADER_STAGE_COMPUTE].Resource);
+	}
 	
-		if(FOpenGLProgramBinaryCache::IsEnabled() || GetOpenGLProgramsCache().IsUsingLRU())
-		{
-			FOpenGL::ProgramParameter(Program, PROGRAM_BINARY_RETRIEVABLE_HINT, GL_TRUE);
-		}
-
-		// Link.
-		glLinkProgram(Program);
+	if(FOpenGLProgramBinaryCache::IsEnabled() || GetOpenGLProgramsCache().IsUsingLRU())
+	{
+		FOpenGL::ProgramParameter(Program, PROGRAM_BINARY_RETRIEVABLE_HINT, GL_TRUE);
 	}
 
-	if (VerifyLinkedProgram(Program))
-	{
-		if (bShouldLinkProgram)
-		{
-			SetNewProgramStats(Program);
+	// Link.
+	glLinkProgram(Program);
 
-			if (FOpenGLProgramBinaryCache::IsEnabled())
-			{
-				check(CachedProgramBinary.Num() == 0);
-				FOpenGLProgramBinaryCache::CacheProgram(Program, Config.ProgramKey, CachedProgramBinary);
-			}
-		}
-	}
-	else
+	if (!VerifyLinkedProgram(Program))
 	{
 		return nullptr;
 	}
-	
+
+	SetNewProgramStats(Program);
+
 	FOpenGL::BindProgramPipeline(Program);
 
 	FOpenGLLinkedProgram* LinkedProgram = new FOpenGLLinkedProgram(Config, Program);
 
-	if (GetOpenGLProgramsCache().IsUsingLRU() && CVarLRUKeepProgramBinaryResident.GetValueOnAnyThread() && CachedProgramBinary.Num())
+	if (FOpenGLProgramBinaryCache::IsEnabled() && FOpenGLProgramBinaryCache::IsBuildingCache() && GetOpenGLProgramsCache().IsUsingLRU() && bFromPSOFileCache)
 	{
-		// Store the binary data in LRUInfo, this avoids requesting a program binary from the driver when this program is evicted.
-		INC_MEMORY_STAT_BY(STAT_OpenGLShaderLRUProgramMemory, CachedProgramBinary.Num());
-		LinkedProgram->LRUInfo.CachedProgramBinary = MoveTemp(CachedProgramBinary);
+		// we need the binary if this is a pso for the currently building cache.
+		FOpenGLProgramBinary ProgramBinary = FOpenGLProgramBinaryCache::CreateAndCacheProgramBinary(Config.ProgramKey, Program);
+		check(ProgramBinary.IsValid());
+
+		if(CVarLRUKeepProgramBinaryResident.GetValueOnAnyThread())
+		{
+			// We've gone to the effort of creating the binary so store it in LRUInfo, avoids re-requesting the binary if the program is evicted.
+			LinkedProgram->LRUInfo.CachedProgramBinary = TUniqueObj<FOpenGLProgramBinary>(MoveTemp(ProgramBinary));
+		}
 	}
 	ConfigureStageStates(LinkedProgram);
 
@@ -2873,19 +2785,18 @@ void FOpenGLDynamicRHI::PrepareGFXBoundShaderState(const FGraphicsPipelineStateI
 			// use special null pixel shader when PixelShader was set to NULL
 			PixelShaderRHI = TShaderMapRef<FNULLPS>(GetGlobalShaderMap(GMaxRHIFeatureLevel)).GetPixelShader();
 		}
-	
+
 		FOpenGLProgramKey ProgramKey;
 		ProgramKey.ShaderHashes[CrossCompiler::SHADER_STAGE_VERTEX] = VertexShaderRHI->GetHash();
 		ProgramKey.ShaderHashes[CrossCompiler::SHADER_STAGE_PIXEL] = PixelShaderRHI->GetHash();
 
 		// add it to the runtime container. (this prevents createboundshader state from adding it to the binary cache.)
 		{
-			FScopeLock Lock(&GProgramBinaryCacheCS);
 			check(GetOpenGLProgramsCache().IsUsingLRU());
-			if (GetOpenGLProgramsCache().Find(ProgramKey, false) != nullptr)
+			if (FOpenGLProgramBinaryCache::DoesOutputCacheContain(ProgramKey))
 			{
-				//already done?
-				UE_LOG(LogRHI, Warning, TEXT("PrepareGFXBoundShaderState Already done? %s is pso %d "), *ProgramKey.ToString(), Initializer.bFromPSOFileCache ? 1 : 0);
+				//already done.
+				UE_LOG(LogRHI, Verbose, TEXT("PrepareGFXBoundShaderState Already in output PSO cache. %s is pso %d "), *ProgramKey.ToString(), Initializer.bFromPSOFileCache ? 1 : 0);
 				return;
 			}
 		}
@@ -2904,21 +2815,23 @@ void FOpenGLDynamicRHI::PrepareGFXBoundShaderState(const FGraphicsPipelineStateI
 		}
 
 		FString FailLog;
-		TArray<uint8> CompiledProgramResult = FAndroidOpenGL::DispatchAndWaitForRemoteGLProgramCompile(TArrayView<uint8>((uint8*)&ProgramKey, sizeof(ProgramKey)), VSCode, PSCode, ComputeGlslCode, FailLog);
+		TArray<uint8> CompiledProgramBytes = FAndroidOpenGL::DispatchAndWaitForRemoteGLProgramCompile(TArrayView<uint8>((uint8*)&ProgramKey, sizeof(ProgramKey)), VSCode, PSCode, ComputeGlslCode, FailLog);
 
 		if(FailLog.IsEmpty())
 		{
-			GLenum glFormat = *(GLenum*)CompiledProgramResult.GetData();		
+			GLenum glFormat = *(GLenum*)CompiledProgramBytes.GetData();		
 			if (UE::OpenGL::IsStoringCompressedBinaryPrograms())
 			{
 				
 				TArray<uint8> CompressedCompiledProgramResult;
-				UE::OpenGL::CompressProgramBinary(CompiledProgramResult, CompressedCompiledProgramResult);
-				CompiledProgramResult = MoveTemp(CompressedCompiledProgramResult);
+				UE::OpenGL::CompressProgramBinary(CompiledProgramBytes, CompressedCompiledProgramResult);
+				CompiledProgramBytes = MoveTemp(CompressedCompiledProgramResult);
 			}
 
+			FOpenGLProgramBinary CompiledProgram(MoveTemp(CompiledProgramBytes));
+
 			// add it to the binary file.
-			FOpenGLProgramBinaryCache::CacheProgramBinary(ProgramKey, CompiledProgramResult);
+			FOpenGLProgramBinaryCache::CacheProgramBinary(ProgramKey, CompiledProgram);
 
 			// add it to the runtime container. (this prevents createboundshader state from adding it to the binary cache.)
 			{
@@ -2926,12 +2839,12 @@ void FOpenGLDynamicRHI::PrepareGFXBoundShaderState(const FGraphicsPipelineStateI
 				check(GetOpenGLProgramsCache().IsUsingLRU());
 				if (GetOpenGLProgramsCache().Find(ProgramKey, false) == nullptr)
 				{
-					GetOpenGLProgramsCache().AddAsEvicted(ProgramKey, MoveTemp(CompiledProgramResult));
+					GetOpenGLProgramsCache().AddOrReplaceEvicted(ProgramKey, TUniqueObj<FOpenGLProgramBinary>(MoveTemp(CompiledProgram)));
 				}
 				else
 				{
-					// beaten to it by another thread.
-					UE_LOG(LogRHI, Warning, TEXT("PrepareGFXBoundShaderState skipped add. %s is pso %d "), *ProgramKey.ToString(), Initializer.bFromPSOFileCache ? 1 : 0);
+					// beaten to it elsewhere.
+					UE_LOG(LogRHI, Verbose, TEXT("PrepareGFXBoundShaderState skipped add. %s is pso %d "), *ProgramKey.ToString(), Initializer.bFromPSOFileCache ? 1 : 0);
 				}
 			}
 
@@ -2976,8 +2889,6 @@ FBoundShaderStateRHIRef FOpenGLDynamicRHI::RHICreateBoundShaderState_OnThisThrea
 		PixelShaderRHI = GetNULLPixelShader();
 	}
 
-
-
 	// Check for an existing bound shader state which matches the parameters
 	FCachedBoundShaderStateLink* CachedBoundShaderStateLink = GetCachedBoundShaderState(
 		VertexDeclarationRHI,
@@ -3026,27 +2937,38 @@ FBoundShaderStateRHIRef FOpenGLDynamicRHI::RHICreateBoundShaderState_OnThisThrea
 
 		if (!LinkedProgram)
 		{
-			bool bFindAndCreateEvictedProgram = true;
+			bool bCreateEvictedBinary = false;
 			// If this is this a request from the PSOFC then do not create an evicted program.
 			if (bFromPSOFileCache && GetOpenGLProgramsCache().IsUsingLRU())
 			{
-				bFindAndCreateEvictedProgram = false;
+				// no-op, No need to create a precompile PSO that we already have.
+				// return null as nothing should be using the resultant BSS of a precompile PSO either.
+				bool bContained = GetOpenGLProgramsCache().Contains(Config.ProgramKey) || FOpenGLProgramBinaryCache::CheckSinglePendingGLProgramCreateRequest(Config.ProgramKey);
+				if (bContained)
+				{
+					return nullptr;
+				}
+				bCreateEvictedBinary = true;
 			}
 
-			FOpenGLLinkedProgram* CachedProgram = GetOpenGLProgramsCache().Find(Config.ProgramKey, bFindAndCreateEvictedProgram);
-			if (!CachedProgram)
+			FOpenGLLinkedProgram* CachedProgram = nullptr;
+			if (!bCreateEvictedBinary)
 			{
-				// ensure that pending request for this program has been completed before
-				if (FOpenGLProgramBinaryCache::CheckSinglePendingGLProgramCreateRequest(Config.ProgramKey))
+				CachedProgram = GetOpenGLProgramsCache().Find(Config.ProgramKey, true);
+				if (!CachedProgram)
 				{
-					CachedProgram = GetOpenGLProgramsCache().Find(Config.ProgramKey, bFindAndCreateEvictedProgram);
+					// ensure that pending request for this program has been completed before
+					if (FOpenGLProgramBinaryCache::CheckSinglePendingGLProgramCreateRequest(Config.ProgramKey))
+					{
+						CachedProgram = GetOpenGLProgramsCache().Find(Config.ProgramKey, true);
+					}
 				}
 			}
 
 			if (CachedProgram)
 			{
 				LinkedProgram = CachedProgram;
-				if (!LinkedProgram->bConfigIsInitalized && bFindAndCreateEvictedProgram)
+				if (!LinkedProgram->bConfigIsInitalized)
 				{
 					LinkedProgram->SetConfig(Config);
 					// We now have the config for this program, we must configure the program for use.
@@ -3087,15 +3009,31 @@ FBoundShaderStateRHIRef FOpenGLDynamicRHI::RHICreateBoundShaderState_OnThisThrea
 					RHIGetPanicDelegate().ExecuteIfBound(LinkFailurePanic);
 					UE_LOG(LogRHI, Fatal, TEXT("Failed to link program [%s]. Current total programs: %d, precompile: %d"), *Config.ProgramKey.ToString(), GNumPrograms, (uint32)bFromPSOFileCache);
 				}
-
-				GetOpenGLProgramsCache().Add(Config.ProgramKey, LinkedProgram);
-
-				// if building the cache file and using the LRU then evict the last shader created. this will reduce the risk of fragmentation of the driver's program memory.
-				if (bFindAndCreateEvictedProgram == false && FOpenGLProgramBinaryCache::IsBuildingCache())
+				else
 				{
-					GetOpenGLProgramsCache().EvictMostRecent();
+					if (bCreateEvictedBinary)
+					{
+						// Store the program as evicted binary only, this only happens for precompile PSOs and the returned result is not used.
+						if(LinkedProgram->LRUInfo.CachedProgramBinary->IsValid())
+						{
+							GetOpenGLProgramsCache().AddOrReplaceEvicted(Config.ProgramKey, MoveTemp(LinkedProgram->LRUInfo.CachedProgramBinary));
+						}
+						else
+						{
+							// not expected. This means the PSO is part of the pso precompile but is somehow not in the existing cache.
+							TUniqueObj<FOpenGLProgramBinary> Bin(UE::OpenGL::GetProgramBinaryFromGLProgram(LinkedProgram->Program));
+							GetOpenGLProgramsCache().AddOrReplaceEvicted(Config.ProgramKey, MoveTemp(Bin));
+						}
+						delete LinkedProgram;
+						return nullptr;
+					}
+					else
+					{
+						GetOpenGLProgramsCache().Add(Config.ProgramKey, LinkedProgram);
+					}
 				}
 			}
+			check(!bCreateEvictedBinary);
 		}
 
 		check(VertexDeclarationRHI);
@@ -3311,14 +3249,6 @@ FOpenGLBoundShaderState::FOpenGLBoundShaderState(
 	}
 }
 
-TAutoConsoleVariable<int32> CVarEvictOnBssDestruct(
-	TEXT("r.OpenGL.EvictOnBSSDestruct"),
-	0,
-	TEXT(""),
-	ECVF_ReadOnly | ECVF_RenderThreadSafe
-);
-
-
 FOpenGLBoundShaderState::~FOpenGLBoundShaderState()
 {
 	check(LinkedProgram);
@@ -3334,11 +3264,6 @@ FOpenGLBoundShaderState::~FOpenGLBoundShaderState()
 				StaticLastReleasedProgramsIndex = 0;
 			}
 
-			if (CVarEvictOnBssDestruct.GetValueOnAnyThread() && GetOpenGLProgramsCache().IsUsingLRU())
-			{
-				FDelayedEvictionContainer::Get().Add(LinkedProgram);
-			}
-
 			OnProgramDeletion(LinkedProgram->Program);
 		}
 	});
@@ -3727,25 +3652,29 @@ namespace UE
 {
 	namespace OpenGL
 	{
-		// Called from the binary file cache when program loads from disk.
-		void OnGLProgramLoadedFromBinaryCache(const FOpenGLProgramKey& ProgramKey, TArray<uint8>&& ProgramBinaryData)
+		// Called from the binary file cache when the binary version of a program has been encountered.
+		void OnGLProgramLoadedFromBinaryCache(const FOpenGLProgramKey& ProgramKey, TUniqueObj<FOpenGLProgramBinary>&& ProgramBinaryData)
 		{
-			const bool bProgramExists = GetOpenGLProgramsCache().Find(ProgramKey, false) != nullptr;
-
+			FOpenGLLinkedProgram* FoundProgram = GetOpenGLProgramsCache().Find(ProgramKey, false);
+			const bool bProgramExists = FoundProgram != nullptr;
+			const bool bIsEvicted = bProgramExists && GetOpenGLProgramsCache().IsEvicted(ProgramKey);
 			if (GetOpenGLProgramsCache().IsUsingLRU())
 			{
-				if (!bProgramExists)
+				// always replace any existing binary data with this.
+				if (!bProgramExists || bIsEvicted)
 				{
 					// Always add programs as evicted, 1st use will create them as programs.
 					// This will reduce pressure on driver by ensuring only used programs
 					// are created.
 					// In this case do not create the GL program.
-					GetOpenGLProgramsCache().AddAsEvicted(ProgramKey, MoveTemp(ProgramBinaryData));
+					GetOpenGLProgramsCache().AddOrReplaceEvicted(ProgramKey, MoveTemp(ProgramBinaryData));
+					check(!bProgramExists || FoundProgram->LRUInfo.CachedProgramBinary->GetDataView().IsEmpty());
 				}
 				else
 				{
-					// The program is already in use, discard the binary data.
-					ProgramBinaryData.Empty();
+					// replace the existing program with the incoming data.
+					// For PSO cache programs this will replace the heap allocated data with a region of the mmapped program file.
+					FoundProgram->LRUInfo.CachedProgramBinary = MoveTemp(ProgramBinaryData);
 				}
 			}
 			else
@@ -3753,53 +3682,46 @@ namespace UE
 				if (!bProgramExists)
 				{
 					GLuint GLProgramId = 0;
-					bool bSuccess = UE::OpenGL::CreateGLProgramFromBinary(GLProgramId, ProgramBinaryData);
+					bool bSuccess = UE::OpenGL::CreateGLProgramFromBinary(GLProgramId, ProgramBinaryData->GetDataView());
 					if (!bSuccess)
 					{
-						UE_LOG(LogRHI, Log, TEXT("[%s, %d, %d]"), *ProgramKey.ToString(), GLProgramId, ProgramBinaryData.Num());
+						UE_LOG(LogRHI, Log, TEXT("[%s, %d, %d]"), *ProgramKey.ToString(), GLProgramId, ProgramBinaryData->GetDataView().Num());
 						RHIGetPanicDelegate().ExecuteIfBound(FName("FailedBinaryProgramCreateLoadRequest"));
 						UE_LOG(LogRHI, Fatal, TEXT("CompleteLoadedGLProgramRequest_internal : Failed to create GL program from binary data! [%s]"), *ProgramKey.ToString());
 					}
 					FOpenGLLinkedProgram* NewLinkedProgram = new FOpenGLLinkedProgram(ProgramKey, GLProgramId);
 					GetOpenGLProgramsCache().Add(ProgramKey, NewLinkedProgram);
-					SetNewProgramStats(GLProgramId);
 				}
-
-				// Finished with binary data.
-				ProgramBinaryData.Empty();
 			}
 		}
 	}
 }
 
-void FOpenGLDynamicRHI::EndFrameTick()
-{
-	FDelayedEvictionContainer::Get().Tick();
-	FOpenGLProgramBinaryCache::CheckPendingGLProgramCreateRequests();
-	FTextureEvictionLRU::Get().TickEviction();
-}
-
-static FDelegateHandle OnSharedShaderCodeRequest;
-//static FDelegateHandle OnSharedShaderCodeRelease;
-
-void OnShaderLibraryRequestShaderCode(const FSHAHash& Hash, FArchive* Ar)
-{
-	FOpenGLProgramBinaryCache::OnShaderLibraryRequestShaderCode(Hash, Ar);
-}
-
-//void OnShaderLibraryReleaseShaderCode(const FSHAHash& Hash)
-//{
-//}
-
-
-void FOpenGLDynamicRHI::RegisterSharedShaderCodeDelegates()
+static void TickProgramLRU()
 {
-	OnSharedShaderCodeRequest = FShaderCodeLibrary::RegisterSharedShaderCodeRequestDelegate_Handle(FSharedShaderCodeRequest::FDelegate::CreateStatic(&OnShaderLibraryRequestShaderCode));
-	//OnSharedShaderCodeRelease = FShaderCodeLibrary::RegisterSharedShaderCodeReleaseDelegate_Handle(FSharedShaderCodeRelease::FDelegate::CreateStatic(&OnShaderLibraryReleaseShaderCode));
+	if (GetOpenGLProgramsCache().IsUsingLRU() && GUntouchedProgramEvictTimeSeconds > 0)
+	{
+		int FramePace = FPlatformRHIFramePacer::GetFramePace();
+		int MaxProgramsToEvictPerFrame = 10;
+		int32 CurrentResidentCount = GetOpenGLProgramsCache().GetLRUSize();
+		MaxProgramsToEvictPerFrame = FMath::Min(FMath::Max(0, CurrentResidentCount - GProgramLRUResidentCountBeforeEviction), MaxProgramsToEvictPerFrame);
+		if (MaxProgramsToEvictPerFrame > 0)
+		{
+			GetOpenGLProgramsCache().EvictLeastRecentByPredicate(
+				[FramePace, &MaxProgramsToEvictPerFrame](FOpenGLLinkedProgram* LeastRecentProgram)
+				{
+					MaxProgramsToEvictPerFrame--;
+					uint32 LastFrameAllowed = GFrameNumber - FMath::Min(GFrameNumber, (uint32)(FramePace * GUntouchedProgramEvictTimeSeconds));
+					return MaxProgramsToEvictPerFrame >= 0 && LeastRecentProgram->LRUInfo.LastTouchedFrame < LastFrameAllowed;
+				}
+			);
+		}
+	}
 }
 
-void FOpenGLDynamicRHI::UnregisterSharedShaderCodeDelegates()
+void FOpenGLDynamicRHI::EndFrameTick()
 {
-	FShaderCodeLibrary::UnregisterSharedShaderCodeRequestDelegate_Handle(OnSharedShaderCodeRequest);
-	//FShaderCodeLibrary::UnregisterSharedShaderCodeReleaseDelegate_Handle(OnSharedShaderCodeRelease);
+	TickProgramLRU();
+	FOpenGLProgramBinaryCache::CheckPendingGLProgramCreateRequests();
+	FTextureEvictionLRU::Get().TickEviction();
 }
diff --git a/Engine/Source/Runtime/OpenGLDrv/Private/OpenGLUtil.cpp b/Engine/Source/Runtime/OpenGLDrv/Private/OpenGLUtil.cpp
--- a/Engine/Source/Runtime/OpenGLDrv/Private/OpenGLUtil.cpp
+++ b/Engine/Source/Runtime/OpenGLDrv/Private/OpenGLUtil.cpp
@@ -84,11 +84,10 @@ DEFINE_STAT(STAT_OpenGLCreateProgramFromBinaryTime);
 DEFINE_STAT(STAT_OpenGLShaderLRUEvictTime);
 DEFINE_STAT(STAT_OpenGLShaderLRUMissTime);
 DEFINE_STAT(STAT_OpenGLShaderLRUProgramCount);
-DEFINE_STAT(STAT_OpenGLShaderLRUEvictionDelaySavedCount);
 DEFINE_STAT(STAT_OpenGLShaderLRUEvictedProgramCount);
-DEFINE_STAT(STAT_OpenGLShaderLRUScopeEvictedProgramCount);
 DEFINE_STAT(STAT_OpenGLShaderLRUMissCount);
 DEFINE_STAT(STAT_OpenGLShaderLRUProgramMemory);
+DEFINE_STAT(STAT_OpenGLShaderLRUProgramMemoryMapped);
 
 #if OPENGLRHI_DETAILED_STATS
 DEFINE_STAT(STAT_OpenGLDrawPrimitiveTime);
@@ -170,117 +169,3 @@ void RunOnGLRenderContextThread(TUniqueFunction<void(void)> GLFunc, bool bWaitFo
 		}
 	}
 }
-
-namespace UE
-{
-	namespace OpenGL
-	{
-		static TAutoConsoleVariable<int32> CVarStoreCompressedBinaries(
-			TEXT("r.OpenGL.StoreCompressedProgramBinaries"),
-			0,
-			TEXT(""),
-			ECVF_ReadOnly | ECVF_RenderThreadSafe
-		);
-
-		bool IsStoringCompressedBinaryPrograms()
-		{
-			return CVarStoreCompressedBinaries.GetValueOnAnyThread() != 0;
-		}
-
-		bool UncompressCompressedBinaryProgram(const TArray<uint8>& CompressedProgramBinary, TArray<uint8>& UncompressedProgramBinaryOUT)
-		{
-			if (ensure(CompressedProgramBinary.Num() > sizeof(FCompressedProgramBinaryHeader)))
-			{
-				FCompressedProgramBinaryHeader* Header = (FCompressedProgramBinaryHeader*)CompressedProgramBinary.GetData();
-
-				if (Header->UncompressedSize == FCompressedProgramBinaryHeader::NotCompressed)
-				{
-					const uint32 ProgramSize = CompressedProgramBinary.Num() - sizeof(FCompressedProgramBinaryHeader);
-					UncompressedProgramBinaryOUT.SetNumUninitialized(ProgramSize);
-					FMemory::Memcpy(UncompressedProgramBinaryOUT.GetData(), CompressedProgramBinary.GetData() + sizeof(FCompressedProgramBinaryHeader), ProgramSize);
-					return true;
-				}
-				else
-				{
-					UncompressedProgramBinaryOUT.AddUninitialized(Header->UncompressedSize);
-
-					if (Header->UncompressedSize > 0
-						&& FCompression::UncompressMemory(NAME_Zlib, UncompressedProgramBinaryOUT.GetData(), UncompressedProgramBinaryOUT.Num(), CompressedProgramBinary.GetData() + sizeof(FCompressedProgramBinaryHeader), CompressedProgramBinary.Num() - sizeof(FCompressedProgramBinaryHeader)))
-					{
-						return true;
-					}
-				}
-			}
-			return false;
-		}
-
-		bool GetUncompressedProgramBinaryFromGLProgram(GLuint Program, TArray<uint8>& ProgramBinaryOUT)
-		{
-			VERIFY_GL_SCOPE();
-
-			// pull binary from linked program
-			GLint BinaryLength = -1;
-			glGetProgramiv(Program, GL_PROGRAM_BINARY_LENGTH, &BinaryLength);
-			if (BinaryLength > 0)
-			{
-				ProgramBinaryOUT.SetNumUninitialized(BinaryLength + sizeof(GLenum));
-				uint8* ProgramBinaryPtr = ProgramBinaryOUT.GetData();
-				// BinaryFormat is stored at the start of ProgramBinary array
-				FOpenGL::GetProgramBinary(Program, BinaryLength, &BinaryLength, (GLenum*)ProgramBinaryPtr, ProgramBinaryPtr + sizeof(GLenum));
-				return true;
-			}
-			return false;
-		}
-
-		void CompressProgramBinary(const TArray<uint8>& UncompressedProgramBinary, TArray<uint8>& ProgramBinaryOUT)
-		{
-			check(IsStoringCompressedBinaryPrograms());
-			check(ProgramBinaryOUT.IsEmpty());
-
-			int32 CompressedSize = FCompression::CompressMemoryBound(NAME_Zlib, UncompressedProgramBinary.Num());
-			uint32 CompressedHeaderSize = sizeof(FCompressedProgramBinaryHeader);
-			ProgramBinaryOUT.AddUninitialized(CompressedSize + CompressedHeaderSize);
-			bool bSuccess = FCompression::CompressMemory(NAME_Zlib, ProgramBinaryOUT.GetData() + CompressedHeaderSize, CompressedSize, UncompressedProgramBinary.GetData(), UncompressedProgramBinary.Num());
-			if (bSuccess)
-			{
-				ProgramBinaryOUT.SetNum(CompressedSize + CompressedHeaderSize);
-				ProgramBinaryOUT.Shrink();
-				FCompressedProgramBinaryHeader* Header = (FCompressedProgramBinaryHeader*)ProgramBinaryOUT.GetData();
-				Header->UncompressedSize = UncompressedProgramBinary.Num();
-			}
-			else
-			{
-				// failed, store the uncompressed version.
-				UE_LOG(LogRHI, Log, TEXT("Storing binary program uncompressed (%d, %d, %d)"), UncompressedProgramBinary.Num(), ProgramBinaryOUT.Num(), CompressedSize);
-				ProgramBinaryOUT.SetNumUninitialized(UncompressedProgramBinary.Num() + CompressedHeaderSize);
-				FCompressedProgramBinaryHeader* Header = (FCompressedProgramBinaryHeader*)ProgramBinaryOUT.GetData();
-				Header->UncompressedSize = FCompressedProgramBinaryHeader::NotCompressed;
-				FMemory::Memcpy(ProgramBinaryOUT.GetData() + sizeof(FCompressedProgramBinaryHeader), UncompressedProgramBinary.GetData(), UncompressedProgramBinary.Num());
-			}
-		}
-
-		bool GetCompressedProgramBinaryFromGLProgram(GLuint Program, TArray<uint8>& ProgramBinaryOUT)
-		{
-			// get uncompressed binary
-			TArray<uint8> UncompressedProgramBinary;
-			if (GetUncompressedProgramBinaryFromGLProgram(Program, UncompressedProgramBinary))
-			{
-				CompressProgramBinary(UncompressedProgramBinary, ProgramBinaryOUT);
-				return true;
-			}
-			return false;
-		}
-
-		bool GetProgramBinaryFromGLProgram(GLuint Program, TArray<uint8>& ProgramBinaryOUT)
-		{
-			if (IsStoringCompressedBinaryPrograms())
-			{
-				return GetCompressedProgramBinaryFromGLProgram(Program, ProgramBinaryOUT);
-			}
-			else
-			{
-				return GetUncompressedProgramBinaryFromGLProgram(Program, ProgramBinaryOUT);
-			}
-		}
-	}
-}
\ No newline at end of file
diff --git a/Engine/Source/Runtime/OpenGLDrv/Public/OpenGLDrv.h b/Engine/Source/Runtime/OpenGLDrv/Public/OpenGLDrv.h
--- a/Engine/Source/Runtime/OpenGLDrv/Public/OpenGLDrv.h
+++ b/Engine/Source/Runtime/OpenGLDrv/Public/OpenGLDrv.h
@@ -1010,9 +1010,6 @@ public:
 	void InitializeGLTextureInternal(FOpenGLTexture* Texture, void const* BulkDataPtr, uint64 BulkDataSize);
 private:
 
-	void RegisterSharedShaderCodeDelegates();
-	void UnregisterSharedShaderCodeDelegates();
-
 	void SetupRecursiveResources();
 
 	void ApplyTextureStage(FOpenGLContextState& ContextState, GLint TextureIndex, const FTextureStage& TextureStage, FOpenGLSamplerState* SamplerState);
diff --git a/Engine/Source/Runtime/OpenGLDrv/Public/OpenGLShaderResources.h b/Engine/Source/Runtime/OpenGLDrv/Public/OpenGLShaderResources.h
--- a/Engine/Source/Runtime/OpenGLDrv/Public/OpenGLShaderResources.h
+++ b/Engine/Source/Runtime/OpenGLDrv/Public/OpenGLShaderResources.h
@@ -393,8 +393,9 @@ struct FOpenGLBindlessSamplerInfo
 };
 
 // unique identifier for a program. (composite of shader keys)
-struct FOpenGLProgramKey
+class FOpenGLProgramKey
 {
+public:
 	FOpenGLProgramKey() {}
 
 	friend bool operator == (const FOpenGLProgramKey& A, const FOpenGLProgramKey& B)
@@ -407,6 +408,11 @@ struct FOpenGLProgramKey
 		return bHashMatch;
 	}
 
+	friend bool operator != (const FOpenGLProgramKey& A, const FOpenGLProgramKey& B)
+	{
+		return !(A==B);
+	}
+
 	friend uint32 GetTypeHash(const FOpenGLProgramKey& Key)
 	{
 		return FCrc::MemCrc32(Key.ShaderHashes, sizeof(Key.ShaderHashes));
diff --git a/Engine/Source/Runtime/OpenGLDrv/Public/OpenGLUtil.h b/Engine/Source/Runtime/OpenGLDrv/Public/OpenGLUtil.h
--- a/Engine/Source/Runtime/OpenGLDrv/Public/OpenGLUtil.h
+++ b/Engine/Source/Runtime/OpenGLDrv/Public/OpenGLUtil.h
@@ -190,26 +190,3 @@ inline bool ShouldRunGLRenderContextOpOnThisThread(FRHICommandListBase& RHICmdLi
 }
 
 #endif
-
-
-// Program Binary helpers.
-namespace UE
-{
-	namespace OpenGL
-	{
-		bool IsStoringCompressedBinaryPrograms();
-
-		struct FCompressedProgramBinaryHeader
-		{
-			static const uint32 NotCompressed = 0xFFFFFFFF;
-			uint32 UncompressedSize;
-		};
-
-		bool GetProgramBinaryFromGLProgram(GLuint Program, TArray<uint8>& ProgramBinaryOUT);
-		bool GetUncompressedProgramBinaryFromGLProgram(GLuint Program, TArray<uint8>& ProgramBinaryOUT);
-		bool GetCompressedProgramBinaryFromGLProgram(GLuint Program, TArray<uint8>& ProgramBinaryOUT);
-		void CompressProgramBinary(const TArray<uint8>& UncompressedProgramBinary, TArray<uint8>& ProgramBinaryOUT);
-		bool UncompressCompressedBinaryProgram(const TArray<uint8>& CompressedProgramBinary, TArray<uint8>& UncompressedProgramBinaryOUT);
-	}
-}
-
