Android GLES: Swap from size to speed biased compression when storing GLSL shaders.
2-3x speed up when creating shaders at the cost of ~10% increased memory use.

diff --git a/Engine/Source/Runtime/OpenGLDrv/Private/OpenGLShaders.cpp b/Engine/Source/Runtime/OpenGLDrv/Private/OpenGLShaders.cpp
--- a/Engine/Source/Runtime/OpenGLDrv/Private/OpenGLShaders.cpp
+++ b/Engine/Source/Runtime/OpenGLDrv/Private/OpenGLShaders.cpp
@@ -352,7 +352,7 @@ private:
 			CompressedSize,
 			(void*)InGlslCode.GetData(),
 			UncompressedSize,
-			COMPRESS_BiasMemory);
+			COMPRESS_BiasSpeed);
 
 		if (bCompressed)
 		{
