Improve OpenGL's handling of multiple binary PSO caches.
Cache filenames are now based on PSO cache name, should prevent orphaned caches.
Empty PSO caches (0 usable entries) are not written.
Removal of critical section guarding GL's program cache (context thread pulls in binary cache data)
Clean up of the binary caching logic, (GL no longer attempts to deal with precache PSOs RHISetGraphicsPipelineState. A higher level bug that caused this situation has been fixed)

diff --git a/Engine/Source/Runtime/OpenGLDrv/Private/Android/AndroidOpenGL.cpp b/Engine/Source/Runtime/OpenGLDrv/Private/Android/AndroidOpenGL.cpp
--- a/Engine/Source/Runtime/OpenGLDrv/Private/Android/AndroidOpenGL.cpp
+++ b/Engine/Source/Runtime/OpenGLDrv/Private/Android/AndroidOpenGL.cpp
@@ -155,12 +155,12 @@ static bool AreAndroidOpenGLRemoteCompileServicesAvailable()
 	{
 		const FString* ConfigRulesDisableProgramCompileServices = FAndroidMisc::GetConfigRulesVariable(TEXT("DisableProgramCompileServices"));
 		bool bConfigRulesDisableProgramCompileServices = ConfigRulesDisableProgramCompileServices && ConfigRulesDisableProgramCompileServices->Equals("true", ESearchCase::IgnoreCase);
-		static const auto CVarProgramLRU = IConsoleManager::Get().FindConsoleVariable(TEXT("r.OpenGL.EnableProgramLRUCache"));
+		static const auto CVarBinaryProgramCache = IConsoleManager::Get().FindConsoleVariable(TEXT("r.ProgramBinaryCache.Enable"));
 		static const auto CVarNumRemoteProgramCompileServices = IConsoleManager::Get().FindConsoleVariable(TEXT("Android.OpenGL.NumRemoteProgramCompileServices"));
 
-		RemoteCompileService = !bConfigRulesDisableProgramCompileServices && OpenGLRemoteGLProgramCompileJNI.bAllFound && (CVarProgramLRU->GetInt() != 0) && (CVarNumRemoteProgramCompileServices->GetInt() > 0);
+		RemoteCompileService = !bConfigRulesDisableProgramCompileServices && OpenGLRemoteGLProgramCompileJNI.bAllFound && (CVarBinaryProgramCache->GetInt() != 0) && (CVarNumRemoteProgramCompileServices->GetInt() > 0);
 		FGenericCrashContext::SetEngineData(TEXT("Android.PSOService"), RemoteCompileService == 0? TEXT("disabled") : TEXT("enabled"));
-		UE_CLOG(!RemoteCompileService, LogRHI, Log, TEXT("Remote PSO services disabled: (%d, %d, %d, %d)"), bConfigRulesDisableProgramCompileServices, OpenGLRemoteGLProgramCompileJNI.bAllFound, CVarProgramLRU->GetInt(), CVarNumRemoteProgramCompileServices->GetInt());
+		UE_CLOG(!RemoteCompileService, LogRHI, Log, TEXT("Remote PSO services disabled: (%d, %d, %d, %d)"), bConfigRulesDisableProgramCompileServices, OpenGLRemoteGLProgramCompileJNI.bAllFound, CVarBinaryProgramCache->GetInt(), CVarNumRemoteProgramCompileServices->GetInt());
 	}
 	return RemoteCompileService;
 }
diff --git a/Engine/Source/Runtime/OpenGLDrv/Private/OpenGLBinaryProgramUtils.h b/Engine/Source/Runtime/OpenGLDrv/Private/OpenGLBinaryProgramUtils.h
--- a/Engine/Source/Runtime/OpenGLDrv/Private/OpenGLBinaryProgramUtils.h
+++ b/Engine/Source/Runtime/OpenGLDrv/Private/OpenGLBinaryProgramUtils.h
@@ -11,18 +11,13 @@ class FOpenGLProgramBinary
 {
 public:
 	FOpenGLProgramBinary() {};
-	FOpenGLProgramBinary(TArray<uint8>&& ProgramMemory) : DataView(ProgramMemory.GetData(), ProgramMemory.Num()), OwnedData(MoveTemp(ProgramMemory))
+	FOpenGLProgramBinary(TArray<uint8>&& ProgramMemory) : OwnedData(MoveTemp(ProgramMemory)), DataView(OwnedData) 
 	{
 		INC_MEMORY_STAT_BY(STAT_OpenGLShaderLRUProgramMemory, DataView.Num());
 		check(IsOwned());
 		check(IsValid());
 	}
 
-	FOpenGLProgramBinary(FOpenGLProgramBinary&& Src)
-	{
-		*this = MoveTemp(Src);
-	}
-
 	// No ownership, just a view to a region within the mmapped PSO cache.
 	FOpenGLProgramBinary(TArrayView<const uint8> ProgramMemory) : DataView(ProgramMemory) 
 	{
@@ -31,12 +26,9 @@ public:
 		check(IsValid());
 	}
 
-	FOpenGLProgramBinary& operator= (FOpenGLProgramBinary&& rhs)
-	{
-		Swap(DataView, rhs.DataView);
-		Swap(OwnedData, rhs.OwnedData);
-		return *this;
-	}
+	FOpenGLProgramBinary& operator= (FOpenGLProgramBinary&& rhs) = default;
+
+	FOpenGLProgramBinary(FOpenGLProgramBinary&& Src) = default;
 
 	~FOpenGLProgramBinary()
 	{
@@ -55,10 +47,10 @@ public:
 	bool IsOwned() const { return OwnedData.Num() > 0; }
 	const TArrayView<const uint8> GetDataView() const { return DataView; }
 private:
+	// The data when programs are owned.
+	TArray<uint8> OwnedData;
 	// a view to the raw program binary, either points to OwnedData or an array view within the mmapped prebuilt PSO cache.
 	TArrayView<const uint8> DataView;
-	// We may own the data where programs are not 
-	TArray<uint8> OwnedData;
 };
 
 namespace UE
diff --git a/Engine/Source/Runtime/OpenGLDrv/Private/OpenGLCommands.cpp b/Engine/Source/Runtime/OpenGLDrv/Private/OpenGLCommands.cpp
--- a/Engine/Source/Runtime/OpenGLDrv/Private/OpenGLCommands.cpp
+++ b/Engine/Source/Runtime/OpenGLDrv/Private/OpenGLCommands.cpp
@@ -2816,14 +2816,8 @@ void FOpenGLDynamicRHI::RHISetGraphicsPipelineState(FRHIGraphicsPipelineState* G
 
 	if (PsoInit.bFromPSOFileCache)
 	{
-		// If we're from the PSO cache we're just preparing the PSO and do not need to set the state.
-		RHICreateBoundShaderState_internal(
-			PsoInit.BoundShaderState.VertexDeclarationRHI,
-			PsoInit.BoundShaderState.VertexShaderRHI,
-			PsoInit.BoundShaderState.PixelShaderRHI,
-			PsoInit.BoundShaderState.GetGeometryShader(),
-			PsoInit.bFromPSOFileCache
-		);
+		checkNoEntry();
+// 		// If we're from the PSO cache we're just preparing the PSO and do not need to set the state.
 		return;
 	}
 
diff --git a/Engine/Source/Runtime/OpenGLDrv/Private/OpenGLProgramBinaryFileCache.cpp b/Engine/Source/Runtime/OpenGLDrv/Private/OpenGLProgramBinaryFileCache.cpp
--- a/Engine/Source/Runtime/OpenGLDrv/Private/OpenGLProgramBinaryFileCache.cpp
+++ b/Engine/Source/Runtime/OpenGLDrv/Private/OpenGLProgramBinaryFileCache.cpp
@@ -135,6 +135,7 @@ public:
 
 	void AddProgramKey(const class FOpenGLProgramKey& KeyIn) { check(!Content.Contains(KeyIn)); Content.Add(KeyIn); }
 	bool HasValidMapping() const { return MappedRegion.IsValid() && MappedCacheFile.IsValid(); }
+	int32 NumPrograms() const { return Content.Num(); }
 private:
 	TUniquePtr<IMappedFileHandle> MappedCacheFile;
 	TUniquePtr<IMappedFileRegion> MappedRegion;
@@ -149,10 +150,10 @@ static FCriticalSection GPendingGLProgramCreateRequestsCS;
 
 FOpenGLProgramBinaryCache* FOpenGLProgramBinaryCache::CachePtr = nullptr;
 
-FOpenGLProgramBinaryCache::FOpenGLProgramBinaryCache(const FString& InCachePath)
-	: CachePath(InCachePath)
+FOpenGLProgramBinaryCache::FOpenGLProgramBinaryCache(const FString& InCachePathRoot)
+	: CachePathRoot(InCachePathRoot)
 	, BinaryCacheWriteFileHandle(nullptr)
-	, BinaryFileState(EBinaryFileState::Uninitialized)
+	, CurrentBinaryFileState(EBinaryFileState::Uninitialized)
 {
 	ANSICHAR* GLVersion = (ANSICHAR*)glGetString(GL_VERSION);
 	ANSICHAR* GLRenderer = (ANSICHAR*)glGetString(GL_RENDERER);
@@ -178,7 +179,32 @@ FOpenGLProgramBinaryCache::FOpenGLProgramBinaryCache(const FString& InCachePath)
 	FSHAHash VersionHash;
 	FSHA1::HashBuffer(TCHAR_TO_ANSI(*HashString), HashString.Len(), VersionHash.Hash);
 
-	CacheFilename = LegacyShaderPlatformToShaderFormat(GMaxRHIShaderPlatform).ToString() + TEXT("_") + VersionHash.ToString();
+	CacheSubDir = LegacyShaderPlatformToShaderFormat(GMaxRHIShaderPlatform).ToString() + TEXT("_") + VersionHash.ToString();
+
+	// delete anything from the binary program root that does not match the device string.
+	IPlatformFile& PlatformFile = FPlatformFileManager::Get().GetPlatformFile();
+	TArray<FString> FoundFiles;
+	IFileManager::Get().FindFiles(FoundFiles, *(CachePathRoot/TEXT("*")), true, true);
+	for(FString& FoundFile : FoundFiles)
+	{
+		if (FoundFile != CacheSubDir)
+		{
+			const FString FullPath = (CachePathRoot / FoundFile);
+			const bool bIsDir = PlatformFile.DirectoryExists(*FullPath);
+			bool bSuccess = false;
+
+			if(bIsDir)
+			{
+				bSuccess = PlatformFile.DeleteDirectoryRecursively(*FullPath);
+			}
+			else
+			{
+				bSuccess = PlatformFile.DeleteFile(*FullPath);
+			}
+			UE_LOG(LogRHI, Verbose, TEXT("FOpenGLProgramBinaryCache Deleting %s %s"), bIsDir ? TEXT("dir") : TEXT("file"), *FullPath );
+			UE_CLOG(!bSuccess, LogRHI, Warning, TEXT("FOpenGLProgramBinaryCache Failed to delete %s"), *FullPath);
+		}
+	}
 }
 
 FOpenGLProgramBinaryCache::~FOpenGLProgramBinaryCache()
@@ -241,25 +267,25 @@ void FOpenGLProgramBinaryCache::Initialize()
 #endif
 
 
-	FString CacheFolderPath;
+	FString CacheFolderPathRoot;
 #if PLATFORM_ANDROID && USE_ANDROID_FILE
 	// @todo Lumin: Use that GetPathForExternalWrite or something?
 	extern FString GExternalFilePath;
-	CacheFolderPath = GExternalFilePath / TEXT("ProgramBinaryCache");
+	CacheFolderPathRoot = GExternalFilePath / TEXT("ProgramBinaryCache");
 
 #else
-	CacheFolderPath = FPaths::ProjectSavedDir() / TEXT("ProgramBinaryCache");
+	CacheFolderPathRoot = FPaths::ProjectSavedDir() / TEXT("ProgramBinaryCache");
 #endif
 
 	// Remove entire ProgramBinaryCache folder if -ClearOpenGLBinaryProgramCache is specified on command line
 	if (FParse::Param(FCommandLine::Get(), TEXT("ClearOpenGLBinaryProgramCache")))
 	{
-		UE_LOG(LogRHI, Log, TEXT("ClearOpenGLBinaryProgramCache specified, deleting binary program cache folder: %s"), *CacheFolderPath);
-		FPlatformFileManager::Get().GetPlatformFile().DeleteDirectoryRecursively(*CacheFolderPath);
+		UE_LOG(LogRHI, Log, TEXT("ClearOpenGLBinaryProgramCache specified, deleting binary program cache folder: %s"), *CacheFolderPathRoot);
+		FPlatformFileManager::Get().GetPlatformFile().DeleteDirectoryRecursively(*CacheFolderPathRoot);
 	}
 
-	CachePtr = new FOpenGLProgramBinaryCache(CacheFolderPath);
-	UE_LOG(LogRHI, Log, TEXT("Enabling program binary cache as %s"), *CachePtr->GetProgramBinaryCacheFilePath());
+	CachePtr = new FOpenGLProgramBinaryCache(CacheFolderPathRoot);
+	UE_LOG(LogRHI, Log, TEXT("Enabling program binary cache dir at %s"), *CachePtr->GetProgramBinaryCacheDir());
 
 	// Add delegates for the ShaderPipelineCache precompile.
 	UE_LOG(LogRHI, Log, TEXT("FOpenGLProgramBinaryCache will be initialized when ShaderPipelineCache opens its file"));
@@ -282,11 +308,20 @@ static FAutoConsoleVariableRef CVarNumRemoteProgramCompileServices(
 
 void FOpenGLProgramBinaryCache::OnShaderPipelineCacheOpened(FString const& Name, EShaderPlatform Platform, uint32 Count, const FGuid& VersionGuid, FShaderPipelineCache::FShaderCachePrecompileContext& ShaderCachePrecompileContext)
 {
-	UE_LOG(LogRHI, Log, TEXT("Scanning Binary program cache, using Shader Pipeline Cache version %s"), *VersionGuid.ToString());
-
 	FScopeLock Lock(&GProgramBinaryFileCacheCS);
-	check(CurrentShaderPipelineCacheVersionGuid == FGuid());
+	checkf(CurrentShaderPipelineCacheVersionGuid == FGuid(), TEXT("OGL: OnShaderPipelineCacheOpened, previous PSO cache %s (%s) has not completed!"), *CurrentShaderPipelineCacheName, *CurrentShaderPipelineCacheVersionGuid.ToString());
 	CurrentShaderPipelineCacheVersionGuid = VersionGuid;
+	CurrentShaderPipelineCacheName = ShaderCachePrecompileContext.GetCacheName();
+
+	if (Count == 0)
+	{		
+		check(CurrentBinaryFileState == EBinaryFileState::Uninitialized);
+		UE_LOG(LogRHI, Verbose, TEXT("OnShaderPipelineCacheOpened, Ignoring empty PSO cache. %s (%s)"), *CurrentShaderPipelineCacheName, *CurrentShaderPipelineCacheVersionGuid.ToString());
+		return;
+	}
+	 
+	UE_LOG(LogRHI, Log, TEXT("Scanning Binary program cache, using Shader Pipeline Cache %s (%s)"), *CurrentShaderPipelineCacheName, *CurrentShaderPipelineCacheVersionGuid.ToString());
+
 	ScanProgramCacheFile();
 	if (IsBuildingCache_internal())
 	{
@@ -303,47 +338,74 @@ void FOpenGLProgramBinaryCache::OnShaderPipelineCacheOpened(FString const& Name,
 void FOpenGLProgramBinaryCache::Reset()
 {
 	check(!BinaryCacheWriteFileHandle);
-	BinaryFileState = EBinaryFileState::Uninitialized;
- 	ProgramsWrittenToOutputCache.Empty();
+	CurrentBinaryFileState = EBinaryFileState::Uninitialized;
+ 	ProgramsInCurrentCache.Empty();
 }
 
 void FOpenGLProgramBinaryCache::OnShaderPipelineCachePrecompilationComplete(uint32 Count, double Seconds, const FShaderPipelineCache::FShaderCachePrecompileContext& ShaderCachePrecompileContext)
 {
+	QUICK_SCOPE_CYCLE_COUNTER(STAT_OpenGLOnShaderPipelineCachePrecompilationComplete);
 	FScopeLock Lock(&GProgramBinaryFileCacheCS);
 
-#if PLATFORM_ANDROID
-	FAndroidOpenGL::StopRemoteCompileServices();
-#endif
+	// We discard the cache if 0 entries were recorded:
+	// if 0 programs were cached when Count >0, then we suffer a performance penalty for invoking the services for no reason.
+	
+	const bool bIsBuildingCache = IsBuildingCache_internal();
+	TRefCountPtr<FOpenGLProgramBinaryMapping>*FoundCache = MappedCacheFiles.Find(CurrentShaderPipelineCacheVersionGuid);
+	const int32 ProgramsInLoadedCache = FoundCache ? (*FoundCache)->NumPrograms() : 0;
+	const int32 ProgramsCached = bIsBuildingCache ? ProgramsInCurrentCache.Num() : ProgramsInLoadedCache;
 
-	UE_LOG(LogRHI, Log, TEXT("OnShaderPipelineCachePrecompilationComplete: %d shaders"), Count);
+	check(!bIsBuildingCache || Count); // we always start cache building if the count>0
+	check(bIsBuildingCache || Count == 0 || CurrentBinaryFileState == EBinaryFileState::ValidCacheFile);
 
-	const bool bFinishedBuildingCache = IsBuildingCache_internal();
-	check(bFinishedBuildingCache || BinaryFileState == EBinaryFileState::ValidCacheFile);
+	const TCHAR* CacheStatusText = bIsBuildingCache ? 
+		(ProgramsCached == 0 ? TEXT("empty cache discarded") : TEXT("cache built"))
+		:
+		(Count == 0 ? TEXT("ignored empty cache") : TEXT("cache loaded"));
 
-	if (bFinishedBuildingCache)
+	UE_LOG(LogRHI, Log, TEXT("OnShaderPipelineCachePrecompilationComplete: %s(%s) - %s %d program binaries (%d requested)"), *CurrentShaderPipelineCacheName, *CurrentShaderPipelineCacheVersionGuid.ToString(), CacheStatusText, ProgramsCached, Count );
+
+	if (bIsBuildingCache)
 	{
-		CloseWriteHandle();
+#if PLATFORM_ANDROID
+		if (GNumRemoteProgramCompileServices)
+		{
+			FAndroidOpenGL::StopRemoteCompileServices();
+		}
+#endif
 
-		if(!ProgramsWrittenToOutputCache.IsEmpty())
+		const bool bSuccess = CloseCacheWriteHandle();
+
+		if(bSuccess && CVarRestartAndroidAfterPrecompile.GetValueOnAnyThread() == 1)
 		{
 #if PLATFORM_ANDROID
-			if (CVarRestartAndroidAfterPrecompile.GetValueOnAnyThread() == 1)
-			{
 				FAndroidMisc::bNeedsRestartAfterPSOPrecompile = true;
 #if USE_ANDROID_JNI
 				extern void AndroidThunkCpp_RestartApplication(const FString & IntentString);
 				AndroidThunkCpp_RestartApplication(TEXT(""));
 #endif
-			}
 #endif
 		}
-		// Reset state and scan the file we just produced.
+		
 		Reset();
-		ScanProgramCacheFile();
-
-		// unset current cache.
-		CurrentShaderPipelineCacheVersionGuid = FGuid();
+		if(bSuccess)
+		{
+			ScanProgramCacheFile();
+			check(!IsBuildingCache_internal());
+			if (IsBuildingCache_internal())
+			{
+				UE_LOG( LogRHI, Error, TEXT("Failed to load just completed cache! : %s(%s)"), *CurrentShaderPipelineCacheName, *CurrentShaderPipelineCacheVersionGuid.ToString());
+				// The cache we've just written is invalid. This is extremely unlikely.
+				CloseCacheWriteHandle();				
+			}
+		}
 	}
+
+	// unset the completed cache.
+	Reset();
+
+	CurrentShaderPipelineCacheVersionGuid = FGuid();
+	CurrentShaderPipelineCacheName.Reset();
 }
 
 // Scan the binary cache file and build a record of all programs.
@@ -357,7 +419,8 @@ void FOpenGLProgramBinaryCache::ScanProgramCacheFile()
 
 	IPlatformFile& PlatformFile = FPlatformFileManager::Get().GetPlatformFile();
 
-	check(BinaryFileState == EBinaryFileState::Uninitialized);
+	check(CurrentBinaryFileState == EBinaryFileState::Uninitialized);
+	check(ProgramsInCurrentCache.IsEmpty());
 
 	bool bBinaryFileIsValid = false;
 	bool bBinaryFileIsValidAndGuidMatch = false;
@@ -384,7 +447,7 @@ void FOpenGLProgramBinaryCache::ScanProgramCacheFile()
 			{
 				int32 ProgramIndex = 0;
 				FScopeLock Lock(&GPendingGLProgramCreateRequestsCS);
-				ScannedBinaryPrograms.Reserve(BinaryCacheHeader.ProgramCount);
+				PendingGLContainerPrograms.Reserve(BinaryCacheHeader.ProgramCount);
 				if (BinaryCacheHeader.ProgramCount > 0)
 				{
 					TRefCountPtr<FOpenGLProgramBinaryMapping> CurrentMapping;
@@ -407,6 +470,7 @@ void FOpenGLProgramBinaryCache::ScanProgramCacheFile()
 					UE_LOG(LogRHI, Log, TEXT("OnShaderScanProgramCacheFile : %s %s"), CurrentMapping->HasValidMapping() ? TEXT("mapped") : TEXT("opened"), *ProgramCacheFilenameTemp);
 
 					MappedCacheFiles.Add(CurrentShaderPipelineCacheVersionGuid, CurrentMapping);
+					ProgramsInCurrentCache.Reserve(BinaryCacheHeader.ProgramCount);
 					while (!Ar.AtEnd())
 					{
 						FOpenGLProgramKey ProgramKey;
@@ -422,9 +486,11 @@ void FOpenGLProgramBinaryCache::ScanProgramCacheFile()
 
 							UE_LOG(LogRHI, VeryVerbose, TEXT(" scan found PSO %s - %d"), *ProgramKey.ToString(), ProgramBinarySize);
 
+							ProgramsInCurrentCache.Add(ProgramKey);
+
 							if(CurrentMapping->HasValidMapping())
 							{
-								ScannedBinaryPrograms.Emplace(ProgramKey, TUniqueObj<FOpenGLProgramBinary>(CurrentMapping->GetView(ProgramBinaryOffset, ProgramBinarySize)));
+								PendingGLContainerPrograms.Emplace(ProgramKey, TUniqueObj<FOpenGLProgramBinary>(CurrentMapping->GetView(ProgramBinaryOffset, ProgramBinarySize)));
 								Ar.Seek(ProgramBinaryOffset + ProgramBinarySize);
 							}
 							else
@@ -433,7 +499,7 @@ void FOpenGLProgramBinaryCache::ScanProgramCacheFile()
 								TArray<uint8> ProgramBytes;
 								ProgramBytes.SetNumUninitialized(ProgramBinarySize);
 								Ar.Serialize(ProgramBytes.GetData(), ProgramBinarySize);
-								ScannedBinaryPrograms.Emplace(ProgramKey, TUniqueObj<FOpenGLProgramBinary>(MoveTemp(ProgramBytes)));
+								PendingGLContainerPrograms.Emplace(ProgramKey, TUniqueObj<FOpenGLProgramBinary>(MoveTemp(ProgramBytes)));
 							}
 						}
 					}
@@ -444,7 +510,7 @@ void FOpenGLProgramBinaryCache::ScanProgramCacheFile()
 					BinaryProgramReader->Close();
 					// Rename the file back after a successful scan.
 					PlatformFile.MoveFile(*ProgramCacheFilename, *ProgramCacheFilenameTemp);
-					BinaryFileState = EBinaryFileState::ValidCacheFile;
+					CurrentBinaryFileState = EBinaryFileState::ValidCacheFile;
 				}
 				else
 				{
@@ -458,24 +524,24 @@ void FOpenGLProgramBinaryCache::ScanProgramCacheFile()
 			}
 			else
 			{
-				UE_LOG(LogRHI, Log, TEXT("OnShaderScanProgramCacheFile : binary file found but is invalid (%s, %s)"), *CurrentShaderPipelineCacheVersionGuid.ToString(), *BinaryCacheHeader.BinaryCacheGuid.ToString());
+				UE_LOG(LogRHI, Log, TEXT("OnShaderScanProgramCacheFile : binary file found but is invalid (%s, %s), recreating."), *CurrentShaderPipelineCacheVersionGuid.ToString(), *BinaryCacheHeader.BinaryCacheGuid.ToString());
 			}
 		}
 		else
 		{
-			UE_LOG(LogRHI, Log, TEXT("OnShaderScanProgramCacheFile : binary file version invalid"));
+			UE_LOG(LogRHI, Log, TEXT("OnShaderScanProgramCacheFile : binary file version invalid, recreating."));
 		}
 	}
 	else
 	{
-		UE_LOG(LogRHI, Log, TEXT("OnShaderScanProgramCacheFile : Failed to open %s"), *ProgramCacheFilename);
+		UE_LOG(LogRHI, Log, TEXT("OnShaderScanProgramCacheFile : %s was not found, recreating."), *ProgramCacheFilename);
 	}
 
 	if (!bBinaryFileIsValidAndGuidMatch)
 	{
-		if (OpenWriteHandle())
+		if (OpenCacheWriteHandle())
 		{
-			BinaryFileState = EBinaryFileState::BuildingCacheFile;
+			CurrentBinaryFileState = EBinaryFileState::BuildingCacheFile;
 			// save header, 0 program count indicates an unfinished file.
 			// The header is overwritten at the end of the process.
 			UE::OpenGL::FBinaryCacheFileHeader OutHeader = UE::OpenGL::FBinaryCacheFileHeader::CreateHeader(CurrentShaderPipelineCacheVersionGuid, 0);
@@ -485,14 +551,14 @@ void FOpenGLProgramBinaryCache::ScanProgramCacheFile()
 		else
 		{
 			// Binary cache file cannot be used, failed to open output file.
-			BinaryFileState = EBinaryFileState::Uninitialized;
+			CurrentBinaryFileState = EBinaryFileState::Uninitialized;
 			RHIGetPanicDelegate().ExecuteIfBound(FName("FailedBinaryProgramArchiveOpen"));
 			UE_LOG(LogRHI, Fatal, TEXT("ScanProgramCacheFile - Failed to open binary cache."));
 		}
 	}
 }
 
-bool FOpenGLProgramBinaryCache::OpenWriteHandle()
+bool FOpenGLProgramBinaryCache::OpenCacheWriteHandle()
 {
 	check(BinaryCacheWriteFileHandle == nullptr);
 
@@ -501,14 +567,15 @@ bool FOpenGLProgramBinaryCache::OpenWriteHandle()
 
 	BinaryCacheWriteFileHandle = IFileManager::Get().CreateFileWriter(*ProgramCacheFilenameWrite, EFileWrite::FILEWRITE_None);
 
-	UE_CLOG(BinaryCacheWriteFileHandle == nullptr, LogRHI, Warning, TEXT("Failed to open OGL binary cache output file."));
+	UE_CLOG(BinaryCacheWriteFileHandle == nullptr, LogRHI, Warning, TEXT("Failed to open OGL binary cache output file. (%s)"), *ProgramCacheFilenameWrite);
 
 	return BinaryCacheWriteFileHandle != nullptr;
 }
 
-void FOpenGLProgramBinaryCache::CloseWriteHandle()
+bool FOpenGLProgramBinaryCache::CloseCacheWriteHandle()
 {
 	check(BinaryCacheWriteFileHandle != nullptr);
+	const bool bCacheFileIsEmpty = ProgramsInCurrentCache.IsEmpty();
 
 	bool bArchiveFailed = BinaryCacheWriteFileHandle->IsError() || BinaryCacheWriteFileHandle->IsCriticalError();
 
@@ -517,7 +584,7 @@ void FOpenGLProgramBinaryCache::CloseWriteHandle()
 	{
 		FArchive& Ar = *BinaryCacheWriteFileHandle;
 		Ar.Seek(0);
-		UE::OpenGL::FBinaryCacheFileHeader OutHeader = UE::OpenGL::FBinaryCacheFileHeader::CreateHeader(CurrentShaderPipelineCacheVersionGuid, ProgramsWrittenToOutputCache.Num());
+		UE::OpenGL::FBinaryCacheFileHeader OutHeader = UE::OpenGL::FBinaryCacheFileHeader::CreateHeader(CurrentShaderPipelineCacheVersionGuid, ProgramsInCurrentCache.Num());
 		Ar << OutHeader;
 		bArchiveFailed = BinaryCacheWriteFileHandle->IsError() || BinaryCacheWriteFileHandle->IsCriticalError();
 	}
@@ -533,40 +600,21 @@ void FOpenGLProgramBinaryCache::CloseWriteHandle()
 	if (bArchiveFailed)
 	{
 		RHIGetPanicDelegate().ExecuteIfBound(FName("FailedBinaryProgramArchiveWrite"));
-		UE_LOG(LogRHI, Fatal, TEXT("CloseWriteHandle - FArchive error bit set, failed to write binary cache."));
+		UE_LOG(LogRHI, Fatal, TEXT("CloseCacheWriteHandle - FArchive error bit set, failed to write binary cache."));
 	}
 
-	// rename the temp filename back to the final filename
-	PlatformFile.DeleteFile(*ProgramCacheFilename); // file should never exist, but for safety
-	PlatformFile.MoveFile(*ProgramCacheFilename, *ProgramCacheFilenameWrite);
-}
-
-// Called when a new program has been created by OGL RHI, creates the binary cache if it's invalid and then appends the new program details to the file and runtime containers.
-FOpenGLProgramBinary FOpenGLProgramBinaryCache::CreateAndCacheProgramBinary_Internal(const FOpenGLProgramKey& ProgramKey, GLuint Program)
-{
-	FOpenGLProgramBinary CachedProgramBinary;
-	// Add to runtime and disk.
-	if (!ProgramsWrittenToOutputCache.Contains(ProgramKey))
+	if (bCacheFileIsEmpty)
 	{
-		uint32 ProgramBinaryOffset = 0, ProgramBinarySize = 0;
-
-		FOpenGLProgramKey SerializedProgramKey = ProgramKey;
-		CachedProgramBinary = UE::OpenGL::GetProgramBinaryFromGLProgram(Program);
-		if (ensure(CachedProgramBinary.IsValid()))
-		{
-			AddProgramBinaryDataToBinaryCache(ProgramKey, CachedProgramBinary);
-		}
-		else
-		{
-			// we've encountered a problem with this program and there's nothing to write.
-			// This likely means the device will never be able to use this program.
-			// Panic!
-			RHIGetPanicDelegate().ExecuteIfBound(FName("FailedBinaryProgramWrite"));
-			UE_LOG(LogRHI, Fatal, TEXT("AppendProgramBinaryFile Binary program returned 0 bytes!"));
-			// Panic!
-		}
+		// we dont want empty files left on disk.
+		PlatformFile.DeleteFile(*ProgramCacheFilenameWrite);
+	}
+	else
+	{
+		// rename the temp filename back to the final filename
+		PlatformFile.DeleteFile(*ProgramCacheFilename); // file should never exist, but for safety
+		PlatformFile.MoveFile(*ProgramCacheFilename, *ProgramCacheFilenameWrite);
 	}
-	return CachedProgramBinary;
+	return !bCacheFileIsEmpty;
 }
 
 void FOpenGLProgramBinaryCache::CacheProgramBinary(const FOpenGLProgramKey& ProgramKey, const FOpenGLProgramBinary& BinaryProgramData)
@@ -575,7 +623,7 @@ void FOpenGLProgramBinaryCache::CacheProgramBinary(const FOpenGLProgramKey& Prog
 	{
 		FScopeLock Lock(&GProgramBinaryFileCacheCS);
 
-		if (!CachePtr->ProgramsWrittenToOutputCache.Contains(ProgramKey))
+		if (!CachePtr->ProgramsInCurrentCache.Contains(ProgramKey))
 		{
 			CachePtr->AddProgramBinaryDataToBinaryCache(ProgramKey, BinaryProgramData);
 		}
@@ -586,7 +634,7 @@ void FOpenGLProgramBinaryCache::CacheProgramBinary(const FOpenGLProgramKey& Prog
 void FOpenGLProgramBinaryCache::AddProgramBinaryDataToBinaryCache(const FOpenGLProgramKey& ProgramKey, const FOpenGLProgramBinary& BinaryProgramData)
 {
 	check(IsBuildingCache_internal());
-
+	check(BinaryProgramData.IsValid());
 	FArchive& Ar = *BinaryCacheWriteFileHandle;
 	// Serialize to output file:
 	FOpenGLProgramKey SerializedProgramKey = ProgramKey;
@@ -608,7 +656,17 @@ void FOpenGLProgramBinaryCache::AddProgramBinaryDataToBinaryCache(const FOpenGLP
 
 		UE_LOG(LogRHI, Verbose, TEXT("AppendProgramBinaryFile: total Uncompressed: %d, total Compressed %d, Total saved so far: %d"), TotalUncompressed, TotalCompressed, TotalUncompressed - TotalCompressed);
 	}
-	ProgramsWrittenToOutputCache.Add(ProgramKey);
+	UE_LOG(LogRHI, VeryVerbose, TEXT("AddProgramBinaryDataToBinaryCache: written Program %s to cache (%d bytes)"), *ProgramKey.ToString(), BinaryProgramDataView.Num());
+	ProgramsInCurrentCache.Add(ProgramKey);
+}
+
+void FOpenGLProgramBinaryCache::EnqueueBinaryForGLProgramContainer(const FOpenGLProgramKey& ProgramKey, TUniqueObj<FOpenGLProgramBinary>&& ProgramBinary)
+{
+	if (CachePtr)
+	{
+		FScopeLock Lock(&GPendingGLProgramCreateRequestsCS);
+		CachePtr->PendingGLContainerPrograms.Emplace(ProgramKey, MoveTemp(ProgramBinary));
+	}
 }
 
 void FOpenGLProgramBinaryCache::Shutdown()
@@ -620,35 +678,25 @@ void FOpenGLProgramBinaryCache::Shutdown()
 	}
 }
 
-bool FOpenGLProgramBinaryCache::DoesOutputCacheContain(const FOpenGLProgramKey& ProgramKey)
+bool FOpenGLProgramBinaryCache::DoesCurrentCacheContain(const FOpenGLProgramKey& ProgramKey)
 {
 	if (CachePtr)
 	{
 		FScopeLock Lock(&GProgramBinaryFileCacheCS);
-		return CachePtr->DoesOutputCacheContain_Internal(ProgramKey);
+		return CachePtr->DoesCurrentCacheContain_Internal(ProgramKey);
 	}
 	return false;
 }
 
-bool FOpenGLProgramBinaryCache::DoesOutputCacheContain_Internal(const FOpenGLProgramKey& ProgramKey) const
-{
-	return ProgramsWrittenToOutputCache.Contains(ProgramKey);
-}
-
-FOpenGLProgramBinary FOpenGLProgramBinaryCache::CreateAndCacheProgramBinary(const FOpenGLProgramKey& ProgramKey, GLuint Program )
+bool FOpenGLProgramBinaryCache::DoesCurrentCacheContain_Internal(const FOpenGLProgramKey& ProgramKey) const
 {
-	FOpenGLProgramBinary ProgramBinary;
-	if (CachePtr)
-	{
-		FScopeLock Lock(&GProgramBinaryFileCacheCS);
-		ProgramBinary = CachePtr->CreateAndCacheProgramBinary_Internal(ProgramKey, Program);
-	}
-	return ProgramBinary;
+	return ProgramsInCurrentCache.Contains(ProgramKey);
 }
 
 FString FOpenGLProgramBinaryCache::GetProgramBinaryCacheFilePath() const
 {
-	FString ProgramFilename = CachePath + TEXT("/") + CacheFilename + TEXT("_")+CurrentShaderPipelineCacheVersionGuid.ToString();
+	check(CurrentShaderPipelineCacheVersionGuid != FGuid());
+	FString ProgramFilename = CachePathRoot / CacheSubDir / CurrentShaderPipelineCacheName;
 	return ProgramFilename;
 }
 
@@ -656,7 +704,7 @@ void FOpenGLProgramBinaryCache::CheckPendingGLProgramCreateRequests()
 {
 	if (CachePtr)
 	{
-		QUICK_SCOPE_CYCLE_COUNTER(STAT_OpenGLShaderCreateShaderLibRequests);
+		QUICK_SCOPE_CYCLE_COUNTER(STAT_OpenGLCheckPendingGLProgramCreateRequests);
 		check(IsInRenderingThread() || IsInRHIThread());
 		FScopeLock Lock(&GPendingGLProgramCreateRequestsCS);
 		CachePtr->CheckPendingGLProgramCreateRequests_internal();
@@ -667,14 +715,14 @@ void FOpenGLProgramBinaryCache::CheckPendingGLProgramCreateRequests()
 // GMaxBinaryProgramLoadTimeMS attempts to reduce hitching, if we're not using the LRU then we still create GL programs and require more time.
 void FOpenGLProgramBinaryCache::CheckPendingGLProgramCreateRequests_internal()
 {
-	if (ScannedBinaryPrograms.Num() > 0)
+	if (PendingGLContainerPrograms.Num() > 0)
 	{
 		//FScopedDurationTimeLogger Timer(TEXT("CheckPendingGLProgramCreateRequests"));
 		float TimeRemainingS = (float)GMaxBinaryProgramLoadTimeMS / 1000.0f;
 		double StartTime = FPlatformTime::Seconds();
 		int32 Count = 0;
 
-		for (auto It = ScannedBinaryPrograms.CreateIterator(); It && TimeRemainingS > 0.0f; ++It)
+		for (auto It = PendingGLContainerPrograms.CreateIterator(); It && TimeRemainingS > 0.0f; ++It)
 		{
 			UE::OpenGL::OnGLProgramLoadedFromBinaryCache(It->Key, MoveTemp(It->Value));
 			TimeRemainingS -= (float)(FPlatformTime::Seconds() - StartTime);
@@ -683,13 +731,13 @@ void FOpenGLProgramBinaryCache::CheckPendingGLProgramCreateRequests_internal()
 			It.RemoveCurrent();
 		}
 		float TimeTaken = (float)GMaxBinaryProgramLoadTimeMS - (TimeRemainingS * 1000.0f);
-		if (TimeRemainingS < 0.0f)
+		if (TimeRemainingS < 0.005f)
 		{
-			UE_LOG(LogRHI, Warning, TEXT("CheckPendingGLProgramCreateRequests : iter count = %d, time taken = %f ms (remaining %d)"), Count, TimeTaken, ScannedBinaryPrograms.Num());
+			UE_LOG(LogRHI, Warning, TEXT("CheckPendingGLProgramCreateRequests : iter count = %d, time taken = %f ms (remaining %d)"), Count, TimeTaken, PendingGLContainerPrograms.Num());
 		}
 		else
 		{
-			UE_LOG(LogRHI, Verbose, TEXT("CheckPendingGLProgramCreateRequests : iter count = %d, time taken = %f ms (remaining %d)"), Count, TimeTaken, ScannedBinaryPrograms.Num());
+			UE_LOG(LogRHI, Verbose, TEXT("CheckPendingGLProgramCreateRequests : iter count = %d, time taken = %f ms (remaining %d)"), Count, TimeTaken, PendingGLContainerPrograms.Num());
 		}
 	}
 }
@@ -709,7 +757,7 @@ bool FOpenGLProgramBinaryCache::CheckSinglePendingGLProgramCreateRequest(const F
 bool FOpenGLProgramBinaryCache::CheckSinglePendingGLProgramCreateRequest_internal(const FOpenGLProgramKey& ProgramKey)
 {
 	TUniqueObj<FOpenGLProgramBinary> ProgFound;
-	if (ScannedBinaryPrograms.RemoveAndCopyValue(ProgramKey, ProgFound))
+	if (PendingGLContainerPrograms.RemoveAndCopyValue(ProgramKey, ProgFound))
 	{		
 		UE::OpenGL::OnGLProgramLoadedFromBinaryCache(ProgramKey, MoveTemp(ProgFound));
 		return true;
diff --git a/Engine/Source/Runtime/OpenGLDrv/Private/OpenGLProgramBinaryFileCache.h b/Engine/Source/Runtime/OpenGLDrv/Private/OpenGLProgramBinaryFileCache.h
--- a/Engine/Source/Runtime/OpenGLDrv/Private/OpenGLProgramBinaryFileCache.h
+++ b/Engine/Source/Runtime/OpenGLDrv/Private/OpenGLProgramBinaryFileCache.h
@@ -18,7 +18,6 @@
 #include "ShaderCore.h"
 #include "CrossCompilerCommon.h"
 #include "ShaderCodeLibrary.h"
-#include "Async/AsyncFileHandle.h"
 #include "ShaderPipelineCache.h"
 
 class FOpenGLLinkedProgram;
@@ -34,11 +33,8 @@ public:
 
 	static bool IsEnabled();
 
-	/** Extract the program binary from GL and store on disk. */
-	static FOpenGLProgramBinary CreateAndCacheProgramBinary(const FOpenGLProgramKey& ProgramKey, GLuint Program);
-
-	/** Is the program already in the currently building cache */
-	static bool DoesOutputCacheContain(const FOpenGLProgramKey& ProgramKey);
+	/** Is the program already in the currently processing cache */
+	static bool DoesCurrentCacheContain(const FOpenGLProgramKey& ProgramKey);
 
 	/** Take an existing program binary and store on disk. Only when ProgramBinaryCache is enabled
 	*/
@@ -53,6 +49,9 @@ public:
 	/** true if the program binary cache is currently in cache build mode */
 	static bool IsBuildingCache();
 
+	/** Add a binary program to the GL container queue, it is moved at a later point to the GL program container */
+	static void EnqueueBinaryForGLProgramContainer(const FOpenGLProgramKey& ProgramKey, TUniqueObj<FOpenGLProgramBinary>&& ProgramBinary);
+
 private:
 	FOpenGLProgramBinaryCache(const FString& InCachePath);
 	~FOpenGLProgramBinaryCache();
@@ -71,14 +70,19 @@ private:
 	static TAutoConsoleVariable<int32> CVarRestartAndroidAfterPrecompile;
 	static FOpenGLProgramBinaryCache* CachePtr;
 
-	/*  Path to directory where binary programs will be stored excluding the cache filename */
-	FString CachePath;
-	/* Just the cache filename, without the path */
-	FString CacheFilename;
+	// Binary program file layout is <osfilelocation>/CachePathRoot/CacheSubDir/CurrentShaderPipelineCacheName
+
+	/*  Path to root directory where the binary program subdir will be created */
+	FString CachePathRoot;
+	/* Subdir is a hash of the device's version info, anything within CachePathRoot that does not == CacheSubDir is deleted. */
+	FString CacheSubDir;
 
 	// Guid of the PSO cache being procesed.
 	FGuid CurrentShaderPipelineCacheVersionGuid;
 
+	// String name of the current cache.
+	FString CurrentShaderPipelineCacheName;
+
 	/**
 	* Shaders that were requested for compilation
 	* They will be compiled just before linking a program only in case when there is no saved binary program
@@ -92,14 +96,12 @@ private:
 
 	void ScanProgramCacheFile();
 
-	void OpenAsyncReadHandle();
-	void CloseAsyncReadHandle();
+	bool OpenCacheWriteHandle();
 
-	bool OpenWriteHandle();
-	void CloseWriteHandle();
+	// finalize the current cache file, returns true if a valid file was created.
+	bool CloseCacheWriteHandle(); 
 
-	FOpenGLProgramBinary CreateAndCacheProgramBinary_Internal(const FOpenGLProgramKey& ProgramKey, GLuint Program);
-	bool DoesOutputCacheContain_Internal(const FOpenGLProgramKey& ProgramKey) const;
+	bool DoesCurrentCacheContain_Internal(const FOpenGLProgramKey& ProgramKey) const;
 
 	void AddProgramBinaryDataToBinaryCache(const FOpenGLProgramKey& ProgramKey, const FOpenGLProgramBinary& BinaryProgramData);
 
@@ -113,15 +115,16 @@ private:
 	void OnShaderPipelineCachePrecompilationComplete(uint32 Count, double Seconds, const FShaderPipelineCache::FShaderCachePrecompileContext& ShaderCachePrecompileContext);
 
 private:
+	FString GetProgramBinaryCacheDir() const { return (CachePathRoot / CacheSubDir);	}
 
 	FDelegateHandle OnShaderPipelineCacheOpenedDelegate;
 	FDelegateHandle OnShaderPipelineCachePrecompilationCompleteDelegate;
 
-	TSet<FOpenGLProgramKey> ProgramsWrittenToOutputCache; // Only populated when creating the cache. Used to prevents duplicates.
+	TSet<FOpenGLProgramKey> ProgramsInCurrentCache; // Only populated during PSO precompile delegates.
 
-	// All of the mmapped binary programs encountered during the cache open scan.
+	// All of the binary programs encountered during PSO cache processing.
 	// The contents are moved to OpenGL's program cache during RHI's end of frame tick.
-	TMap<FOpenGLProgramKey, TUniqueObj<FOpenGLProgramBinary>> ScannedBinaryPrograms;
+	TMap<FOpenGLProgramKey, TUniqueObj<FOpenGLProgramBinary>> PendingGLContainerPrograms;
 
 	FArchive* BinaryCacheWriteFileHandle;
 	bool bShownLoadingScreen;
@@ -133,9 +136,9 @@ private:
 		ValidCacheFile,					// We have a valid cache file we can use for reading.
 	};
 
-	bool IsBuildingCache_internal() const	{ return BinaryFileState == EBinaryFileState::BuildingCacheFile; }
+	bool IsBuildingCache_internal() const	{ return CurrentBinaryFileState == EBinaryFileState::BuildingCacheFile; }
 
-	EBinaryFileState BinaryFileState;
+	EBinaryFileState CurrentBinaryFileState;
 
 	// Container of all currently open program cache files.
 	TMap<FGuid, TRefCountPtr<FOpenGLProgramBinaryMapping>> MappedCacheFiles;
diff --git a/Engine/Source/Runtime/OpenGLDrv/Private/OpenGLShaders.cpp b/Engine/Source/Runtime/OpenGLDrv/Private/OpenGLShaders.cpp
--- a/Engine/Source/Runtime/OpenGLDrv/Private/OpenGLShaders.cpp
+++ b/Engine/Source/Runtime/OpenGLDrv/Private/OpenGLShaders.cpp
@@ -1312,13 +1312,14 @@ public:
 	void DeleteGLResources()
 	{
 		VERIFY_GL_SCOPE();
-		SetDeletedProgramStats(Program);
-		FOpenGL::DeleteProgramPipelines(1, &Program);
-
-		GetOpenGLUniformBlockLocations().Remove(Program);
-		GetOpenGLUniformBlockBindings().Remove(Program);
-
-		Program = 0;
+		if (Program != 0)
+		{
+			SetDeletedProgramStats(Program);
+			FOpenGL::DeleteProgramPipelines(1, &Program);
+			GetOpenGLUniformBlockLocations().Remove(Program);
+			GetOpenGLUniformBlockBindings().Remove(Program);
+			Program = 0;
+		}
 
 		for (int Stage = 0; Stage < CrossCompiler::NUM_SHADER_STAGES; Stage++)
 		{
@@ -1825,27 +1826,28 @@ class FGLProgramCache
 {
 	FGLProgramCacheLRU ProgramCacheLRU;
 	FOpenGLProgramsMap ProgramCache;
-	bool bUseLRUCache;
+	inline static uint32 UseLRUCacheStatus = -1;
 public:
-	FGLProgramCache()
+
+	static FORCEINLINE_DEBUGGABLE bool IsUsingLRU()
 	{
-		if(CVarEnableLRU.GetValueOnAnyThread() && !FOpenGL::SupportsProgramBinary())
+		if (UseLRUCacheStatus == -1)
 		{
-			UE_LOG(LogRHI, Warning, TEXT("Requesting OpenGL program LRU cache, but program binary is not supported by driver. Falling back to non-lru cache."));
-		}
-
-		bUseLRUCache = CVarEnableLRU.GetValueOnAnyThread() == 1 && FOpenGL::SupportsProgramBinary();
-		UE_LOG(LogRHI, Log, TEXT("Using OpenGL program LRU cache: %d"), bUseLRUCache ? 1 : 0);
-	}
+			if (CVarEnableLRU.GetValueOnAnyThread() && !FOpenGL::SupportsProgramBinary())
+			{
+				UE_LOG(LogRHI, Warning, TEXT("Requesting OpenGL program LRU cache, but program binary is not supported by driver. Falling back to non-lru cache."));
+			}
 
-	FORCEINLINE_DEBUGGABLE bool IsUsingLRU() const
-	{
-		return bUseLRUCache;
+			UseLRUCacheStatus = CVarEnableLRU.GetValueOnAnyThread() == 1 && FOpenGLProgramBinaryCache::IsEnabled();
+			UE_LOG(LogRHI, Log, TEXT("OpenGL program LRU cache active = %d (%d, %d)"), UseLRUCacheStatus, CVarEnableLRU.GetValueOnAnyThread(), FOpenGLProgramBinaryCache::IsEnabled());
+		}
+		check(UseLRUCacheStatus != -1);
+		return UseLRUCacheStatus == 1;
 	}
 
 	FORCEINLINE_DEBUGGABLE void Touch(FOpenGLLinkedProgram* LinkedProgram)
 	{
-		if (bUseLRUCache)
+		if (IsUsingLRU())
 		{
 			ProgramCacheLRU.Touch(LinkedProgram);
 		}
@@ -1853,7 +1855,7 @@ public:
 
 	FORCEINLINE_DEBUGGABLE FOpenGLLinkedProgram* Find(const FOpenGLProgramKey& ProgramKey, bool bFindAndCreateEvictedProgram)
 	{
-		if (bUseLRUCache)
+		if (IsUsingLRU())
 		{
 			return ProgramCacheLRU.Find(ProgramKey, bFindAndCreateEvictedProgram);
 		}
@@ -1866,7 +1868,7 @@ public:
 
 	FORCEINLINE_DEBUGGABLE bool Contains(const FOpenGLProgramKey& ProgramKey) const
 	{
-		if (bUseLRUCache)
+		if (IsUsingLRU())
 		{
 			return ProgramCacheLRU.Contains(ProgramKey);
 		}
@@ -1879,7 +1881,7 @@ public:
 
 	FORCEINLINE_DEBUGGABLE void Add(const FOpenGLProgramKey& ProgramKey, FOpenGLLinkedProgram* LinkedProgram)
 	{
-		if (bUseLRUCache)
+		if (IsUsingLRU())
 		{
 			ProgramCacheLRU.Add(ProgramKey, LinkedProgram);
 		}
@@ -1892,7 +1894,7 @@ public:
 
 	void Empty()
 	{
-		if (bUseLRUCache)
+		if (IsUsingLRU())
 		{
 			ProgramCacheLRU.Empty();
 		}
@@ -1909,7 +1911,7 @@ public:
 
 	bool IsLRUAtCapacity() const
 	{
-		if (bUseLRUCache)
+		if (IsUsingLRU())
 		{
 			ProgramCacheLRU.IsLRUAtCapacity();
 		}
@@ -1941,6 +1943,13 @@ public:
 		ProgramCacheLRU.EvictProgramFromLRU(ProgramKey);
 	}
 
+	void AddEvicted(const FOpenGLProgramKey& ProgramKey, TUniqueObj<FOpenGLProgramBinary>&& ProgramBinary)
+	{
+		check(IsUsingLRU());
+		check(!Contains(ProgramKey));
+		AddOrReplaceEvicted(ProgramKey, MoveTemp(ProgramBinary));
+	}
+
 	void AddOrReplaceEvicted(const FOpenGLProgramKey& ProgramKey, TUniqueObj<FOpenGLProgramBinary>&& ProgramBinary)
 	{
 		check(IsUsingLRU());
@@ -1961,7 +1970,7 @@ public:
 
 	void EnumerateLinkedPrograms(TFunction<void(FOpenGLLinkedProgram*)> EnumFunc)
 	{
-		if (bUseLRUCache)
+		if (IsUsingLRU())
 		{
 			ProgramCacheLRU.EnumerateLinkedPrograms(EnumFunc);
 		}
@@ -1978,6 +1987,7 @@ public:
 
 static FGLProgramCache& GetOpenGLProgramsCache()
 {
+	check(IsInRenderingThread() || IsInRHIThread());
 	static FGLProgramCache ProgramsCache;
 	return ProgramsCache;
 }
@@ -2489,8 +2499,13 @@ static void VerifyUniformBufferLayouts(GLuint Program)
 /**
  * Link vertex and pixel shaders in to an OpenGL program.
  */
-static FOpenGLLinkedProgram* LinkProgram( const FOpenGLLinkedProgramConfiguration& Config, bool bFromPSOFileCache)
+FOpenGLLinkedProgram* FOpenGLDynamicRHI::LinkProgram( const FOpenGLLinkedProgramConfiguration& Config)
 {
+	// Make sure we have OpenGL context set up, and invalidate the parameters cache and current program (as we'll link a new one soon)
+	GetContextStateForCurrentContext().Program = -1;
+	MarkShaderParameterCachesDirty(PendingState.ShaderParameters, false);
+	PendingState.LinkedProgramAndDirtyFlag = nullptr;
+
 	ANSICHAR Buf[32] = {0};
 
 	SCOPE_CYCLE_COUNTER(STAT_OpenGLShaderLinkTime);
@@ -2503,9 +2518,6 @@ static FOpenGLLinkedProgram* LinkProgram( const FOpenGLLinkedProgramConfiguratio
 	GLuint Program = 0;
 	if (FOpenGLProgramBinaryCache::IsEnabled())
 	{
-		UE_CLOG(bFromPSOFileCache && !FOpenGLProgramBinaryCache::IsBuildingCache(), LogRHI, Error, TEXT("precompile PSO %s should be precompiled but is requesting compile."), *Config.ProgramKey.ToString());
-		// We should only see PSO cache programs here during cache creation. Otherwise the PSO should have been found already from the program cache.
-		//check(!bFromPSOFileCache || FOpenGLProgramBinaryCache::IsBuildingCache());
 		UE::OpenGL::CompilePendingShaders(Config);
 	}
 
@@ -2532,7 +2544,7 @@ static FOpenGLLinkedProgram* LinkProgram( const FOpenGLLinkedProgramConfiguratio
 		FOpenGL::UseProgramStages(Program, GL_COMPUTE_SHADER_BIT, Config.Shaders[CrossCompiler::SHADER_STAGE_COMPUTE].Resource);
 	}
 	
-	if(FOpenGLProgramBinaryCache::IsEnabled() || GetOpenGLProgramsCache().IsUsingLRU())
+	if(FOpenGLProgramBinaryCache::IsEnabled() || FGLProgramCache::IsUsingLRU())
 	{
 		FOpenGL::ProgramParameter(Program, PROGRAM_BINARY_RETRIEVABLE_HINT, GL_TRUE);
 	}
@@ -2551,18 +2563,6 @@ static FOpenGLLinkedProgram* LinkProgram( const FOpenGLLinkedProgramConfiguratio
 
 	FOpenGLLinkedProgram* LinkedProgram = new FOpenGLLinkedProgram(Config, Program);
 
-	if (FOpenGLProgramBinaryCache::IsEnabled() && FOpenGLProgramBinaryCache::IsBuildingCache() && GetOpenGLProgramsCache().IsUsingLRU() && bFromPSOFileCache)
-	{
-		// we need the binary if this is a pso for the currently building cache.
-		FOpenGLProgramBinary ProgramBinary = FOpenGLProgramBinaryCache::CreateAndCacheProgramBinary(Config.ProgramKey, Program);
-		check(ProgramBinary.IsValid());
-
-		if(CVarLRUKeepProgramBinaryResident.GetValueOnAnyThread())
-		{
-			// We've gone to the effort of creating the binary so store it in LRUInfo, avoids re-requesting the binary if the program is evicted.
-			LinkedProgram->LRUInfo.CachedProgramBinary = TUniqueObj<FOpenGLProgramBinary>(MoveTemp(ProgramBinary));
-		}
-	}
 	ConfigureStageStates(LinkedProgram);
 
 #if ENABLE_UNIFORM_BUFFER_LAYOUT_VERIFICATION
@@ -2571,7 +2571,7 @@ static FOpenGLLinkedProgram* LinkProgram( const FOpenGLLinkedProgramConfiguratio
 	return LinkedProgram;
 }
 
-static bool LinkComputeShader(FRHIComputeShader* ComputeShaderRHI, FOpenGLComputeShader* ComputeShader)
+bool FOpenGLDynamicRHI::LinkComputeShader(FRHIComputeShader* ComputeShaderRHI, FOpenGLComputeShader* ComputeShader)
 {
 	check(ComputeShader);
 	check(ComputeShader->Resource != 0);
@@ -2587,7 +2587,7 @@ static bool LinkComputeShader(FRHIComputeShader* ComputeShaderRHI, FOpenGLComput
 
 	if (ComputeShader->LinkedProgram == nullptr)
 	{
-		ComputeShader->LinkedProgram = LinkProgram(Config, false);
+		ComputeShader->LinkedProgram = LinkProgram(Config);
 		if(ComputeShader->LinkedProgram == nullptr)
 		{
 		#if DEBUG_GL_SHADERS
@@ -2715,9 +2715,6 @@ static void BindShaderStage(FOpenGLLinkedProgramConfiguration& Config, CrossComp
 }
 
 // ============================================================================================================================
-static FCriticalSection GProgramBinaryCacheCS;
-
-
 
 static FOpenGLLinkedProgramConfiguration CreateConfig(FRHIVertexShader* VertexShaderRHI, FRHIPixelShader* PixelShaderRHI, FRHIGeometryShader* GeometryShaderRHI )
 {
@@ -2772,99 +2769,144 @@ static bool CanCreateExternally(bool bIsFromPSO)
 	return false;
 }
 
-void FOpenGLDynamicRHI::PrepareGFXBoundShaderState(const FGraphicsPipelineStateInitializer& Initializer)
+static FOpenGLProgramBinary ExternalProgramCompile(const FOpenGLProgramKey& ProgramKey, const FRHIVertexShader* VertexShaderRHI, const FRHIPixelShader* PixelShaderRHI)
 {
+	FOpenGLProgramBinary CompiledProgram;
 #if PLATFORM_ANDROID
-	if (CanCreateExternally(Initializer.bFromPSOFileCache))
+	// compile externally, sit and wait for the linked result	
+	const FOpenGLCompiledShaderKey& VSKey = FOpenGLDynamicRHI::ResourceCastProxy(VertexShaderRHI)->GetCompiledShaderKey();
+	const FOpenGLCompiledShaderKey& PSKey = FOpenGLDynamicRHI::ResourceCastProxy(PixelShaderRHI)->GetCompiledShaderKey();
+
+	TArray<ANSICHAR> VSCode;
+	TArray<ANSICHAR> PSCode;
+	TArray<ANSICHAR> ComputeGlslCode;
 	{
-		FRHIVertexShader* VertexShaderRHI = Initializer.BoundShaderState.GetVertexShader();
-		FRHIPixelShader* PixelShaderRHI = Initializer.BoundShaderState.GetPixelShader();
+		FScopeLock Lock(&GCompiledShaderCacheCS);
+		VSCode = GetOpenGLCompiledShaderCache().FindRef(VSKey)->GetUncompressedShader();
+		PSCode = GetOpenGLCompiledShaderCache().FindRef(PSKey)->GetUncompressedShader();
+	}
 
-		if (!PixelShaderRHI)
+	FString FailLog;
+	TArray<uint8> CompiledProgramBytes = FAndroidOpenGL::DispatchAndWaitForRemoteGLProgramCompile(TArrayView<uint8>((uint8*)&ProgramKey, sizeof(ProgramKey)), VSCode, PSCode, ComputeGlslCode, FailLog);
+
+	if (FailLog.IsEmpty())
+	{
+		GLenum glFormat = *(GLenum*)CompiledProgramBytes.GetData();
+		if (UE::OpenGL::IsStoringCompressedBinaryPrograms())
 		{
-			// use special null pixel shader when PixelShader was set to NULL
-			PixelShaderRHI = TShaderMapRef<FNULLPS>(GetGlobalShaderMap(GMaxRHIFeatureLevel)).GetPixelShader();
+			TArray<uint8> CompressedCompiledProgramResult;
+			UE::OpenGL::CompressProgramBinary(CompiledProgramBytes, CompressedCompiledProgramResult);
+			CompiledProgramBytes = MoveTemp(CompressedCompiledProgramResult);
 		}
 
-		FOpenGLProgramKey ProgramKey;
-		ProgramKey.ShaderHashes[CrossCompiler::SHADER_STAGE_VERTEX] = VertexShaderRHI->GetHash();
-		ProgramKey.ShaderHashes[CrossCompiler::SHADER_STAGE_PIXEL] = PixelShaderRHI->GetHash();
-
-		// add it to the runtime container. (this prevents createboundshader state from adding it to the binary cache.)
+		CompiledProgram = FOpenGLProgramBinary(MoveTemp(CompiledProgramBytes));
+	}
+	else
+	{
+		UE_LOG(LogRHI, Error, TEXT("External compile of program %s failed: %s "), *ProgramKey.ToString(), *FailLog);
+#if DEBUG_GL_SHADERS
+		if (VSCode.Num())
 		{
-			check(GetOpenGLProgramsCache().IsUsingLRU());
-			if (FOpenGLProgramBinaryCache::DoesOutputCacheContain(ProgramKey))
-			{
-				//already done.
-				UE_LOG(LogRHI, Verbose, TEXT("PrepareGFXBoundShaderState Already in output PSO cache. %s is pso %d "), *ProgramKey.ToString(), Initializer.bFromPSOFileCache ? 1 : 0);
-				return;
-			}
+			UE_LOG(LogRHI, Error, TEXT("Vertex Shader:\n%s"), ANSI_TO_TCHAR(VSCode.GetData()));
 		}
-
-		// compile externally, sit and wait for the linked result	
-		const FOpenGLCompiledShaderKey& VSKey = ResourceCastProxy(VertexShaderRHI)->GetCompiledShaderKey();
-		const FOpenGLCompiledShaderKey& PSKey = ResourceCastProxy(PixelShaderRHI)->GetCompiledShaderKey();
-
-		TArray<ANSICHAR> VSCode;
-		TArray<ANSICHAR> PSCode;
-		TArray<ANSICHAR> ComputeGlslCode; 
+		if (PSCode.Num())
 		{
-			FScopeLock Lock(&GCompiledShaderCacheCS);
-			VSCode = GetOpenGLCompiledShaderCache().FindRef(VSKey)->GetUncompressedShader();
-			PSCode = GetOpenGLCompiledShaderCache().FindRef(PSKey)->GetUncompressedShader();
+			UE_LOG(LogRHI, Error, TEXT("Pixel Shader:\n%s"), ANSI_TO_TCHAR(PSCode.GetData()));
 		}
+#endif //DEBUG_GL_SHADERS
+	}
+#else
+	checkNoEntry();
+#endif
+	return CompiledProgram;
+}
 
-		FString FailLog;
-		TArray<uint8> CompiledProgramBytes = FAndroidOpenGL::DispatchAndWaitForRemoteGLProgramCompile(TArrayView<uint8>((uint8*)&ProgramKey, sizeof(ProgramKey)), VSCode, PSCode, ComputeGlslCode, FailLog);
+void FOpenGLDynamicRHI::PrepareGFXBoundShaderState(const FGraphicsPipelineStateInitializer& Initializer)
+{
+	if (!Initializer.bFromPSOFileCache || !FOpenGLProgramBinaryCache::IsEnabled())
+	{
+		return;
+	}
 
-		if(FailLog.IsEmpty())
-		{
-			GLenum glFormat = *(GLenum*)CompiledProgramBytes.GetData();		
-			if (UE::OpenGL::IsStoringCompressedBinaryPrograms())
-			{
-				
-				TArray<uint8> CompressedCompiledProgramResult;
-				UE::OpenGL::CompressProgramBinary(CompiledProgramBytes, CompressedCompiledProgramResult);
-				CompiledProgramBytes = MoveTemp(CompressedCompiledProgramResult);
-			}
+	FRHIVertexShader* VertexShaderRHI = Initializer.BoundShaderState.GetVertexShader();
+	FRHIPixelShader* PixelShaderRHI = Initializer.BoundShaderState.GetPixelShader();
+	FRHIGeometryShader* GeometryShaderRHI = Initializer.BoundShaderState.GetGeometryShader();
 
-			FOpenGLProgramBinary CompiledProgram(MoveTemp(CompiledProgramBytes));
+	if (!PixelShaderRHI)
+	{
+		// use special null pixel shader when PixelShader was set to NULL
+		PixelShaderRHI = TShaderMapRef<FNULLPS>(GetGlobalShaderMap(GMaxRHIFeatureLevel)).GetPixelShader();
+	}
 
-			// add it to the binary file.
-			FOpenGLProgramBinaryCache::CacheProgramBinary(ProgramKey, CompiledProgram);
+	FOpenGLProgramKey ProgramKey;
+	ProgramKey.ShaderHashes[CrossCompiler::SHADER_STAGE_VERTEX] = VertexShaderRHI->GetHash();
+	ProgramKey.ShaderHashes[CrossCompiler::SHADER_STAGE_PIXEL] = PixelShaderRHI->GetHash();
+	bool bCreateProgram = false;
+	bool bCreateBinary = false;
 
-			// add it to the runtime container. (this prevents createboundshader state from adding it to the binary cache.)
+	if (!FOpenGLProgramBinaryCache::DoesCurrentCacheContain(ProgramKey))
+	{
+		if (FOpenGLProgramBinaryCache::IsBuildingCache())
+		{
+			FOpenGLProgramBinary CompiledProgram;
+			if (CanCreateExternally(Initializer.bFromPSOFileCache))
 			{
-				FScopeLock Lock(&GProgramBinaryCacheCS);
-				check(GetOpenGLProgramsCache().IsUsingLRU());
-				if (GetOpenGLProgramsCache().Find(ProgramKey, false) == nullptr)
-				{
-					GetOpenGLProgramsCache().AddOrReplaceEvicted(ProgramKey, TUniqueObj<FOpenGLProgramBinary>(MoveTemp(CompiledProgram)));
-				}
-				else
+				CompiledProgram = ExternalProgramCompile(ProgramKey, VertexShaderRHI, PixelShaderRHI);
+			}
+			else
+			{
+				RunOnGLRenderContextThread([&CompiledProgram, &VertexShaderRHI, &PixelShaderRHI, &GeometryShaderRHI,this]()
 				{
-					// beaten to it elsewhere.
-					UE_LOG(LogRHI, Verbose, TEXT("PrepareGFXBoundShaderState skipped add. %s is pso %d "), *ProgramKey.ToString(), Initializer.bFromPSOFileCache ? 1 : 0);
-				}
+					check(IsInRenderingThread() || IsInRHIThread());
+					VERIFY_GL_SCOPE();
+					FOpenGLLinkedProgramConfiguration Config = CreateConfig(VertexShaderRHI,
+						PixelShaderRHI,
+						GeometryShaderRHI);
+
+					FOpenGLLinkedProgram* LinkedProgram = LinkProgram(Config);
+					if(LinkedProgram == nullptr)
+					{
+						RHIGetPanicDelegate().ExecuteIfBound(FName("FailedProgramLinkDuringPrecompile"));
+						UE_LOG(LogRHI, Fatal, TEXT("Failed to link program [%s]. Current total programs: %d"), *Config.ProgramKey.ToString(), GNumPrograms);
+						return;
+					}
+
+					CompiledProgram = UE::OpenGL::GetProgramBinaryFromGLProgram(LinkedProgram->Program);
+					// optional: enqueue this binary, it will be picked up during createboundshaderstate.
+					// If this is enabled then the binary program will be sent to the GL container.
+					// not doing this for now as the completed binary cache will be loaded as normal at the end of the pre-caching process.
+					// FOpenGLProgramBinaryCache::EnqueueBinaryForGLProgramContainer(ProgramKey, TUniqueObj<FOpenGLProgramBinary>(MoveTemp(CompiledProgram)));
+
+					if (!FGLProgramCache::IsUsingLRU())
+					{
+						GetOpenGLProgramsCache().Add(Config.ProgramKey, LinkedProgram);	// if we're not using the LRU then we add the program as normal.
+					}
+					else
+					{
+						// with LRU mode the program is made available as an evicted program after the PSO cache completes.
+						delete LinkedProgram;
+						LinkedProgram = nullptr;
+					}
+				}, true); // TODO: this wait can be optimized away
 			}
 
-		}
-		else
-		{
-			UE_LOG(LogRHI, Error, TEXT("External compile of program %s failed: %s "), *ProgramKey.ToString(), *FailLog);
-#if DEBUG_GL_SHADERS
-			if (VSCode.Num())
+			if (CompiledProgram.IsValid())
 			{
-				UE_LOG(LogRHI, Error, TEXT("Vertex Shader:\n%s"), ANSI_TO_TCHAR(VSCode.GetData()));
+				FOpenGLProgramBinaryCache::CacheProgramBinary(ProgramKey, CompiledProgram);
 			}
-			if (PSCode.Num())
+			else
 			{
-				UE_LOG(LogRHI, Error, TEXT("Pixel Shader:\n%s"), ANSI_TO_TCHAR(PSCode.GetData()));
+				UE_LOG(LogRHI, Warning, TEXT("Program binary generation failed (%s), omitted from binary cache."), *ProgramKey.ToString());
 			}
-#endif //DEBUG_GL_SHADERS
+		}
+		else
+		{
+			// This PSO was not present when the binary cache was created.
+			// Do nothing and hitch when used for rendering.
+			// Mitigation could be to compile on the RHIT or append to the existing cache?
+			UE_LOG(LogRHI, Warning, TEXT("New PSO encountered during precompile %s"), *ProgramKey.ToString());
 		}
 	}
-#endif
 }
 
 FBoundShaderStateRHIRef FOpenGLDynamicRHI::RHICreateBoundShaderState_OnThisThread(
@@ -2877,12 +2919,10 @@ FBoundShaderStateRHIRef FOpenGLDynamicRHI::RHICreateBoundShaderState_OnThisThrea
 {
 	check(IsInRenderingThread() || IsInRHIThread());
 
-	FScopeLock Lock(&GProgramBinaryCacheCS);
-
 	VERIFY_GL_SCOPE();
 
 	SCOPE_CYCLE_COUNTER(STAT_OpenGLCreateBoundShaderStateTime);
-
+	check(!bFromPSOFileCache);
 	if (!PixelShaderRHI)
 	{
 		// use special null pixel shader when PixelShader was set to NULL
@@ -2937,31 +2977,13 @@ FBoundShaderStateRHIRef FOpenGLDynamicRHI::RHICreateBoundShaderState_OnThisThrea
 
 		if (!LinkedProgram)
 		{
-			bool bCreateEvictedBinary = false;
-			// If this is this a request from the PSOFC then do not create an evicted program.
-			if (bFromPSOFileCache && GetOpenGLProgramsCache().IsUsingLRU())
-			{
-				// no-op, No need to create a precompile PSO that we already have.
-				// return null as nothing should be using the resultant BSS of a precompile PSO either.
-				bool bContained = GetOpenGLProgramsCache().Contains(Config.ProgramKey) || FOpenGLProgramBinaryCache::CheckSinglePendingGLProgramCreateRequest(Config.ProgramKey);
-				if (bContained)
-				{
-					return nullptr;
-				}
-				bCreateEvictedBinary = true;
-			}
-
-			FOpenGLLinkedProgram* CachedProgram = nullptr;
-			if (!bCreateEvictedBinary)
+			FOpenGLLinkedProgram* CachedProgram = GetOpenGLProgramsCache().Find(Config.ProgramKey, true);
+			if (!CachedProgram)
 			{
-				CachedProgram = GetOpenGLProgramsCache().Find(Config.ProgramKey, true);
-				if (!CachedProgram)
+				// ensure that pending request for this program has been completed before
+				if (FOpenGLProgramBinaryCache::CheckSinglePendingGLProgramCreateRequest(Config.ProgramKey))
 				{
-					// ensure that pending request for this program has been completed before
-					if (FOpenGLProgramBinaryCache::CheckSinglePendingGLProgramCreateRequest(Config.ProgramKey))
-					{
-						CachedProgram = GetOpenGLProgramsCache().Find(Config.ProgramKey, true);
-					}
+					CachedProgram = GetOpenGLProgramsCache().Find(Config.ProgramKey, true);
 				}
 			}
 
@@ -2981,13 +3003,10 @@ FBoundShaderStateRHIRef FOpenGLDynamicRHI::RHICreateBoundShaderState_OnThisThrea
 				FOpenGLPixelShader* PixelShader = ResourceCast(PixelShaderRHI);
 				FOpenGLGeometryShader* GeometryShader = ResourceCast(GeometryShaderRHI);
 		
-				// Make sure we have OpenGL context set up, and invalidate the parameters cache and current program (as we'll link a new one soon)
-				GetContextStateForCurrentContext().Program = -1;
-				MarkShaderParameterCachesDirty(PendingState.ShaderParameters, false);
-				PendingState.LinkedProgramAndDirtyFlag = nullptr;
+
 
 				// Link program, using the data provided in config
-				LinkedProgram = LinkProgram(Config, bFromPSOFileCache);
+				LinkedProgram = LinkProgram(Config);
 
 				if (LinkedProgram == NULL)
 				{
@@ -3011,29 +3030,9 @@ FBoundShaderStateRHIRef FOpenGLDynamicRHI::RHICreateBoundShaderState_OnThisThrea
 				}
 				else
 				{
-					if (bCreateEvictedBinary)
-					{
-						// Store the program as evicted binary only, this only happens for precompile PSOs and the returned result is not used.
-						if(LinkedProgram->LRUInfo.CachedProgramBinary->IsValid())
-						{
-							GetOpenGLProgramsCache().AddOrReplaceEvicted(Config.ProgramKey, MoveTemp(LinkedProgram->LRUInfo.CachedProgramBinary));
-						}
-						else
-						{
-							// not expected. This means the PSO is part of the pso precompile but is somehow not in the existing cache.
-							TUniqueObj<FOpenGLProgramBinary> Bin(UE::OpenGL::GetProgramBinaryFromGLProgram(LinkedProgram->Program));
-							GetOpenGLProgramsCache().AddOrReplaceEvicted(Config.ProgramKey, MoveTemp(Bin));
-						}
-						delete LinkedProgram;
-						return nullptr;
-					}
-					else
-					{
-						GetOpenGLProgramsCache().Add(Config.ProgramKey, LinkedProgram);
-					}
+					GetOpenGLProgramsCache().Add(Config.ProgramKey, LinkedProgram);
 				}
 			}
-			check(!bCreateEvictedBinary);
 		}
 
 		check(VertexDeclarationRHI);
@@ -3254,7 +3253,7 @@ FOpenGLBoundShaderState::~FOpenGLBoundShaderState()
 	check(LinkedProgram);
 	RunOnGLRenderContextThread([LinkedProgram = LinkedProgram]()
 	{
-		const bool bIsEvicted = GetOpenGLProgramsCache().IsUsingLRU() && GetOpenGLProgramsCache().IsEvicted(LinkedProgram->Config.ProgramKey);
+		const bool bIsEvicted = FGLProgramCache::IsUsingLRU() && GetOpenGLProgramsCache().IsEvicted(LinkedProgram->Config.ProgramKey);
 		if( !bIsEvicted )
 		{
 			FOpenGLLinkedProgram* Prog = StaticLastReleasedPrograms[StaticLastReleasedProgramsIndex];
@@ -3655,20 +3654,24 @@ namespace UE
 		// Called from the binary file cache when the binary version of a program has been encountered.
 		void OnGLProgramLoadedFromBinaryCache(const FOpenGLProgramKey& ProgramKey, TUniqueObj<FOpenGLProgramBinary>&& ProgramBinaryData)
 		{
+			QUICK_SCOPE_CYCLE_COUNTER(STAT_OpenGLOnGLProgramLoadedFromBinaryCache);
+
+// 			FScopeLock Lock(&GProgramBinaryCacheCS);
 			FOpenGLLinkedProgram* FoundProgram = GetOpenGLProgramsCache().Find(ProgramKey, false);
 			const bool bProgramExists = FoundProgram != nullptr;
-			const bool bIsEvicted = bProgramExists && GetOpenGLProgramsCache().IsEvicted(ProgramKey);
-			if (GetOpenGLProgramsCache().IsUsingLRU())
+			if (FGLProgramCache::IsUsingLRU())
 			{
+				const bool bIsEvicted = bProgramExists && GetOpenGLProgramsCache().IsEvicted(ProgramKey);
 				// always replace any existing binary data with this.
 				if (!bProgramExists || bIsEvicted)
 				{
+					check(!bProgramExists || FoundProgram->LRUInfo.CachedProgramBinary->GetDataView().IsEmpty());
+
 					// Always add programs as evicted, 1st use will create them as programs.
 					// This will reduce pressure on driver by ensuring only used programs
 					// are created.
 					// In this case do not create the GL program.
 					GetOpenGLProgramsCache().AddOrReplaceEvicted(ProgramKey, MoveTemp(ProgramBinaryData));
-					check(!bProgramExists || FoundProgram->LRUInfo.CachedProgramBinary->GetDataView().IsEmpty());
 				}
 				else
 				{
@@ -3699,8 +3702,9 @@ namespace UE
 
 static void TickProgramLRU()
 {
-	if (GetOpenGLProgramsCache().IsUsingLRU() && GUntouchedProgramEvictTimeSeconds > 0)
+	if (FGLProgramCache::IsUsingLRU() && GUntouchedProgramEvictTimeSeconds > 0)
 	{
+// 		FScopeLock Lock(&GProgramBinaryCacheCS);
 		int FramePace = FPlatformRHIFramePacer::GetFramePace();
 		int MaxProgramsToEvictPerFrame = 10;
 		int32 CurrentResidentCount = GetOpenGLProgramsCache().GetLRUSize();
diff --git a/Engine/Source/Runtime/OpenGLDrv/Public/OpenGLDrv.h b/Engine/Source/Runtime/OpenGLDrv/Public/OpenGLDrv.h
--- a/Engine/Source/Runtime/OpenGLDrv/Public/OpenGLDrv.h
+++ b/Engine/Source/Runtime/OpenGLDrv/Public/OpenGLDrv.h
@@ -338,6 +338,12 @@ public:
 		static_assert(TIsGLProxyObject_V<typename TOpenGLResourceTraits<TRHIType>::TConcreteType>, "Wrong type past");
 		return static_cast<typename TOpenGLResourceTraits<TRHIType>::TConcreteType*>(Resource);
 	}
+	template<typename TRHIType>
+	static FORCEINLINE const typename TOpenGLResourceTraits<TRHIType>::TConcreteType* ResourceCastProxy(const TRHIType* Resource)
+	{
+		static_assert(TIsGLProxyObject_V<typename TOpenGLResourceTraits<TRHIType>::TConcreteType>, "Wrong type past");
+		return static_cast<const typename TOpenGLResourceTraits<TRHIType>::TConcreteType*>(Resource);
+	}
 
 	// If using a Proxy object return the contained GL object rather than the proxy itself.
 	template<typename TRHIType>
@@ -864,6 +870,7 @@ public:
 			false);
 	}
 
+	bool LinkComputeShader(FRHIComputeShader* ComputeShaderRHI, FOpenGLComputeShader* ComputeShader);
 	class FOpenGLLinkedProgram* GetLinkedComputeProgram(FRHIComputeShader* ComputeShaderRHI);
 
 	FBoundShaderStateRHIRef RHICreateBoundShaderState_OnThisThread(FRHIVertexDeclaration* VertexDeclaration, FRHIVertexShader* VertexShader, FRHIPixelShader* PixelShader, FRHIGeometryShader* GeometryShader, bool FromPSOFileCache);
@@ -903,6 +910,7 @@ private:
 
 	void PrepareGFXBoundShaderState(const FGraphicsPipelineStateInitializer& Initializer);
 
+	FOpenGLLinkedProgram* LinkProgram(const class FOpenGLLinkedProgramConfiguration& Config);
 
 	/** called once per frame, used for resource processing */
 	void EndFrameTick();
