Remove SupportsCopyImage as GL ES 3.2 standard has glcopyImageSubData

diff --git a/Engine/Source/Runtime/OpenGLDrv/Private/Linux/OpenGLLinux.cpp b/Engine/Source/Runtime/OpenGLDrv/Private/Linux/OpenGLLinux.cpp
--- a/Engine/Source/Runtime/OpenGLDrv/Private/Linux/OpenGLLinux.cpp
+++ b/Engine/Source/Runtime/OpenGLDrv/Private/Linux/OpenGLLinux.cpp
@@ -1022,7 +1022,5 @@ void FLinuxOpenGL::ProcessExtensions( const FString& ExtensionsString )
 		glTexStorage1D = nullptr;
 		glTexStorage2D = nullptr;
 		glTexStorage3D = nullptr;
-
-		FOpenGLBase::bSupportsCopyImage  = false;
 	}
 }
diff --git a/Engine/Source/Runtime/OpenGLDrv/Private/OpenGLDevice.cpp b/Engine/Source/Runtime/OpenGLDrv/Private/OpenGLDevice.cpp
--- a/Engine/Source/Runtime/OpenGLDrv/Private/OpenGLDevice.cpp
+++ b/Engine/Source/Runtime/OpenGLDrv/Private/OpenGLDevice.cpp
@@ -1081,15 +1081,8 @@ static void InitRHICapabilitiesForGL()
 	GShaderPlatformForFeatureLevel[ERHIFeatureLevel::SM4_REMOVED] = SP_NumPlatforms;
 	GShaderPlatformForFeatureLevel[ERHIFeatureLevel::SM5] = SP_NumPlatforms;
 
-	// Disable texture streaming on devices with ES3.1 or lower, unless we have the GL_APPLE_copy_texture_levels extension or support for glCopyImageSubData
-	if (GMaxRHIFeatureLevel <= ERHIFeatureLevel::ES3_1)
-	{
-		GRHISupportsTextureStreaming = FOpenGL::SupportsCopyImage();
-	}
-	else
-	{
-		GRHISupportsTextureStreaming = true;
-	}
+
+	GRHISupportsTextureStreaming = true;
 
 	// Disable texture streaming if forced off by r.OpenGL.DisableTextureStreamingSupport
 	static const auto CVarDisableOpenGLTextureStreamingSupport = IConsoleManager::Get().FindTConsoleVariableDataInt(TEXT("r.OpenGL.DisableTextureStreamingSupport"));
diff --git a/Engine/Source/Runtime/OpenGLDrv/Private/OpenGLDrv.cpp b/Engine/Source/Runtime/OpenGLDrv/Private/OpenGLDrv.cpp
--- a/Engine/Source/Runtime/OpenGLDrv/Private/OpenGLDrv.cpp
+++ b/Engine/Source/Runtime/OpenGLDrv/Private/OpenGLDrv.cpp
@@ -487,7 +487,6 @@ GLint FOpenGLBase::MaxGeometryUniformComponents = -1;
 bool  FOpenGLBase::bSupportsClipControl = false;
 bool  FOpenGLBase::bSupportsASTC = false;
 bool  FOpenGLBase::bSupportsASTCHDR = false;
-bool  FOpenGLBase::bSupportsCopyImage = false;
 bool  FOpenGLBase::bSupportsSeamlessCubemap = false;
 bool  FOpenGLBase::bSupportsVolumeTextureRendering = false;
 bool  FOpenGLBase::bSupportsTextureFilterAnisotropic = false;
@@ -536,9 +535,6 @@ void FOpenGLBase::ProcessExtensions( const FString& ExtensionsString )
 
 	bSupportsASTCHDR = bSupportsASTC && ExtensionsString.Contains(TEXT("GL_KHR_texture_compression_astc_hdr"));
 
-	// check for copy image support
-	bSupportsCopyImage = ExtensionsString.Contains(TEXT("GL_ARB_copy_image"));
-
 	bSupportsSeamlessCubemap = ExtensionsString.Contains(TEXT("GL_ARB_seamless_cube_map"));
 	
 	bSupportsTextureFilterAnisotropic = ExtensionsString.Contains(TEXT("GL_EXT_texture_filter_anisotropic"));
diff --git a/Engine/Source/Runtime/OpenGLDrv/Private/OpenGLES.cpp b/Engine/Source/Runtime/OpenGLDrv/Private/OpenGLES.cpp
--- a/Engine/Source/Runtime/OpenGLDrv/Private/OpenGLES.cpp
+++ b/Engine/Source/Runtime/OpenGLDrv/Private/OpenGLES.cpp
@@ -47,7 +47,7 @@ namespace GLFuncPointers
 	// ES 3.2
 	PFNGLTEXBUFFEREXTPROC					glTexBufferEXT = nullptr;
 	PFNGLTEXBUFFERRANGEEXTPROC				glTexBufferRangeEXT = nullptr;
-	PFNGLCOPYIMAGESUBDATAEXTPROC			glCopyImageSubDataEXT = nullptr;
+	PFNGLCOPYIMAGESUBDATAEXTPROC			glCopyImageSubData = nullptr;
 	PFNGLENABLEIEXTPROC						glEnableiEXT = nullptr;
 	PFNGLDISABLEIEXTPROC					glDisableiEXT = nullptr;
 	PFNGLBLENDEQUATIONIEXTPROC				glBlendEquationiEXT = nullptr;
@@ -297,7 +297,7 @@ void FOpenGLES::ProcessExtensions(const FString& ExtensionsString)
 	{
 		glTexBufferEXT = (PFNGLTEXBUFFEREXTPROC)((void*)eglGetProcAddress("glTexBuffer"));
 		glTexBufferRangeEXT = (PFNGLTEXBUFFERRANGEEXTPROC)((void*)eglGetProcAddress("glTexBufferRange"));
-		glCopyImageSubDataEXT = (PFNGLCOPYIMAGESUBDATAEXTPROC)((void*)eglGetProcAddress("glCopyImageSubData"));
+		glCopyImageSubData = (PFNGLCOPYIMAGESUBDATAEXTPROC)((void*)eglGetProcAddress("glCopyImageSubData"));
 		glEnableiEXT = (PFNGLENABLEIEXTPROC)((void*)eglGetProcAddress("glEnablei"));
 		glDisableiEXT = (PFNGLDISABLEIEXTPROC)((void*)eglGetProcAddress("glDisablei"));
 		glBlendEquationiEXT = (PFNGLBLENDEQUATIONIEXTPROC)((void*)eglGetProcAddress("glBlendEquationi"));
@@ -326,13 +326,6 @@ void FOpenGLES::ProcessExtensions(const FString& ExtensionsString)
 		glTexBufferEXT = (PFNGLTEXBUFFEREXTPROC)((void*)eglGetProcAddress("glTexBufferEXT"));
 		glTexBufferRangeEXT = (PFNGLTEXBUFFERRANGEEXTPROC)((void*)eglGetProcAddress("glTexBufferRangeEXT"));
 	}
-
-	if (!glCopyImageSubDataEXT && ExtensionsString.Contains(TEXT("GL_EXT_copy_image")))
-	{
-		// GL_EXT_copy_image
-		glCopyImageSubDataEXT = (PFNGLCOPYIMAGESUBDATAEXTPROC)((void*)eglGetProcAddress("glCopyImageSubDataEXT"));
-	}
-	bSupportsCopyImage = (glCopyImageSubDataEXT != nullptr);
 }
 
 #endif
diff --git a/Engine/Source/Runtime/OpenGLDrv/Private/OpenGLRenderTarget.cpp b/Engine/Source/Runtime/OpenGLDrv/Private/OpenGLRenderTarget.cpp
--- a/Engine/Source/Runtime/OpenGLDrv/Private/OpenGLRenderTarget.cpp
+++ b/Engine/Source/Runtime/OpenGLDrv/Private/OpenGLRenderTarget.cpp
@@ -505,7 +505,7 @@ void FOpenGLDynamicRHI::RHICopyToResolveTarget(FRHITexture* SourceTextureRHI, FR
 			&& SourceTexture->Target == DestTexture->Target // glCopyImageSubData() doesn't like copying from a texture to a renderbuffer on Android
 			;
 		
-		if ( !bTrueBlit || !FOpenGL::SupportsCopyImage() )
+		if ( !bTrueBlit)
 		{
 			// Color buffers can be GL_NONE for attachment purposes if they aren't used as render targets
 			const bool bIsColorBuffer = SourceTexture->Attachment != GL_DEPTH_STENCIL_ATTACHMENT
diff --git a/Engine/Source/Runtime/OpenGLDrv/Private/OpenGLTexture.cpp b/Engine/Source/Runtime/OpenGLDrv/Private/OpenGLTexture.cpp
--- a/Engine/Source/Runtime/OpenGLDrv/Private/OpenGLTexture.cpp
+++ b/Engine/Source/Runtime/OpenGLDrv/Private/OpenGLTexture.cpp
@@ -249,7 +249,6 @@ FOpenGLTextureDesc::FOpenGLTextureDesc(FRHITextureDesc const& InDesc)
 	// can run on RT.
 	bCanCreateAsEvicted =
 		CanDeferTextureCreation()
-		&& FOpenGL::SupportsCopyImage()
 		&& InDesc.Flags != TexCreate_None // ignore TexCreate_None
 		&& !EnumHasAnyFlags((ETextureCreateFlags)CVarDeferTextureCreationExcludeMask.GetValueOnAnyThread(), InDesc.Flags)  // Anything outside of these flags cannot be evicted.
 		&& Target == GL_TEXTURE_2D
@@ -257,8 +256,8 @@ FOpenGLTextureDesc::FOpenGLTextureDesc(FRHITextureDesc const& InDesc)
 
 	if (GOGLTextureEvictLogging)
 	{
-		UE_CLOG(!bCanCreateAsEvicted, LogRHI, Warning, TEXT("CanDeferTextureCreation:%d, SupportsCopyImage:%d, Flags:%llx Mask:%x, Target:%x"),
-			bCanCreateAsEvicted, FOpenGL::SupportsCopyImage(), InDesc.Flags, CVarDeferTextureCreationExcludeMask.GetValueOnAnyThread(), Target);
+		UE_CLOG(!bCanCreateAsEvicted, LogRHI, Warning, TEXT("CanDeferTextureCreation:%d, Flags:%llx Mask:%x, Target:%x"),
+			bCanCreateAsEvicted, InDesc.Flags, CVarDeferTextureCreationExcludeMask.GetValueOnAnyThread(), Target);
 	}
 
 	bIsPowerOfTwo = 
@@ -1412,8 +1411,6 @@ void FOpenGLTexture::CloneViaCopyImage(FOpenGLTexture* Src, uint32 InNumMips, in
 {
 	VERIFY_GL_SCOPE();
 
-	check(FOpenGL::SupportsCopyImage());
-	
 	check(Src->bCanCreateAsEvicted == bCanCreateAsEvicted);
 	if (bCanCreateAsEvicted)
 	{
@@ -1458,190 +1455,6 @@ void FOpenGLTexture::CloneViaCopyImage(FOpenGLTexture* Src, uint32 InNumMips, in
 	
 }
 
-void FOpenGLTexture::CloneViaPBO(FOpenGLTexture* Src, uint32 InNumMips, int32 SrcOffset, int32 DstOffset)
-{
-	VERIFY_GL_SCOPE();
-	
-	// apparently it's not possible to retrieve compressed image from GL_TEXTURE_2D_ARRAY in OpenGL for compressed images
-	// and for uncompressed ones it's not possible to specify the image index
-	check(this->GetSizeZ() == 0);
-
-	EPixelFormat PixelFormat = this->GetFormat();
-	check(PixelFormat == Src->GetFormat());
-
-	const FOpenGLTextureFormat& GLFormat = GOpenGLTextureFormats[PixelFormat];
-	const bool bSRGB = EnumHasAnyFlags(this->GetFlags(), TexCreate_SRGB);
-	check(bSRGB == EnumHasAnyFlags(Src->GetFlags(), TexCreate_SRGB));
-	
-	const uint32 BlockSizeX = GPixelFormats[PixelFormat].BlockSizeX;
-	const uint32 BlockSizeY = GPixelFormats[PixelFormat].BlockSizeY;
-	const uint32 BlockBytes = GPixelFormats[PixelFormat].BlockBytes;
-	
-	FOpenGLContextState& ContextState = FOpenGLDynamicRHI::Get().GetContextStateForCurrentContext();
-	
-	for (uint32 ArrayIndex = 0; ArrayIndex < this->GetEffectiveSizeZ(); ArrayIndex++)
-	{
-		// use PBO functionality to copy mip level by mip level
-		for(uint32 MipIndex = 0;MipIndex < InNumMips;++MipIndex)
-		{
-			// Actual mip levels
-			const uint32 DstMipIndex = MipIndex + DstOffset;
-			const uint32 SrcMipIndex = MipIndex + SrcOffset;
-			
-			// Calculate the dimensions of the mip-map.
-			const uint32 MipSizeX = FMath::Max(this->GetSizeX() >> DstMipIndex,1u);
-			const uint32 MipSizeY = FMath::Max(this->GetSizeY() >> DstMipIndex,1u);
-			
-			// Then the rounded PBO size required to capture this mip
-			const uint32 DataSizeX = FMath::Max(MipSizeX,BlockSizeX);
-			const uint32 DataSizeY = FMath::Max(MipSizeY,BlockSizeY);
-			uint32 NumBlocksX = (DataSizeX + BlockSizeX - 1) / BlockSizeX;
-			uint32 NumBlocksY = (DataSizeY + BlockSizeY - 1) / BlockSizeY;
-
-			const uint32 MipBytes = NumBlocksX * NumBlocksY * BlockBytes;
-			const int32 BufferIndex = DstMipIndex * (bCubemap ? 6 : 1) * this->GetEffectiveSizeZ() + ArrayIndex;
-			const int32 SrcBufferIndex = SrcMipIndex * (Src->bCubemap ? 6 : 1) * Src->GetEffectiveSizeZ() + ArrayIndex;
-			
-			// Standard path with a PBO mirroring ever slice of a texture to allow multiple simulataneous maps
-			if (!IsValidRef(PixelBuffers[BufferIndex]))
-			{
-				PixelBuffers[BufferIndex] = new FOpenGLPixelBuffer(nullptr, GL_PIXEL_UNPACK_BUFFER, 0, MipBytes, BUF_Dynamic, nullptr);
-			}
-			
-			TRefCountPtr<FOpenGLPixelBuffer> PixelBuffer = PixelBuffers[BufferIndex];
-			check(PixelBuffer->GetSize() == MipBytes);
-			check(!PixelBuffer->IsLocked());
-			
-			// Transfer data from texture to pixel buffer.
-			// This may be further optimized by caching information if surface content was changed since last lock.
-			// Use a texture stage that's not likely to be used for draws, to avoid waiting
-			FOpenGLDynamicRHI::Get().CachedSetupTextureStage(ContextState, FOpenGL::GetMaxCombinedTextureImageUnits() - 1, Src->Target, Src->GetResource(), -1, this->GetNumMips());
-			
-			glBindBuffer( GL_PIXEL_PACK_BUFFER, PixelBuffer->Resource );
-			
-			if (GLFormat.bCompressed)
-			{
-				FOpenGL::GetCompressedTexImage(Src->bCubemap ? GL_TEXTURE_CUBE_MAP_POSITIVE_X + ArrayIndex : Src->Target,
-											   SrcMipIndex,
-											   0);	// offset into PBO
-			}
-			else
-			{
-				glPixelStorei(GL_PACK_ALIGNMENT, 1);
-				FOpenGL::GetTexImage(Src->bCubemap ? GL_TEXTURE_CUBE_MAP_POSITIVE_X + ArrayIndex : Src->Target,
-									 SrcMipIndex,
-									 GLFormat.Format,
-									 GLFormat.Type,
-									 0);	// offset into PBO
-				glPixelStorei(GL_PACK_ALIGNMENT, 4);
-			}
-			
-			// copy the texture data
-			// Upload directly into Dst to avoid out-of-band synchronization caused by glMapBuffer!
-			{
-				CachedBindPixelUnpackBuffer( GL_PIXEL_UNPACK_BUFFER, PixelBuffer->Resource );
-				
-				// Use a texture stage that's not likely to be used for draws, to avoid waiting
-				FOpenGLDynamicRHI::Get().CachedSetupTextureStage(ContextState, FOpenGL::GetMaxCombinedTextureImageUnits() - 1, Target, GetResource(), -1, this->GetNumMips());
-				
-				if( this->GetSizeZ() )
-				{
-					// texture 2D array
-					if (GLFormat.bCompressed)
-					{
-						FOpenGL::CompressedTexSubImage3D(Target,
-														 DstMipIndex,
-														 0,
-														 0,
-														 ArrayIndex,
-														 MipSizeX,
-														 MipSizeY,
-														 1,
-														 GLFormat.InternalFormat[bSRGB],
-														 PixelBuffer->GetSize(),
-														 0);
-					}
-					else
-					{
-						glPixelStorei(GL_UNPACK_ALIGNMENT, 1);
-						check( FOpenGL::SupportsTexture3D() );
-						FOpenGL::TexSubImage3D(Target,
-											   DstMipIndex,
-											   0,
-											   0,
-											   ArrayIndex,
-											   MipSizeX,
-											   MipSizeY,
-											   1,
-											   GLFormat.Format,
-											   GLFormat.Type,
-											   0);	// offset into PBO
-						glPixelStorei(GL_UNPACK_ALIGNMENT, 4);
-					}
-				}
-				else
-				{
-					if (GLFormat.bCompressed)
-					{
-						if (GetAllocatedStorageForMip(DstMipIndex,ArrayIndex))
-						{
-							glCompressedTexSubImage2D(bCubemap ? GL_TEXTURE_CUBE_MAP_POSITIVE_X + ArrayIndex : Target,
-													  DstMipIndex,
-													  0,
-													  0,
-													  MipSizeX,
-													  MipSizeY,
-													  GLFormat.InternalFormat[bSRGB],
-													  PixelBuffer->GetSize(),
-													  0);	// offset into PBO
-						}
-						else
-						{
-							glCompressedTexImage2D(bCubemap ? GL_TEXTURE_CUBE_MAP_POSITIVE_X + ArrayIndex : Target,
-												   DstMipIndex,
-												   GLFormat.InternalFormat[bSRGB],
-												   MipSizeX,
-												   MipSizeY,
-												   0,
-												   PixelBuffer->GetSize(),
-												   0);	// offset into PBO
-							SetAllocatedStorageForMip(DstMipIndex,ArrayIndex);
-						}
-					}
-					else
-					{
-						// All construction paths should have called TexStorage2D or TexImage2D. So we will
-						// always call TexSubImage2D.
-						check(GetAllocatedStorageForMip(DstMipIndex,ArrayIndex) == true);
-						glPixelStorei(GL_UNPACK_ALIGNMENT, 1);
-						glTexSubImage2D(bCubemap ? GL_TEXTURE_CUBE_MAP_POSITIVE_X + ArrayIndex : Target,
-										DstMipIndex,
-										0,
-										0,
-										MipSizeX,
-										MipSizeY,
-										GLFormat.Format,
-										GLFormat.Type,
-										0);	// offset into PBO
-						glPixelStorei(GL_UNPACK_ALIGNMENT, 4);
-					}
-				}
-			}
-			
-			// need to free PBO if we aren't keeping shadow copies
-			PixelBuffers[BufferIndex] = NULL;
-			
-			// No need to restore texture stage; leave it like this,
-			// and the next draw will take care of cleaning it up; or
-			// next operation that needs the stage will switch something else in on it.
-		}
-	}
-	
-	// Reset the buffer bindings on exit only
-	glBindBuffer( GL_PIXEL_PACK_BUFFER, 0 );
-	CachedBindPixelUnpackBuffer(GL_PIXEL_UNPACK_BUFFER, 0);
-}
-
 /*-----------------------------------------------------------------------------
 	2D texture support.
 -----------------------------------------------------------------------------*/
@@ -1776,14 +1589,7 @@ FTexture2DRHIRef FOpenGLDynamicRHI::RHIAsyncReallocateTexture2D(FRHITexture2D* T
 		const uint32 SourceMipOffset = SourceMipCount - NumSharedMips;
 		const uint32 DestMipOffset = NewMipCount - NumSharedMips;
 
-		if (FOpenGL::SupportsCopyImage())
-		{
-			NewTexture->CloneViaCopyImage(OldTexture, NumSharedMips, SourceMipOffset, DestMipOffset);
-		}
-		else
-		{
-			NewTexture->CloneViaPBO(OldTexture, NumSharedMips, SourceMipOffset, DestMipOffset);
-		}
+		NewTexture->CloneViaCopyImage(OldTexture, NumSharedMips, SourceMipOffset, DestMipOffset);
 
 		RequestStatus->Decrement();
 	});
@@ -2082,11 +1888,6 @@ void FOpenGLDynamicRHI::RHICopyTexture(FRHITexture* SourceTextureRHI, FRHITextur
 	FOpenGLTexture* SourceTexture = GetOpenGLTextureFromRHITexture(SourceTextureRHI);
 	FOpenGLTexture* DestTexture = GetOpenGLTextureFromRHITexture(DestTextureRHI);
 
-	checkf(SourceTexture->Target == DestTexture->Target, TEXT("Cannot copy between different texture targets, SourceTexture Target=%x, Format=%d, Flags=%x; DestTexture Target=%x, Format=%d, Flags=%x"),
-		SourceTexture->Target, SourceTextureRHI->GetFormat(), SourceTextureRHI->GetFlags(),
-		DestTexture->Target, DestTextureRHI->GetFormat(), DestTextureRHI->GetFlags()
-	);
-	
 
 	GLsizei Width, Height, Depth;
 
@@ -2121,120 +1922,52 @@ void FOpenGLDynamicRHI::RHICopyTexture(FRHITexture* SourceTextureRHI, FRHITextur
 	GLint SrcMip = CopyInfo.SourceMipIndex;
 	GLint DestMip = CopyInfo.DestMipIndex;
 
-	if (FOpenGL::SupportsCopyImage())
+	for (uint32 MipIndex = 0; MipIndex < CopyInfo.NumMips; ++MipIndex)
 	{
-		for (uint32 MipIndex = 0; MipIndex < CopyInfo.NumMips; ++MipIndex)
+		// X, Y, Z
+		FIntVector Src, Dst;
+
+		auto SetOffsets = [MipIndex, &CopyInfo, &Depth] (const GLenum Target, const FIntVector& Position, const uint32& SliceIndex, FIntVector& OutOffsets)
 		{
-			GLint SrcZOffset, DestZOffset;
-			switch (SourceTexture->Target)
+			switch (Target)
 			{
 			case GL_TEXTURE_3D:
 			case GL_TEXTURE_CUBE_MAP:
 				// For cube maps, the Z offsets select the starting faces.
-				SrcZOffset = CopyInfo.SourcePosition.Z >> MipIndex;
-				DestZOffset = CopyInfo.DestPosition.Z  >> MipIndex;
+				OutOffsets.Z = Position.Z >> MipIndex;
 				break;
 			case GL_TEXTURE_1D_ARRAY:
 			case GL_TEXTURE_2D_ARRAY:
 				// For texture arrays, the Z offsets and depth actually refer to the range of slices to copy.
-				SrcZOffset = CopyInfo.SourceSliceIndex;
-				DestZOffset = CopyInfo.DestSliceIndex;
+				OutOffsets.Z = SliceIndex;
 				Depth = CopyInfo.NumSlices;
 				break;
 			default:
-				SrcZOffset = 0;
-				DestZOffset = 0;
+				OutOffsets.Z = 0;
 				break;
 			}
 
-			GLint SourceX = CopyInfo.SourcePosition.X >> MipIndex;
-			GLint SourceY = CopyInfo.SourcePosition.Y >> MipIndex;
-
-			GLint DestX = CopyInfo.DestPosition.X >> MipIndex;
-			GLint DestY = CopyInfo.DestPosition.Y >> MipIndex;
-
-			FOpenGL::CopyImageSubData(SourceTexture->GetResource(), SourceTexture->Target, SrcMip, SourceX, SourceY, SrcZOffset,
-				DestTexture->GetResource(), DestTexture->Target, DestMip, DestX, DestY, DestZOffset,
-				Width, Height, Depth);
-
-			++SrcMip;
-			++DestMip;
-
-			Width = FMath::Max(1, Width >> 1);
-			Height = FMath::Max(1, Height >> 1);
-			if(DestTexture->Target == GL_TEXTURE_3D)
-			{
-				Depth = FMath::Max(1, Depth >> 1);
-			}
-		}
-
-		return;
-	}
-	
-	// Convert sub texture regions to GL types
-	GLint XOffset = CopyInfo.DestPosition.X;
-	GLint YOffset = CopyInfo.DestPosition.Y;
-	GLint ZOffset = CopyInfo.DestPosition.Z;
-	GLint X = CopyInfo.SourcePosition.X;
-	GLint Y = CopyInfo.SourcePosition.Y;
-	GLint Z = CopyInfo.SourcePosition.Z;
-
-	// Use a texture stage that's not likely to be used for draws, to avoid waiting
-	FOpenGLContextState& ContextState = GetContextStateForCurrentContext();
-	CachedSetupTextureStage(ContextState, FOpenGL::GetMaxCombinedTextureImageUnits() - 1, DestTexture->Target, DestTexture->GetResource(), 0, DestTextureRHI->GetNumMips());
-	CachedBindPixelUnpackBuffer(ContextState, 0);
-
-	// Bind source texture to an FBO to read from
-	for (uint32 SliceIndex = 0; SliceIndex < CopyInfo.NumSlices; ++SliceIndex)
-	{
-		for (uint32 MipIndex = 0; MipIndex < CopyInfo.NumMips; ++MipIndex)
-		{
-			FOpenGLTexture* RenderTargets[1] = { SourceTexture };
-			uint32 MipLevels[1] = { static_cast<uint32>(SrcMip) };
-			uint32 ArrayIndices[1] = { CopyInfo.SourceSliceIndex + SliceIndex };
-
-			GLuint SourceFBO = GetOpenGLFramebuffer(1, RenderTargets, ArrayIndices, MipLevels, nullptr);
-			check(SourceFBO != 0);
-
-			glBindFramebuffer(GL_FRAMEBUFFER, SourceFBO);
+			OutOffsets.X = Position.X >> MipIndex;
+			OutOffsets.Y = Position.Y >> MipIndex;
+		};
 
-			FOpenGL::ReadBuffer(GL_COLOR_ATTACHMENT0);
+		SetOffsets(SourceTexture->Target, CopyInfo.SourcePosition, CopyInfo.SourceSliceIndex, Src);
+		SetOffsets(DestTexture->Target, CopyInfo.DestPosition, CopyInfo.DestSliceIndex, Dst);
 
-			switch (DestTexture->Target)
-			{
-				case GL_TEXTURE_1D:
-					FOpenGL::CopyTexSubImage1D(DestTexture->Target, DestMip, XOffset >> MipIndex, X >> MipIndex, 0, Width);
-					break;
-				case GL_TEXTURE_1D_ARRAY:
-					FOpenGL::CopyTexSubImage2D(DestTexture->Target, DestMip, XOffset >> MipIndex, CopyInfo.DestSliceIndex + SliceIndex, X >> MipIndex, 0, Width, 1);
-					break;
-				case GL_TEXTURE_2D:
-				case GL_TEXTURE_RECTANGLE:
-					FOpenGL::CopyTexSubImage2D(DestTexture->Target, DestMip, XOffset >> MipIndex, YOffset >> MipIndex, X >> MipIndex, Y >> MipIndex, Width, Height);
-					break;
-				case GL_TEXTURE_2D_ARRAY:
-					FOpenGL::CopyTexSubImage3D(DestTexture->Target, DestMip, XOffset >> MipIndex, YOffset >> MipIndex, CopyInfo.DestSliceIndex + SliceIndex, X >> MipIndex, Y >> MipIndex, Width, Height);
-					break;
-				case GL_TEXTURE_3D:
-					FOpenGL::CopyTexSubImage3D(DestTexture->Target, DestMip, XOffset >> MipIndex, YOffset >> MipIndex, ZOffset >> MipIndex, X >> MipIndex, Y >> MipIndex, Width, Height);
-					break;
-				case GL_TEXTURE_CUBE_MAP:
-					for (int32 FaceIndex = FMath::Min((int32)CopyInfo.NumSlices, 6) - 1; FaceIndex >= 0; FaceIndex--)
-					{
-						FOpenGL::CopyTexSubImage2D(GL_TEXTURE_CUBE_MAP_POSITIVE_X + (uint32)FaceIndex, CopyInfo.DestMipIndex, XOffset >> MipIndex, YOffset >> MipIndex, X >> MipIndex, Y >> MipIndex, Width, Height);
-					}					
-					break;
-			}
-		}
+		FOpenGL::CopyImageSubData(SourceTexture->GetResource(), SourceTexture->Target, SrcMip, Src.X, Src.Y, Src.Z,
+			DestTexture->GetResource(), DestTexture->Target, DestMip, Dst.X, Dst.Y, Dst.Z,
+			Width, Height, Depth);
 
 		++SrcMip;
 		++DestMip;
 
 		Width = FMath::Max(1, Width >> 1);
 		Height = FMath::Max(1, Height >> 1);
+		if(DestTexture->Target == GL_TEXTURE_3D)
+		{
+			Depth = FMath::Max(1, Depth >> 1);
+		}
 	}
-
-	ContextState.Framebuffer = (GLuint)-1;
 }
 
 
diff --git a/Engine/Source/Runtime/OpenGLDrv/Private/Windows/OpenGLWindows.h b/Engine/Source/Runtime/OpenGLDrv/Private/Windows/OpenGLWindows.h
--- a/Engine/Source/Runtime/OpenGLDrv/Private/Windows/OpenGLWindows.h
+++ b/Engine/Source/Runtime/OpenGLDrv/Private/Windows/OpenGLWindows.h
@@ -367,7 +367,6 @@
 	EnumMacro(PFNGLDISABLEIPROC,glDisableiEXT) \
 	EnumMacro(PFNGLENABLEIPROC,glEnableiEXT) \
 	EnumMacro(PFNGLFRAMEBUFFERTEXTUREPROC,glFramebufferTextureEXT) \
-	EnumMacro(PFNGLCOPYIMAGESUBDATAPROC, glCopyImageSubDataEXT) \
 	EnumMacro(PFNGLTEXBUFFERPROC,glTexBufferEXT) \
 	EnumMacro(PFNGLTEXBUFFERRANGEPROC,glTexBufferRangeEXT) \
 	EnumMacro(PFNGLDEPTHRANGEFPROC,glDepthRangef) \
@@ -677,7 +676,7 @@ struct FWindowsOpenGL : public FOpenGLESDeferred
 
 	static FORCEINLINE void CopyImageSubData(GLuint SrcName, GLenum SrcTarget, GLint SrcLevel, GLint SrcX, GLint SrcY, GLint SrcZ, GLuint DstName, GLenum DstTarget, GLint DstLevel, GLint DstX, GLint DstY, GLint DstZ, GLsizei Width, GLsizei Height, GLsizei Depth)
 	{
-		glCopyImageSubDataEXT( SrcName, SrcTarget, SrcLevel, SrcX, SrcY, SrcZ, DstName, DstTarget, DstLevel, DstX, DstY, DstZ, Width, Height, Depth);
+		glCopyImageSubData( SrcName, SrcTarget, SrcLevel, SrcX, SrcY, SrcZ, DstName, DstTarget, DstLevel, DstX, DstY, DstZ, Width, Height, Depth);
 	}
 };
 
diff --git a/Engine/Source/Runtime/OpenGLDrv/Public/OpenGL.h b/Engine/Source/Runtime/OpenGLDrv/Public/OpenGL.h
--- a/Engine/Source/Runtime/OpenGLDrv/Public/OpenGL.h
+++ b/Engine/Source/Runtime/OpenGLDrv/Public/OpenGL.h
@@ -121,7 +121,6 @@ public:
 	static FORCEINLINE bool SupportsETC2()								{ return false; }
 	static FORCEINLINE bool SupportsFramebufferSRGBEnable()				{ return true; }
 	static FORCEINLINE bool SupportsFastBufferData()					{ return true; }
-	static FORCEINLINE bool SupportsCopyImage()							{ return bSupportsCopyImage; }
 	static FORCEINLINE bool SupportsTextureFilterAnisotropic()			{ return bSupportsTextureFilterAnisotropic; }
 	static FORCEINLINE bool SupportsSeparateAlphaBlend()				{ return bSupportsDrawBuffersBlend; }
 	static FORCEINLINE void EnableSupportsClipControl()					{ bSupportsClipControl = true; }
@@ -395,9 +394,6 @@ protected:
 	/** GL_KHR_texture_compression_astc_hdr */
 	static bool bSupportsASTCHDR;
 
-	/** GL_ARB_copy_image */
-	static bool bSupportsCopyImage;
-
 	/** GL_ARB_seamless_cube_map */
 	static bool bSupportsSeamlessCubemap;
 	
diff --git a/Engine/Source/Runtime/OpenGLDrv/Public/OpenGLES.h b/Engine/Source/Runtime/OpenGLDrv/Public/OpenGLES.h
--- a/Engine/Source/Runtime/OpenGLDrv/Public/OpenGLES.h
+++ b/Engine/Source/Runtime/OpenGLDrv/Public/OpenGLES.h
@@ -153,7 +153,7 @@ namespace GLFuncPointers
 	// ES 3.2
 	extern PFNGLTEXBUFFEREXTPROC				glTexBufferEXT;
 	extern PFNGLTEXBUFFERRANGEEXTPROC			glTexBufferRangeEXT;
-	extern PFNGLCOPYIMAGESUBDATAEXTPROC			glCopyImageSubDataEXT;
+	extern PFNGLCOPYIMAGESUBDATAEXTPROC			glCopyImageSubData;
 	extern PFNGLENABLEIEXTPROC					glEnableiEXT;
 	extern PFNGLDISABLEIEXTPROC					glDisableiEXT;
 	extern PFNGLBLENDEQUATIONIEXTPROC			glBlendEquationiEXT;
@@ -375,8 +375,7 @@ struct FOpenGLES : public FOpenGLBase
 
 	static FORCEINLINE void CopyImageSubData(GLuint SrcName, GLenum SrcTarget, GLint SrcLevel, GLint SrcX, GLint SrcY, GLint SrcZ, GLuint DstName, GLenum DstTarget, GLint DstLevel, GLint DstX, GLint DstY, GLint DstZ, GLsizei Width, GLsizei Height, GLsizei Depth)
 	{
-		check(bSupportsCopyImage);
-		glCopyImageSubDataEXT(SrcName, SrcTarget, SrcLevel, SrcX, SrcY, SrcZ, DstName, DstTarget, DstLevel, DstX, DstY, DstZ, Width, Height, Depth);
+		glCopyImageSubData(SrcName, SrcTarget, SrcLevel, SrcX, SrcY, SrcZ, DstName, DstTarget, DstLevel, DstX, DstY, DstZ, Width, Height, Depth);
 	}
 
 	static FORCEINLINE bool TexStorage2DMultisample(GLenum Target, GLsizei Samples, GLint InternalFormat, GLsizei Width, GLsizei Height, GLboolean FixedSampleLocations)
diff --git a/Engine/Source/Runtime/OpenGLDrv/Public/OpenGLResources.h b/Engine/Source/Runtime/OpenGLDrv/Public/OpenGLResources.h
--- a/Engine/Source/Runtime/OpenGLDrv/Public/OpenGLResources.h
+++ b/Engine/Source/Runtime/OpenGLDrv/Public/OpenGLResources.h
@@ -1372,11 +1372,6 @@ public:
 	 */
 	void CloneViaCopyImage(FOpenGLTexture* Src, uint32 InNumMips, int32 SrcOffset, int32 DstOffset);
 
-	/**
-	 * Clone texture from a source going via PBOs
-	 */
-	void CloneViaPBO(FOpenGLTexture* Src, uint32 InNumMips, int32 SrcOffset, int32 DstOffset);
-
 	/**
 	 * Resolved the specified face for a read Lock, for non-renderable, CPU readable surfaces this eliminates the readback inside Lock itself.
 	 */
