Fix CIS PVS warning.

diff --git a/Engine/Source/Runtime/OpenGLDrv/Private/OpenGLShaders.cpp b/Engine/Source/Runtime/OpenGLDrv/Private/OpenGLShaders.cpp
--- a/Engine/Source/Runtime/OpenGLDrv/Private/OpenGLShaders.cpp
+++ b/Engine/Source/Runtime/OpenGLDrv/Private/OpenGLShaders.cpp
@@ -2515,17 +2515,13 @@ FOpenGLLinkedProgram* FOpenGLDynamicRHI::LinkProgram( const FOpenGLLinkedProgram
 	check( (Config.Shaders[CrossCompiler::SHADER_STAGE_VERTEX].Resource == 0) != (Config.Shaders[CrossCompiler::SHADER_STAGE_COMPUTE].Resource == 0));
 	check( (Config.Shaders[CrossCompiler::SHADER_STAGE_PIXEL].Resource == 0) != (Config.Shaders[CrossCompiler::SHADER_STAGE_COMPUTE].Resource == 0));
 
-	GLuint Program = 0;
 	if (FOpenGLProgramBinaryCache::IsEnabled())
 	{
 		UE::OpenGL::CompilePendingShaders(Config);
 	}
 
-	if (Program == 0)
-	{
-		QUICK_SCOPE_CYCLE_COUNTER(STAT_glGenProgramPipelines);
-		FOpenGL::GenProgramPipelines(1, &Program);
-	}
+	GLuint Program = 0;
+	FOpenGL::GenProgramPipelines(1, &Program);
 
 	if (Config.Shaders[CrossCompiler::SHADER_STAGE_VERTEX].Resource)
 	{
