Added support for 32-bit depth on GL. Disabled by default, can be activated with r.OpenGL.Depth24Bit=0

diff --git a/Engine/Source/Runtime/OpenGLDrv/Private/OpenGLDevice.cpp b/Engine/Source/Runtime/OpenGLDrv/Private/OpenGLDevice.cpp
--- a/Engine/Source/Runtime/OpenGLDrv/Private/OpenGLDevice.cpp
+++ b/Engine/Source/Runtime/OpenGLDrv/Private/OpenGLDevice.cpp
@@ -64,6 +64,12 @@ static TAutoConsoleVariable<int32> CVarGLExtraDeletionLatency(
 	TEXT("Toggle the engine's deferred deletion queue for RHI resources. (default:1)"),
 	ECVF_RenderThreadSafe | ECVF_ReadOnly);
 
+static TAutoConsoleVariable<int32> CVarGLDepth24Bit(
+	TEXT("r.OpenGL.Depth24Bit"),
+	1,
+	TEXT("0: Use 32-bit float depth buffer \n1: Use 24-bit fixed point depth buffer (default)\n"),
+	ECVF_RenderThreadSafe | ECVF_ReadOnly);
+
 void OnQueryCreation( FOpenGLRenderQuery* Query )
 {
 	check(PrivateOpenGLDevicePtr);
@@ -587,6 +593,7 @@ static EOpenGLFormatCapabilities GetOpenGLFormatCapabilities(const FOpenGLTextur
 		break;
 
 	case GL_DEPTH24_STENCIL8:
+	case GL_DEPTH32F_STENCIL8:
 	case GL_DEPTH_COMPONENT16:
 	case GL_DEPTH_COMPONENT24:
 		Capabilities |= EOpenGLFormatCapabilities::Texture | EOpenGLFormatCapabilities::Render | EOpenGLFormatCapabilities::Filterable | EOpenGLFormatCapabilities::DepthStencil;
@@ -1103,7 +1110,16 @@ static void InitRHICapabilitiesForGL()
 	SetupTextureFormat( PF_Unknown,				FOpenGLTextureFormat( ));
 	SetupTextureFormat( PF_A32B32G32R32F,		FOpenGLTextureFormat( GL_RGBA32F,				GL_RGBA32F,				GL_RGBA,			GL_FLOAT,						false,			false));
 	SetupTextureFormat( PF_UYVY,				FOpenGLTextureFormat( ));
-	SetupTextureFormat( PF_DepthStencil,		FOpenGLTextureFormat( GL_DEPTH24_STENCIL8,		GL_NONE,				GL_DEPTH_STENCIL,	GL_UNSIGNED_INT_24_8,			false,			false));
+	if (CVarGLDepth24Bit.GetValueOnAnyThread() != 0)
+	{
+		SetupTextureFormat(PF_DepthStencil, FOpenGLTextureFormat(GL_DEPTH24_STENCIL8, GL_NONE, GL_DEPTH_STENCIL, GL_UNSIGNED_INT_24_8, false, false));
+		GPixelFormats[PF_DepthStencil].BlockBytes = 4;
+	}
+	else
+	{
+		SetupTextureFormat(PF_DepthStencil, FOpenGLTextureFormat(GL_DEPTH32F_STENCIL8, GL_NONE, GL_DEPTH_STENCIL, GL_FLOAT_32_UNSIGNED_INT_24_8_REV, false, false));
+		GPixelFormats[PF_DepthStencil].BlockBytes = 8;
+	}
 	SetupTextureFormat( PF_ShadowDepth,			FOpenGLTextureFormat( ShadowDepthFormat,		ShadowDepthFormat,		GL_DEPTH_COMPONENT,	GL_UNSIGNED_INT,				false,			false));
 	SetupTextureFormat( PF_D24,					FOpenGLTextureFormat( DepthFormat,				DepthFormat,			GL_DEPTH_COMPONENT,	GL_UNSIGNED_INT,				false,			false));
 	SetupTextureFormat( PF_A16B16G16R16,		FOpenGLTextureFormat( GL_RGBA16F,				GL_RGBA16F,				GL_RGBA,			GL_HALF_FLOAT,					false,			false));
@@ -1211,7 +1227,6 @@ static void InitRHICapabilitiesForGL()
 		SetupTextureFormat(PF_ASTC_12x12_HDR, FOpenGLTextureFormat(GL_COMPRESSED_RGBA_ASTC_12x12_KHR, GL_COMPRESSED_SRGB8_ALPHA8_ASTC_12x12_KHR, GL_RGBA, GL_HALF_FLOAT, true, false));
 	}
 	// Some formats need to know how large a block is.
-	GPixelFormats[ PF_DepthStencil		].BlockBytes	 = 4;
 	GPixelFormats[ PF_FloatRGB			].BlockBytes	 = 4;
 	GPixelFormats[ PF_FloatRGBA			].BlockBytes	 = 8;
 
