Allow RHIT on GL Windows
Set r.Occlusion.SingleRHIThreadStall=0 for GL Windows as it should be enabled only for D3D11

diff --git a/Engine/Source/Runtime/OpenGLDrv/Private/OpenGLDevice.cpp b/Engine/Source/Runtime/OpenGLDrv/Private/OpenGLDevice.cpp
--- a/Engine/Source/Runtime/OpenGLDrv/Private/OpenGLDevice.cpp
+++ b/Engine/Source/Runtime/OpenGLDrv/Private/OpenGLDevice.cpp
@@ -1004,11 +1004,7 @@ static void InitRHICapabilitiesForGL()
 	const GLint MinorVersion = FOpenGL::GetMinorVersion();
  
 	// Enable the OGL rhi thread if explicitly requested.
-	GRHISupportsRHIThread = (GMaxRHIFeatureLevel <= ERHIFeatureLevel::ES3_1 && CVarAllowRGLHIThread.GetValueOnAnyThread())
-#if WITH_EDITOR
-		&& !IsPCPlatform(GMaxRHIShaderPlatform)
-#endif
-		;
+	GRHISupportsRHIThread = (GMaxRHIFeatureLevel <= ERHIFeatureLevel::ES3_1 && CVarAllowRGLHIThread.GetValueOnAnyThread());
 
 	GRHISupportsMultithreadedResources = GRHISupportsRHIThread;
 
@@ -1311,6 +1307,15 @@ FOpenGLDynamicRHI::FOpenGLDynamicRHI()
 			}
 		}
 	}
+	
+	// Disable SingleRHIThreadStall for GL occlusion queiresn, which should be set for D3D11 only. Enabling it causes RT->RHIT deadlock
+	{
+		auto* CVar = IConsoleManager::Get().FindConsoleVariable(TEXT("r.Occlusion.SingleRHIThreadStall"));
+		if (CVar)
+		{
+			CVar->Set(0);
+		}
+	}
 
 	PrivateOpenGLDevicePtr = this;
 	GlobalUniformBuffers.AddZeroed(FUniformBufferStaticSlotRegistry::Get().GetSlotCount());
