Cleanup recursive shader initialization in platform RHIs
 - Removed overrides of PostInit where applicable.
 - Added GRequiredRecursiveShaders as a flags enum set by the platform RHI, which indicates which group of global shaders will be used with recursive command lists.
 - Replaced calls to CreateClearReplacementShaders() in ShaderCompiler.cpp with CreateRecursiveShaders(), enqueued in a render thread command.

Fix ensure calling IsInRenderingThread() from scene renderer tasks
 - The tasks may be launched inline, meaning they're still running on the render thread, so shouldn't be tagged as EParallelRenderingThread.
 - Use FOptionalTaskTagScope instead of FTaskTagScope.

General cleanup of OpenGL threading.
 - Removed shader proxy objects. Proxies were used to defer creation of shaders to the RHI thread. This is can be achieved simply with RHICmdList.EnqueueLambda() to defer those operations.
 - The FOpenGL...RHIThreadFence types were used to fence deletion of resources on the RT. All RHI resources are deleted on the RHIT now, so there is no need to track lifetimes in this way.
 - Replaced multiple uses of ALLOC_COMMAND_CL / bCanRunOnThisThread with RHICmdList.EnqueueLambda().

Removed types and functions:
 - FOpenGLAssertRHIThreadFence
 - FOpenGLRHIThreadResourceFence
 - TOpenGLResourceProxy
 - TOpenGLShaderProxy
 - TIsGLProxyObject_V
 - TIsGLResourceWithFence
 - FOpenGLVertexShaderProxy
 - FOpenGLGeometryShaderProxy
 - FOpenGLPixelShaderProxy
 - FOpenGLComputeShaderProxy
 - CheckRHITFence()
 - FOpenGLDynamicRHI::ResourceCastProxy()
 - FOpenGLDynamicRHI::ResourceCast_Unfenced()

diff --git a/Engine/Source/Runtime/D3D12RHI/Private/D3D12Commands.cpp b/Engine/Source/Runtime/D3D12RHI/Private/D3D12Commands.cpp
--- a/Engine/Source/Runtime/D3D12RHI/Private/D3D12Commands.cpp
+++ b/Engine/Source/Runtime/D3D12RHI/Private/D3D12Commands.cpp
@@ -99,35 +99,6 @@ static void BindUniformBuffer(FD3D12CommandContext& Context, FRHIShader* Shader,
 	#define D3D12_PLATFORM_SUPPORTS_RESOLVE_SHADERS 1
 #endif
 
-void FD3D12DynamicRHI::SetupRecursiveResources()
-{
-	auto ShaderMap = GetGlobalShaderMap(GMaxRHIFeatureLevel);
-
-	{
-		TShaderMapRef<FLongGPUTaskPS> PixelShader(ShaderMap);
-		PixelShader.GetPixelShader();
-	}
-
-	{
-		TShaderMapRef<FLongGPUTaskPS> PixelShader(ShaderMap);
-		PixelShader.GetPixelShader();
-	}
-
-	// TODO: Waiting to integrate MSAA fix for ResolveShader.h
-	if (!D3D12_PLATFORM_SUPPORTS_RESOLVE_SHADERS)
-		return;
-
-	TShaderMapRef<FResolveVS> ResolveVertexShader(ShaderMap);
-
-	{
-		TShaderMapRef<FResolveDepthPS> ResolvePixelShader_Depth(ShaderMap);
-		ResolvePixelShader_Depth.GetPixelShader();
-
-		TShaderMapRef<FResolveDepthPS> ResolvePixelShader_SingleSample(ShaderMap);
-		ResolvePixelShader_SingleSample.GetPixelShader();
-	}
-}
-
 // Vertex state.
 void FD3D12CommandContext::RHISetStreamSource(uint32 StreamIndex, FRHIBuffer* VertexBufferRHI, uint32 Offset)
 {
diff --git a/Engine/Source/Runtime/D3D12RHI/Private/D3D12RHI.cpp b/Engine/Source/Runtime/D3D12RHI/Private/D3D12RHI.cpp
--- a/Engine/Source/Runtime/D3D12RHI/Private/D3D12RHI.cpp
+++ b/Engine/Source/Runtime/D3D12RHI/Private/D3D12RHI.cpp
@@ -249,6 +249,17 @@ FD3D12DynamicRHI::FD3D12DynamicRHI(const TArray<TSharedPtr<FD3D12Adapter>>& Chos
 	GRHISupportsRHIOnTaskThread = true;
 }
 
+void FD3D12DynamicRHI::PostInit()
+{
+	if (GRHISupportsRayTracing)
+	{
+		for (TSharedPtr<FD3D12Adapter>& Adapter : ChosenAdapters)
+		{
+			Adapter->InitializeRayTracing();
+		}
+	}
+}
+
 FD3D12DynamicRHI::~FD3D12DynamicRHI()
 {
 	UE_LOG(LogD3D12RHI, Log, TEXT("~FD3D12DynamicRHI"));
diff --git a/Engine/Source/Runtime/D3D12RHI/Private/D3D12RHIPrivate.h b/Engine/Source/Runtime/D3D12RHI/Private/D3D12RHIPrivate.h
--- a/Engine/Source/Runtime/D3D12RHI/Private/D3D12RHIPrivate.h
+++ b/Engine/Source/Runtime/D3D12RHI/Private/D3D12RHIPrivate.h
@@ -764,8 +764,6 @@ protected:
 
 	void ReadSurfaceDataMSAARaw(FRHITexture* TextureRHI, FIntRect Rect, TArray<uint8>& OutData, FReadSurfaceDataFlags InFlags);
 
-	void SetupRecursiveResources();
-
 	// This should only be called by Dynamic RHI member functions
 	inline FD3D12Device* GetRHIDevice(uint32 GPUIndex) const
 	{
diff --git a/Engine/Source/Runtime/D3D12RHI/Private/Windows/WindowsD3D12Device.cpp b/Engine/Source/Runtime/D3D12RHI/Private/Windows/WindowsD3D12Device.cpp
--- a/Engine/Source/Runtime/D3D12RHI/Private/Windows/WindowsD3D12Device.cpp
+++ b/Engine/Source/Runtime/D3D12RHI/Private/Windows/WindowsD3D12Device.cpp
@@ -1731,33 +1731,6 @@ void FD3D12DynamicRHI::Init()
 	GIsRHIInitialized = true;
 }
 
-void FD3D12DynamicRHI::PostInit()
-{
-	if (!FPlatformProperties::RequiresCookedData() && (GRHISupportsRayTracing || GRHISupportsRHIThread))
-	{
-		// Make sure all global shaders are complete at this point
-		extern RENDERCORE_API const int32 GlobalShaderMapId;
-
-		TArray<int32> ShaderMapIds;
-		ShaderMapIds.Add(GlobalShaderMapId);
-
-		GShaderCompilingManager->FinishCompilation(TEXT("Global"), ShaderMapIds);
-	}
-
-	if (GRHISupportsRayTracing)
-	{
-		for (TSharedPtr<FD3D12Adapter>& Adapter : ChosenAdapters)
-		{
-			Adapter->InitializeRayTracing();
-		}
-	}
-
-	if (GRHISupportsRHIThread)
-	{
-		SetupRecursiveResources();
-	}
-}
-
 bool FD3D12DynamicRHI::IsQuadBufferStereoEnabled() const
 {
 	return bIsQuadBufferStereoEnabled;
diff --git a/Engine/Source/Runtime/Engine/Private/ShaderCompiler/ShaderCompiler.cpp b/Engine/Source/Runtime/Engine/Private/ShaderCompiler/ShaderCompiler.cpp
--- a/Engine/Source/Runtime/Engine/Private/ShaderCompiler/ShaderCompiler.cpp
+++ b/Engine/Source/Runtime/Engine/Private/ShaderCompiler/ShaderCompiler.cpp
@@ -7920,7 +7920,10 @@ void CompileGlobalShaderMap(EShaderPlatform Platform, const ITargetPlatform* Tar
 		// While we're early in the game's startup, create certain global shaders that may be later created on random threads otherwise. 
 		if (!bShaderMapIsBeingCompiled && !GRHISupportsMultithreadedShaderCreation)
 		{
-			CreateClearReplacementShaders();
+			ENQUEUE_RENDER_COMMAND(CreateRecursiveShaders)([](FRHICommandListImmediate&)
+			{
+				CreateRecursiveShaders();
+			});
 		}
 	}
 }
@@ -8513,7 +8516,10 @@ void ProcessCompiledGlobalShaders(const TArray<FShaderCommonCompileJobPtr>& Comp
 
 			if (!GRHISupportsMultithreadedShaderCreation && Platform == GMaxRHIShaderPlatform)
 			{
-				CreateClearReplacementShaders();
+				ENQUEUE_RENDER_COMMAND(CreateRecursiveShaders)([](FRHICommandListImmediate&)
+				{
+					CreateRecursiveShaders();
+				});
 			}
 		}
 	}
diff --git a/Engine/Source/Runtime/OpenGLDrv/Private/OpenGLCommands.cpp b/Engine/Source/Runtime/OpenGLDrv/Private/OpenGLCommands.cpp
--- a/Engine/Source/Runtime/OpenGLDrv/Private/OpenGLCommands.cpp
+++ b/Engine/Source/Runtime/OpenGLDrv/Private/OpenGLCommands.cpp
@@ -1762,7 +1762,7 @@ void FOpenGLDynamicRHI::SetupVertexArrays(FOpenGLContextState& ContextState, uin
 
 	check(IsValidRef(PendingState.BoundShaderState));
 	FOpenGLVertexDeclaration* VertexDeclaration = PendingState.BoundShaderState->VertexDeclaration;
-	const CrossCompiler::FShaderBindingInOutMask& AttributeMask = PendingState.BoundShaderState->GetVertexShader()->GetGLResourceObject()->Bindings.InOutMask;
+	const CrossCompiler::FShaderBindingInOutMask& AttributeMask = PendingState.BoundShaderState->GetVertexShader()->Bindings.InOutMask;
 
 	if (ContextState.VertexDecl != VertexDeclaration || AttributeMask.Bitmask != ContextState.VertexAttrs_EnabledBits)
 	{
@@ -2072,11 +2072,11 @@ void FOpenGLDynamicRHI::CommitNonComputeShaderConstants()
 	FOpenGLLinkedProgram* LinkedProgram = PendingState.BoundShaderState->LinkedProgram;
 	if (GUseEmulatedUniformBuffers)
 	{
-		PendingState.ShaderParameters[CrossCompiler::SHADER_STAGE_VERTEX].CommitPackedUniformBuffers(LinkedProgram, CrossCompiler::SHADER_STAGE_VERTEX, PendingState.BoundUniformBuffers[SF_Vertex], PendingState.BoundShaderState->GetVertexShader()->GetGLResourceObject()->UniformBuffersCopyInfo);
-		PendingState.ShaderParameters[CrossCompiler::SHADER_STAGE_PIXEL].CommitPackedUniformBuffers(LinkedProgram, CrossCompiler::SHADER_STAGE_PIXEL, PendingState.BoundUniformBuffers[SF_Pixel], PendingState.BoundShaderState->GetPixelShader()->GetGLResourceObject()->UniformBuffersCopyInfo);
+		PendingState.ShaderParameters[CrossCompiler::SHADER_STAGE_VERTEX].CommitPackedUniformBuffers(LinkedProgram, CrossCompiler::SHADER_STAGE_VERTEX, PendingState.BoundUniformBuffers[SF_Vertex], PendingState.BoundShaderState->GetVertexShader()->UniformBuffersCopyInfo);
+		PendingState.ShaderParameters[CrossCompiler::SHADER_STAGE_PIXEL].CommitPackedUniformBuffers(LinkedProgram, CrossCompiler::SHADER_STAGE_PIXEL, PendingState.BoundUniformBuffers[SF_Pixel], PendingState.BoundShaderState->GetPixelShader()->UniformBuffersCopyInfo);
 		if (PendingState.BoundShaderState->GetGeometryShader())
 		{
-			PendingState.ShaderParameters[CrossCompiler::SHADER_STAGE_GEOMETRY].CommitPackedUniformBuffers(LinkedProgram, CrossCompiler::SHADER_STAGE_GEOMETRY, PendingState.BoundUniformBuffers[SF_Geometry], PendingState.BoundShaderState->GetGeometryShader()->GetGLResourceObject()->UniformBuffersCopyInfo);
+			PendingState.ShaderParameters[CrossCompiler::SHADER_STAGE_GEOMETRY].CommitPackedUniformBuffers(LinkedProgram, CrossCompiler::SHADER_STAGE_GEOMETRY, PendingState.BoundUniformBuffers[SF_Geometry], PendingState.BoundShaderState->GetGeometryShader()->UniformBuffersCopyInfo);
 		}
 	}
 	
@@ -2137,12 +2137,12 @@ template <> FORCEINLINE uint32 GetNumUAVUnits<SF_Pixel>()	{ return FOpenGL::GetM
 template <> FORCEINLINE uint32 GetNumUAVUnits<SF_Vertex>()	{ return FOpenGL::GetMaxPixelUAVUnits(); }
 
 template <class ShaderType>
-FORCEINLINE void FOpenGLDynamicRHI::SetResourcesFromTables(ShaderType* RESTRICT Shader)
+FORCEINLINE void FOpenGLDynamicRHI::SetResourcesFromTables(ShaderType* Shader)
 {
 	checkSlow(Shader);
 	VERIFY_GL_SCOPE();
 
-	static constexpr EShaderFrequency Frequency = static_cast<EShaderFrequency>(ShaderType::ContainedGLType::StaticFrequency);
+	static constexpr EShaderFrequency Frequency = ShaderType::Frequency;
 
 	struct FUniformResourceBinder
 	{
@@ -2228,7 +2228,7 @@ FORCEINLINE void FOpenGLDynamicRHI::SetResourcesFromTables(ShaderType* RESTRICT
 	UE::RHICore::SetResourcesFromTables(
 		  FUniformResourceBinder { *this }
 		, *Shader
-		, Shader->GetGLResourceObject()->Bindings.ShaderResourceTable
+		, Shader->Bindings.ShaderResourceTable
 		, PendingState.DirtyUniformBuffers[Frequency]
 		, PendingState.BoundUniformBuffers[Frequency]
 #if ENABLE_RHI_VALIDATION
@@ -2271,7 +2271,7 @@ void FOpenGLDynamicRHI::CommitGraphicsResourceTablesInner()
 }
 
 
-void FOpenGLDynamicRHI::CommitComputeResourceTables(FOpenGLComputeShaderProxy* ComputeShader)
+void FOpenGLDynamicRHI::CommitComputeResourceTables(FOpenGLComputeShader* ComputeShader)
 {
 	VERIFY_GL_SCOPE();
 
@@ -2869,14 +2869,12 @@ void FOpenGLDynamicRHI::RHIDispatchComputeShader(uint32 ThreadGroupCountX, uint3
 		return;
 	}
 
-
 	VERIFY_GL_SCOPE();
 		
 	FRHIComputeShader* ComputeShaderRHI = PendingState.CurrentComputeShader;
 	check(ComputeShaderRHI);
 
-	FOpenGLComputeShaderProxy* ComputeShaderProxy = ResourceCastProxy(ComputeShaderRHI);
-	FOpenGLComputeShader* ComputeShader = ComputeShaderProxy->GetGLResourceObject();
+	FOpenGLComputeShader* ComputeShader = ResourceCast(ComputeShaderRHI);
 
 	if (ComputeShader->LinkedProgram == nullptr)
 	{
@@ -2887,7 +2885,7 @@ void FOpenGLDynamicRHI::RHIDispatchComputeShader(uint32 ThreadGroupCountX, uint3
 	GPUProfilingData.RegisterGPUDispatch(FIntVector(ThreadGroupCountX, ThreadGroupCountY, ThreadGroupCountZ));	
 
 	BindPendingComputeShaderState(ContextState, ComputeShader);
-	CommitComputeResourceTables(ComputeShaderProxy);
+	CommitComputeResourceTables(ComputeShader);
 	SetupTexturesForDraw(ContextState, ComputeShader, FOpenGL::GetMaxComputeTextureImageUnits());
 	SetupUAVsForCompute(ContextState, ComputeShader);
 	CommitComputeShaderConstants(ComputeShader);
diff --git a/Engine/Source/Runtime/OpenGLDrv/Private/OpenGLDevice.cpp b/Engine/Source/Runtime/OpenGLDrv/Private/OpenGLDevice.cpp
--- a/Engine/Source/Runtime/OpenGLDrv/Private/OpenGLDevice.cpp
+++ b/Engine/Source/Runtime/OpenGLDrv/Private/OpenGLDevice.cpp
@@ -243,16 +243,12 @@ void FOpenGLDynamicRHI::RHIBeginScene()
 	}
 
 	BeginSceneContextType = (int32)PlatformOpenGLCurrentContext(PlatformDevice);
-
-	// recache NULL shader as it can change with ODSC
-	NULLPixelShaderRHI = GetNULLPixelShader();
 }
 
 void FOpenGLDynamicRHI::RHIEndScene()
 {
 	ResourceTableFrameCounter = INDEX_NONE;
 	BeginSceneContextType = CONTEXT_Other;
-	NULLPixelShaderRHI = nullptr;
 }
 
 #if PLATFORM_ANDROID
@@ -1248,7 +1244,6 @@ FOpenGLDynamicRHI::FOpenGLDynamicRHI()
 ,   BeginSceneContextType(CONTEXT_Other)
 ,	PlatformDevice(NULL)
 ,	GPUProfilingData(this)
-,	NULLPixelShaderRHI(nullptr)
 {
 	check(Singleton == nullptr);
 	Singleton = this;
@@ -1365,6 +1360,7 @@ void FOpenGLDynamicRHI::Init()
 	VERIFY_GL_SCOPE();
 
 	GRHISupportsMultithreadedShaderCreation = false;
+	GRequiredRecursiveShaders = ERecursiveShader::Resolve | ERecursiveShader::Clear | ERecursiveShader::Null;
 
 	FOpenGLProgramBinaryCache::Initialize();
 
@@ -1421,14 +1417,6 @@ void FOpenGLDynamicRHI::Init()
 	GIsRHIInitialized = true;
 }
 
-void FOpenGLDynamicRHI::PostInit()
-{
-	if (GRHISupportsRHIThread)
-	{
-		SetupRecursiveResources();
-	}
-}
-
 void FOpenGLDynamicRHI::Shutdown()
 {
 	check(IsInGameThread() && IsInRenderingThread()); // require that the render thread has been shut down
diff --git a/Engine/Source/Runtime/OpenGLDrv/Private/OpenGLDrv.cpp b/Engine/Source/Runtime/OpenGLDrv/Private/OpenGLDrv.cpp
--- a/Engine/Source/Runtime/OpenGLDrv/Private/OpenGLDrv.cpp
+++ b/Engine/Source/Runtime/OpenGLDrv/Private/OpenGLDrv.cpp
@@ -152,11 +152,6 @@ bool FOpenGLDynamicRHI::RHIRequiresComputeGenerateMips() const
 
 void FOpenGLGPUProfiler::BeginFrame(FOpenGLDynamicRHI* InRHI)
 {
-	if (!bIntialized)
-	{
-		bIntialized = true;
-		InitResources();
-	}
 	if (NestedFrameCount++>0)
 	{
 		// guard against nested Begin/EndFrame calls.
@@ -391,17 +386,13 @@ void FOpenGLGPUProfiler::EndFrame()
 
 void FOpenGLGPUProfiler::Cleanup()
 {
-	if (bIntialized)
+	for (int32 Index = 0; Index < MAX_GPUFRAMEQUERIES; ++Index)
 	{
-		for (int32 Index = 0; Index < MAX_GPUFRAMEQUERIES; ++Index)
-		{
-			DisjointGPUFrameTimeQuery[Index].ReleaseResources();
-		}
-
-		FrameTiming.ReleaseResources();
-		NestedFrameCount = 0;
-		bIntialized = false;
+		DisjointGPUFrameTimeQuery[Index].ReleaseResources();
 	}
+
+	FrameTiming.ReleaseResources();
+	NestedFrameCount = 0;
 }
 
 /** Start this frame of per tracking */
diff --git a/Engine/Source/Runtime/OpenGLDrv/Private/OpenGLQuery.cpp b/Engine/Source/Runtime/OpenGLDrv/Private/OpenGLQuery.cpp
--- a/Engine/Source/Runtime/OpenGLDrv/Private/OpenGLQuery.cpp
+++ b/Engine/Source/Runtime/OpenGLDrv/Private/OpenGLQuery.cpp
@@ -637,43 +637,22 @@ FOpenGLRenderQuery::FOpenGLRenderQuery(ERenderQueryType InQueryType)
 	, bInvalidResource(true)
 	, QueryType(InQueryType)
 {
-	FRHICommandListImmediate& RHICmdList = FRHICommandListExecutor::GetImmediateCommandList();
-
-	const bool bCanRunOnThisThread = RHICmdList.Bypass() || (!IsRunningRHIInSeparateThread() && IsInRenderingThread()) || IsInRHIThread();
-
-	if (bCanRunOnThisThread)
-	{
-		AcquireResource();
-	}
-	else
-	{
-		CreationFence.Reset();
-		ALLOC_COMMAND_CL(RHICmdList, FRHICommandGLCommand)([=]() {AcquireResource(); CreationFence.WriteAssertFence(); });
-		CreationFence.SetRHIThreadFence();
-	}
+	check(IsInRenderingThread());
+	FRHICommandListExecutor::GetImmediateCommandList().EnqueueLambda(
+		[this](FRHICommandListImmediate&) { AcquireResource(); }
+	);
 }
 
 
 FOpenGLRenderQuery::~FOpenGLRenderQuery()
 {
-
+	VERIFY_GL_SCOPE();
 	OnQueryDeletion( this );
 
-	FRHICommandListImmediate& RHICmdList = FRHICommandListExecutor::GetImmediateCommandList();
-
-	const bool bCanRunOnThisThread = RHICmdList.Bypass() || (!IsRunningRHIInSeparateThread() && IsInRenderingThread()) || IsInRHIThread();
 	if (Resource && !bInvalidResource)
 	{
 		bInvalidResource = true;
-		if (bCanRunOnThisThread)
-		{
-			ReleaseResource(Resource, ResourceContext);
-		}
-		else
-		{
-			CreationFence.WaitFenceRenderThreadOnly();
-			ALLOC_COMMAND_CL(RHICmdList, FRHICommandGLCommand)([Resource = Resource, ResourceContext = ResourceContext]() {VERIFY_GL_SCOPE(); ReleaseResource(Resource, ResourceContext); });
-		}
+		ReleaseResource(Resource, ResourceContext);
 	}
 }
 
diff --git a/Engine/Source/Runtime/OpenGLDrv/Private/OpenGLShaders.cpp b/Engine/Source/Runtime/OpenGLDrv/Private/OpenGLShaders.cpp
--- a/Engine/Source/Runtime/OpenGLDrv/Private/OpenGLShaders.cpp
+++ b/Engine/Source/Runtime/OpenGLDrv/Private/OpenGLShaders.cpp
@@ -193,26 +193,6 @@ static void SetDeletedProgramStats(GLuint Program)
 	GNumPrograms--;
 }
 
-// Create any resources that are required by internal ogl rhi functions.
-void FOpenGLDynamicRHI::SetupRecursiveResources()
-{
-	NULLPixelShaderRHI = GetNULLPixelShader();
-}
-
-FRHIPixelShader* FOpenGLDynamicRHI::GetNULLPixelShader() const
-{
-	if (NULLPixelShaderRHI)
-	{
-		return NULLPixelShaderRHI;
-	}
-	else
-	{
-		auto ShaderMap = GetGlobalShaderMap(GMaxRHIFeatureLevel);
-		TShaderMapRef<FNULLPS> PixelShader(ShaderMap);
-		return PixelShader.GetPixelShader();
-	}
-}
-
 const uint32 SizeOfFloat4 = 16;
 const uint32 NumFloatsInFloat4 = 4;
 
@@ -634,40 +614,36 @@ static void GLSLToPlatform(const FOpenGLCodeHeader& Header, GLenum TypeEnum, FAn
 
 /**
  * Compiles an OpenGL shader using the given GLSL microcode.
- * @returns the compiled shader upon success.
  */
-template <typename ShaderType, typename TRHIType>
-ShaderType* CompileOpenGLShader(const FOpenGLCodeHeader& Header, const FOpenGLCompiledShaderKey& SourceKey, const FSHAHash& LibraryHash, TRHIType* RHIShader)
+void FOpenGLShader::Compile(const FOpenGLCodeHeader& Header, const FSHAHash& LibraryHash, GLenum TypeEnum)
 {
 	SCOPE_CYCLE_COUNTER(STAT_OpenGLShaderCompileTime);
 	VERIFY_GL_SCOPE();
 
-	ShaderType* Shader = nullptr;
 	{
 		FLibraryShaderCacheValue *Val = GetOpenGLCompiledLibraryShaderCache().Find(LibraryHash);
 		if (Val)
 		{
-			Shader = new ShaderType();
-			Shader->Resource = Val->GLShader;
-			Shader->Bindings = Val->Header->Bindings;
-			Shader->UniformBuffersCopyInfo = Val->Header->UniformBuffersCopyInfo;
-			Shader->StaticSlots = Val->StaticSlots;
+			Resource = Val->GLShader;
+			Bindings = Val->Header->Bindings;
+			UniformBuffersCopyInfo = Val->Header->UniformBuffersCopyInfo;
+			StaticSlots = Val->StaticSlots;
 #if DEBUG_GL_SHADERS
-			Shader->GlslCode = Val->GlslCode;
-			Shader->GlslCodeString = (ANSICHAR*)Shader->GlslCode.GetData();
+			GlslCode = Val->GlslCode;
+			GlslCodeString = (ANSICHAR*)GlslCode.GetData();
 #endif
-			return Shader;
+			return;
 		}
 	}
+
 	// Find the existing compiled shader in the cache.
 	FScopeLock Lock(&GCompiledShaderCacheCS);
-	TSharedPtr<FOpenGLCompiledShaderValue> FoundShader = GetOpenGLCompiledShaderCache().FindRef(SourceKey);
+	TSharedPtr<FOpenGLCompiledShaderValue> FoundShader = GetOpenGLCompiledShaderCache().FindRef(ShaderCodeKey);
 	check(FoundShader);
-	GLuint Resource = FoundShader->Resource;
+	Resource = FoundShader->Resource;
 
 	if (Resource == 0)
 	{
-		const GLenum TypeEnum = ShaderType::TypeEnum;
 		Resource = FOpenGL::CreateShader(TypeEnum);
 		const FOpenGLShaderDeviceCapabilities& Capabilities = GetOpenGLShaderDeviceCapabilities();
 
@@ -687,18 +663,16 @@ ShaderType* CompileOpenGLShader(const FOpenGLCodeHeader& Header, const FOpenGLCo
 		FoundShader->Resource = Resource;
 	}
 
-	Shader = new ShaderType();
-	Shader->Resource = Resource;
-	Shader->Bindings = Header.Bindings;
-	Shader->UniformBuffersCopyInfo = Header.UniformBuffersCopyInfo;
+	Bindings = Header.Bindings;
+	UniformBuffersCopyInfo = Header.UniformBuffersCopyInfo;
 
-	UE::RHICore::InitStaticUniformBufferSlots(Shader->StaticSlots, Shader->Bindings.ShaderResourceTable);
+	UE::RHICore::InitStaticUniformBufferSlots(StaticSlots, Bindings.ShaderResourceTable);
 
 #if DEBUG_GL_SHADERS
 	{
 		const FAnsiCharArray glslPlatformCode = FoundShader->GetUncompressedShader();
-		Shader->GlslCode = glslPlatformCode;
-		Shader->GlslCodeString = (ANSICHAR*)Shader->GlslCode.GetData();
+		GlslCode = glslPlatformCode;
+		GlslCodeString = (ANSICHAR*)GlslCode.GetData();
 	}
 #endif
 	if (LibraryHash != FSHAHash() && !GetOpenGLCompiledLibraryShaderCache().Contains(LibraryHash))
@@ -707,16 +681,14 @@ ShaderType* CompileOpenGLShader(const FOpenGLCodeHeader& Header, const FOpenGLCo
 		Val.GLShader = Resource;
 		Val.Header = new FOpenGLCodeHeader;
 		*Val.Header = Header;
-		Val.ShaderCrc = SourceKey.GetCodeCRC();
-		Val.StaticSlots = Shader->StaticSlots;
+		Val.ShaderCrc = ShaderCodeKey.GetCodeCRC();
+		Val.StaticSlots = StaticSlots;
 #if DEBUG_GL_SHADERS
 		Val.GlslCode = FoundShader->GetUncompressedShader();
-		Val.GlslCodeString = (ANSICHAR*)Shader->GlslCode.GetData();
+		Val.GlslCodeString = (ANSICHAR*)GlslCode.GetData();
 #endif
 		GetOpenGLCompiledLibraryShaderCache().Add(LibraryHash, Val);
 	}
-
-	return Shader;
 }
 
 void OPENGLDRV_API GetCurrentOpenGLShaderDeviceCapabilities(FOpenGLShaderDeviceCapabilities& Capabilities)
@@ -941,28 +913,24 @@ static ANSICHAR* SetIndex(ANSICHAR* Str, int32 Offset, int32 Index)
 	return Str;
 }
 
-template<typename RHIType>
-RHIType* CreateProxyShader(TArrayView<const uint8> Code, const FSHAHash& Hash)
+FOpenGLShader::FOpenGLShader(FRHICommandListImmediate& RHICmdList, TArrayView<const uint8> Code, const FSHAHash& Hash, GLenum TypeEnum)
 {
-	typedef typename TOpenGLResourceTraits<RHIType>::TConcreteType TOGLProxyType;
+	FMemory::Memzero(&Bindings, sizeof(Bindings));
 
-	FRHICommandListImmediate& RHICmdList = FRHICommandListExecutor::GetImmediateCommandList();
 	FShaderCodeReader ShaderCode(Code);
 
-	const GLenum TypeEnum = TOGLProxyType::ContainedGLType::TypeEnum;
 	FMemoryReaderView Ar(Code, true);
 
 	Ar.SetLimitSize(ShaderCode.GetActualShaderCodeSize());
 
 	FOpenGLCodeHeader Header = { 0 };
 	Ar << Header;
-	// Suppress static code analysis warning about a potential comparison of two constants
-	CA_SUPPRESS(6326);
+
 	if (Header.GlslMarker != 0x474c534c
-		|| (TypeEnum == GL_VERTEX_SHADER && Header.FrequencyMarker != 0x5653)
+		|| (TypeEnum == GL_VERTEX_SHADER   && Header.FrequencyMarker != 0x5653)
 		|| (TypeEnum == GL_FRAGMENT_SHADER && Header.FrequencyMarker != 0x5053)
 		|| (TypeEnum == GL_GEOMETRY_SHADER && Header.FrequencyMarker != 0x4753)
-		|| (TypeEnum == GL_COMPUTE_SHADER && Header.FrequencyMarker != 0x4353)
+		|| (TypeEnum == GL_COMPUTE_SHADER  && Header.FrequencyMarker != 0x4353)
 		)
 	{
 		UE_LOG(LogRHI, Fatal,
@@ -970,7 +938,7 @@ RHIType* CreateProxyShader(TArrayView<const uint8> Code, const FSHAHash& Hash)
 			Header.GlslMarker,
 			Header.FrequencyMarker
 		);
-		return nullptr;
+		return;
 	}
 
 	int32 CodeOffset = Ar.Tell();
@@ -983,7 +951,7 @@ RHIType* CreateProxyShader(TArrayView<const uint8> Code, const FSHAHash& Hash)
  	FAnsiCharArray GlslCodeOriginal;
  	AppendCString(GlslCodeOriginal, (ANSICHAR*)Code.GetData() + CodeOffset);
 	uint32 CodeCRC = FCrc::MemCrc32(GlslCodeOriginal.GetData(), GlslCodeOriginal.Num());
-	FOpenGLCompiledShaderKey ShaderCodeKey(TypeEnum, GlslCodeOriginal.Num(), CodeCRC);
+	ShaderCodeKey = FOpenGLCompiledShaderKey(TypeEnum, GlslCodeOriginal.Num(), CodeCRC);
 
 #if CHECK_FOR_GL_SHADERS_TO_REPLACE
 	{
@@ -1041,46 +1009,28 @@ RHIType* CreateProxyShader(TArrayView<const uint8> Code, const FSHAHash& Hash)
 #endif
 	}
 
-	RHIType* RHIShader = nullptr;
-	if (ShouldRunGLRenderContextOpOnThisThread(RHICmdList))
+	RHICmdList.EnqueueLambda([this, Header = MoveTemp(Header), Hash, TypeEnum](FRHICommandListImmediate&)
 	{
-		RHIShader = new TOGLProxyType(ShaderCodeKey, [&](RHIType* OwnerRHI)
-		{
-			return CompileOpenGLShader<typename TOGLProxyType::ContainedGLType>(Header, ShaderCodeKey, Hash, OwnerRHI);
-		});
-	}
-	else
-	{
-		// take a copy of the code for RHIT version.
-		RHIShader = new TOGLProxyType(ShaderCodeKey, [Header = Header, ShaderCodeKey= ShaderCodeKey, Hash](RHIType* OwnerRHI)
-		{
-			return CompileOpenGLShader<typename TOGLProxyType::ContainedGLType>(Header, ShaderCodeKey, Hash, OwnerRHI);
-		});
-	}
-
-#if !(UE_BUILD_SHIPPING || UE_BUILD_TEST)
-	if (RHIShader)
-	{
-		RHIShader->Debug.ShaderName = ShaderCode.FindOptionalData(FShaderCodeName::Key);
-	}
-#endif
-
-	return RHIShader;
+		Compile(Header, Hash, TypeEnum);
+	});
 }
 
 FVertexShaderRHIRef FOpenGLDynamicRHI::RHICreateVertexShader(TArrayView<const uint8> Code, const FSHAHash& Hash)
 {
-	return CreateProxyShader<FRHIVertexShader>(Code, Hash);
+	check(IsInRenderingThread());
+	return new FOpenGLVertexShader(FRHICommandListExecutor::GetImmediateCommandList(), Code, Hash);
 }
 
 FPixelShaderRHIRef FOpenGLDynamicRHI::RHICreatePixelShader(TArrayView<const uint8> Code, const FSHAHash& Hash)
 {
-	return CreateProxyShader<FRHIPixelShader>(Code, Hash);
+	check(IsInRenderingThread());
+	return new FOpenGLPixelShader(FRHICommandListExecutor::GetImmediateCommandList(), Code, Hash);
 }
 
 FGeometryShaderRHIRef FOpenGLDynamicRHI::RHICreateGeometryShader(TArrayView<const uint8> Code, const FSHAHash& Hash)
 {
-	return CreateProxyShader<FRHIGeometryShader>(Code, Hash);
+	check(IsInRenderingThread());
+	return new FOpenGLGeometryShader(FRHICommandListExecutor::GetImmediateCommandList(), Code, Hash);
 }
 
 static void MarkShaderParameterCachesDirty(FOpenGLShaderParameterCache* ShaderParameters, bool UpdateCompute)
@@ -2577,7 +2527,7 @@ bool FOpenGLDynamicRHI::LinkComputeShader(FRHIComputeShader* ComputeShaderRHI, F
 	Config.Shaders[CrossCompiler::SHADER_STAGE_COMPUTE].Resource = ComputeShader->Resource;
 	Config.Shaders[CrossCompiler::SHADER_STAGE_COMPUTE].Bindings = ComputeShader->Bindings;
 	Config.ProgramKey.ShaderHashes[CrossCompiler::SHADER_STAGE_COMPUTE] = ComputeShaderRHI->GetHash();
-	Config.Shaders[CrossCompiler::SHADER_STAGE_COMPUTE].ShaderKey = FOpenGLDynamicRHI::ResourceCastProxy(ComputeShaderRHI)->GetCompiledShaderKey();
+	Config.Shaders[CrossCompiler::SHADER_STAGE_COMPUTE].ShaderKey = FOpenGLDynamicRHI::ResourceCast(ComputeShaderRHI)->ShaderCodeKey;
 
 	ComputeShader->LinkedProgram = GetOpenGLProgramsCache().Find(Config.ProgramKey, true);
 
@@ -2645,7 +2595,8 @@ FOpenGLLinkedProgram* FOpenGLDynamicRHI::GetLinkedComputeProgram(FRHIComputeShad
 
 FComputeShaderRHIRef FOpenGLDynamicRHI::RHICreateComputeShader(TArrayView<const uint8> Code, const FSHAHash& Hash)
 {
-	return CreateProxyShader<FRHIComputeShader>(Code, Hash);
+	check(IsInRenderingThread());
+	return new FOpenGLComputeShader(FRHICommandListExecutor::GetImmediateCommandList(), Code, Hash);
 }
 
 template<class TOpenGLStage>
@@ -2697,9 +2648,6 @@ static void BindShaderStage(FOpenGLLinkedProgramConfiguration& Config, CrossComp
 
 	check(NextStageShader && PrevStageShader);
 
-	typedef typename TOpenGLResourceTraits<TOpenGLStage0RHI>::TConcreteType::ContainedGLType TOpenGLStage0;
-	typedef typename TOpenGLResourceTraits<TOpenGLStage1RHI>::TConcreteType::ContainedGLType TOpenGLStage1;
-
 	FOpenGLLinkedProgramConfiguration::ShaderInfo& ShaderInfo = Config.Shaders[NextStage];
 	FOpenGLLinkedProgramConfiguration::ShaderInfo& PrevInfo = Config.Shaders[PrevStage];
 
@@ -2726,7 +2674,7 @@ static FOpenGLLinkedProgramConfiguration CreateConfig(FRHIVertexShader* VertexSh
 	// Fill-in the configuration
 	Config.Shaders[CrossCompiler::SHADER_STAGE_VERTEX].Bindings = VertexShader->Bindings;
 	Config.Shaders[CrossCompiler::SHADER_STAGE_VERTEX].Resource = VertexShader->Resource;
-	Config.Shaders[CrossCompiler::SHADER_STAGE_VERTEX].ShaderKey = FOpenGLDynamicRHI::ResourceCastProxy(VertexShaderRHI)->GetCompiledShaderKey();
+	Config.Shaders[CrossCompiler::SHADER_STAGE_VERTEX].ShaderKey = FOpenGLDynamicRHI::ResourceCast(VertexShaderRHI)->ShaderCodeKey;
 	Config.ProgramKey.ShaderHashes[CrossCompiler::SHADER_STAGE_VERTEX] = VertexShaderRHI->GetHash();
 
 	if (GeometryShaderRHI)
@@ -2734,7 +2682,7 @@ static FOpenGLLinkedProgramConfiguration CreateConfig(FRHIVertexShader* VertexSh
 		check(VertexShader);
 		BindShaderStage(Config, CrossCompiler::SHADER_STAGE_GEOMETRY, GeometryShaderRHI, CrossCompiler::SHADER_STAGE_VERTEX, VertexShaderRHI);
 		Config.ProgramKey.ShaderHashes[CrossCompiler::SHADER_STAGE_GEOMETRY] = GeometryShaderRHI->GetHash();
-		Config.Shaders[CrossCompiler::SHADER_STAGE_GEOMETRY].ShaderKey = FOpenGLDynamicRHI::ResourceCastProxy(GeometryShaderRHI)->GetCompiledShaderKey();
+		Config.Shaders[CrossCompiler::SHADER_STAGE_GEOMETRY].ShaderKey = FOpenGLDynamicRHI::ResourceCast(GeometryShaderRHI)->ShaderCodeKey;
 	}
 
 	check(GeometryShaderRHI || VertexShaderRHI);
@@ -2747,7 +2695,7 @@ static FOpenGLLinkedProgramConfiguration CreateConfig(FRHIVertexShader* VertexSh
 		BindShaderStage(Config, CrossCompiler::SHADER_STAGE_PIXEL, PixelShaderRHI, CrossCompiler::SHADER_STAGE_VERTEX, VertexShaderRHI);
 	}
 	Config.ProgramKey.ShaderHashes[CrossCompiler::SHADER_STAGE_PIXEL] = PixelShaderRHI->GetHash();
-	Config.Shaders[CrossCompiler::SHADER_STAGE_PIXEL].ShaderKey = FOpenGLDynamicRHI::ResourceCastProxy(PixelShaderRHI)->GetCompiledShaderKey();
+	Config.Shaders[CrossCompiler::SHADER_STAGE_PIXEL].ShaderKey = FOpenGLDynamicRHI::ResourceCast(PixelShaderRHI)->ShaderCodeKey;
 
 
 	return Config;
@@ -2765,13 +2713,13 @@ static bool CanCreateExternally(bool bIsFromPSO)
 	return false;
 }
 
-static FOpenGLProgramBinary ExternalProgramCompile(const FOpenGLProgramKey& ProgramKey, const FRHIVertexShader* VertexShaderRHI, const FRHIPixelShader* PixelShaderRHI)
+static FOpenGLProgramBinary ExternalProgramCompile(const FOpenGLProgramKey& ProgramKey, FRHIVertexShader* VertexShaderRHI, FRHIPixelShader* PixelShaderRHI)
 {
 	FOpenGLProgramBinary CompiledProgram;
 #if PLATFORM_ANDROID
 	// compile externally, sit and wait for the linked result	
-	const FOpenGLCompiledShaderKey& VSKey = FOpenGLDynamicRHI::ResourceCastProxy(VertexShaderRHI)->GetCompiledShaderKey();
-	const FOpenGLCompiledShaderKey& PSKey = FOpenGLDynamicRHI::ResourceCastProxy(PixelShaderRHI)->GetCompiledShaderKey();
+	const FOpenGLCompiledShaderKey& VSKey = FOpenGLDynamicRHI::ResourceCast(VertexShaderRHI)->ShaderCodeKey;
+	const FOpenGLCompiledShaderKey& PSKey = FOpenGLDynamicRHI::ResourceCast(PixelShaderRHI)->ShaderCodeKey;
 
 	TArray<ANSICHAR> VSCode;
 	TArray<ANSICHAR> PSCode;
@@ -2922,7 +2870,8 @@ FBoundShaderStateRHIRef FOpenGLDynamicRHI::RHICreateBoundShaderState_OnThisThrea
 	if (!PixelShaderRHI)
 	{
 		// use special null pixel shader when PixelShader was set to NULL
-		PixelShaderRHI = GetNULLPixelShader();
+		TShaderMapRef<FNULLPS> PixelShader(GetGlobalShaderMap(GMaxRHIFeatureLevel));
+		PixelShaderRHI = PixelShader.GetPixelShader();
 	}
 
 	// Check for an existing bound shader state which matches the parameters
@@ -3223,20 +3172,18 @@ FOpenGLBoundShaderState::FOpenGLBoundShaderState(
 	FRHIPixelShader* InPixelShaderRHI,
 	FRHIGeometryShader* InGeometryShaderRHI
 	)
-	:	CacheLink(InVertexDeclarationRHI, InVertexShaderRHI, InPixelShaderRHI,
-		InGeometryShaderRHI, this)
+	: CacheLink(InVertexDeclarationRHI, InVertexShaderRHI, InPixelShaderRHI, InGeometryShaderRHI, this)
 {
-	FOpenGLVertexDeclaration* InVertexDeclaration = FOpenGLDynamicRHI::ResourceCast(InVertexDeclarationRHI);
-	VertexDeclaration = InVertexDeclaration;
-	VertexShaderProxy = static_cast<FOpenGLVertexShaderProxy*>(InVertexShaderRHI);
-	PixelShaderProxy = static_cast<FOpenGLPixelShaderProxy*>(InPixelShaderRHI);
-	GeometryShaderProxy = static_cast<FOpenGLGeometryShaderProxy*>(InGeometryShaderRHI);
+	VertexDeclaration = FOpenGLDynamicRHI::ResourceCast(InVertexDeclarationRHI);
+	VertexShader      = FOpenGLDynamicRHI::ResourceCast(InVertexShaderRHI);
+	PixelShader       = FOpenGLDynamicRHI::ResourceCast(InPixelShaderRHI);
+	GeometryShader    = FOpenGLDynamicRHI::ResourceCast(InGeometryShaderRHI);
 
 	LinkedProgram = InLinkedProgram;
 
-	if (InVertexDeclaration)
+	if (VertexDeclaration)
 	{
-		FMemory::Memcpy(StreamStrides, InVertexDeclaration->StreamStrides, sizeof(StreamStrides));
+		FMemory::Memcpy(StreamStrides, VertexDeclaration->StreamStrides, sizeof(StreamStrides));
 	}
 	else
 	{
@@ -3288,23 +3235,11 @@ const TBitArray<>& FOpenGLBoundShaderState::GetUAVNeeds(int32& OutMaxUAVUnitUsed
 
 void FOpenGLBoundShaderState::GetNumUniformBuffers(int32 NumUniformBuffers[SF_NumGraphicsFrequencies])
 {
-	if (IsRunningRHIInSeparateThread())
-	{
-		// fast path, no need to check any fences....
-		check(IsInRHIThread());
-		check(IsValidRef(VertexShaderProxy) && IsValidRef(PixelShaderProxy));
-
+	check(IsValidRef(VertexShader) && IsValidRef(PixelShader));
 
-		NumUniformBuffers[SF_Vertex] = VertexShaderProxy->GetGLResourceObject_OnRHIThread()->Bindings.NumUniformBuffers;
-		NumUniformBuffers[SF_Pixel] = PixelShaderProxy->GetGLResourceObject_OnRHIThread()->Bindings.NumUniformBuffers;
-		NumUniformBuffers[SF_Geometry] = GeometryShaderProxy ? GeometryShaderProxy->GetGLResourceObject_OnRHIThread()->Bindings.NumUniformBuffers : -1;
-	}
-	else
-	{
-		NumUniformBuffers[SF_Vertex] = VertexShaderProxy->GetGLResourceObject()->Bindings.NumUniformBuffers;
-		NumUniformBuffers[SF_Pixel] = PixelShaderProxy->GetGLResourceObject()->Bindings.NumUniformBuffers;
-		NumUniformBuffers[SF_Geometry] = GeometryShaderProxy ? GeometryShaderProxy->GetGLResourceObject()->Bindings.NumUniformBuffers : -1;
-	}
+	NumUniformBuffers[SF_Vertex]   = VertexShader->Bindings.NumUniformBuffers;
+	NumUniformBuffers[SF_Pixel]    = PixelShader->Bindings.NumUniformBuffers;
+	NumUniformBuffers[SF_Geometry] = GeometryShader ? GeometryShader->Bindings.NumUniformBuffers : -1;
 }
 
 
diff --git a/Engine/Source/Runtime/OpenGLDrv/Private/OpenGLState.cpp b/Engine/Source/Runtime/OpenGLDrv/Private/OpenGLState.cpp
--- a/Engine/Source/Runtime/OpenGLDrv/Private/OpenGLState.cpp
+++ b/Engine/Source/Runtime/OpenGLDrv/Private/OpenGLState.cpp
@@ -223,7 +223,6 @@ static EBlendFactor TranslateBlendFactor(GLenum BlendFactor)
 
 FOpenGLSamplerState::~FOpenGLSamplerState()
 {
-	CreationFence.WaitFenceRenderThreadOnly();
 	VERIFY_GL_SCOPE();
 	FOpenGL::DeleteSamplers(1,&Resource);
 }
@@ -298,10 +297,10 @@ FSamplerStateRHIRef FOpenGLDynamicRHI::RHICreateSamplerState(const FSamplerState
 		SamplerState->Data.MinFilter = GL_LINEAR_MIPMAP_NEAREST;
 	}
 
-	SamplerState->CreationFence.Reset();
 	SamplerState->Resource = 0;
 
-	auto CreateGLSamplerState = [SamplerState]()
+	check(IsInRenderingThread());
+	FRHICommandListExecutor::GetImmediateCommandList().EnqueueLambda([SamplerState](FRHICommandListImmediate&)
 	{
 		VERIFY_GL_SCOPE();
 		FOpenGL::GenSamplers( 1, &SamplerState->Resource);
@@ -329,11 +328,7 @@ FSamplerStateRHIRef FOpenGLDynamicRHI::RHICreateSamplerState(const FSamplerState
 			FOpenGL::SetSamplerParameter(SamplerState->Resource, GL_TEXTURE_COMPARE_MODE, SamplerState->Data.CompareMode);
 			FOpenGL::SetSamplerParameter(SamplerState->Resource, GL_TEXTURE_COMPARE_FUNC, SamplerState->Data.CompareFunc);
 		}
-		SamplerState->CreationFence.WriteAssertFence();
-	};
-
-	RunOnGLRenderContextThread(MoveTemp(CreateGLSamplerState));
-	SamplerState->CreationFence.SetRHIThreadFence();
+	});
 	
 	// Manually add reference as we control the creation/destructions
 	SamplerState->AddRef();
diff --git a/Engine/Source/Runtime/OpenGLDrv/Private/OpenGLUniformBuffer.cpp b/Engine/Source/Runtime/OpenGLDrv/Private/OpenGLUniformBuffer.cpp
--- a/Engine/Source/Runtime/OpenGLDrv/Private/OpenGLUniformBuffer.cpp
+++ b/Engine/Source/Runtime/OpenGLDrv/Private/OpenGLUniformBuffer.cpp
@@ -434,8 +434,7 @@ void FOpenGLUniformBuffer::SetGLUniformBufferParams(GLuint InResource, uint32 In
 
 FOpenGLUniformBuffer::~FOpenGLUniformBuffer()
 {
-	AccessFence.WaitFenceRenderThreadOnly();
-	CopyFence.WaitFenceRenderThreadOnly();
+	VERIFY_GL_SCOPE();
 
 	if (Resource != 0)
 	{
@@ -500,10 +499,11 @@ void FOpenGLUniformBuffer::SetLayoutTable(const void* Contents, EUniformBufferVa
 }
 
 
-void CopyDataToUniformBuffer(const bool bCanRunOnThisThread, FOpenGLUniformBuffer* NewUniformBuffer, const void* Contents, uint32 ContentSize)
+void CopyDataToUniformBuffer(FRHICommandListImmediate& RHICmdList, FOpenGLUniformBuffer* NewUniformBuffer, const void* Contents, uint32 ContentSize)
 {
 	FOpenGLEUniformBufferDataRef EmulatedUniformDataRef = NewUniformBuffer->EmulatedBufferData;	
 	uint8* PersistentlyMappedBuffer = NewUniformBuffer->PersistentlyMappedBuffer;
+
 	// Copy the contents of the uniform buffer.
 	if (IsValidRef(EmulatedUniformDataRef))
 	{
@@ -515,30 +515,20 @@ void CopyDataToUniformBuffer(const bool bCanRunOnThisThread, FOpenGLUniformBuffe
 	}
 	else
 	{
-		if (bCanRunOnThisThread)
-		{
-			VERIFY_GL_SCOPE();
-			FOpenGL::BufferSubData(GL_UNIFORM_BUFFER, 0, ContentSize, Contents);
-		}
-		else
+		if (RHICmdList.IsTopOfPipe())
 		{
-			NewUniformBuffer->CopyFence.Reset();
-			// running on RHI thread take a copy of the incoming data.
-			FRHICommandListImmediate& RHICmdList = FRHICommandListExecutor::GetImmediateCommandList();
+			// Copy the data to the command list since we'll be deferring the buffer init.
 			void* ConstantBufferCopy = RHICmdList.Alloc(ContentSize, 16);
 			FMemory::Memcpy(ConstantBufferCopy, Contents, ContentSize);
-			
-			ALLOC_COMMAND_CL(RHICmdList, FRHICommandGLCommand)(
-				[=]() 
-				{
-					VERIFY_GL_SCOPE();
-					FOpenGL::BufferSubData(GL_UNIFORM_BUFFER, 0, ContentSize, ConstantBufferCopy);
-					NewUniformBuffer->CopyFence.WriteAssertFence();
-				});
-
-			NewUniformBuffer->CopyFence.SetRHIThreadFence();
 
+			Contents = ConstantBufferCopy;
 		}
+
+		RHICmdList.EnqueueLambda([ContentSize, Contents](FRHICommandListImmediate&)
+		{
+			VERIFY_GL_SCOPE();
+			FOpenGL::BufferSubData(GL_UNIFORM_BUFFER, 0, ContentSize, Contents);
+		});
 	}
 }
 
@@ -552,8 +542,6 @@ static FUniformBufferRHIRef CreateUniformBuffer(const void* Contents, const FRHI
 
 	FOpenGLUniformBuffer* NewUniformBuffer = new FOpenGLUniformBuffer(Layout);
 
-	TFunction<void(void)> GLCreationFunc;
-
 	const uint32 BucketIndex = GetPoolBucketIndex(Layout->ConstantBufferSize);
 	const uint32 SizeOfBufferToAllocate = UniformBufferSizeBuckets[BucketIndex];
 	const uint32 AllocatedSize = (SizeOfBufferToAllocate > 0) ? SizeOfBufferToAllocate : Layout->ConstantBufferSize;;
@@ -566,7 +554,9 @@ static FUniformBufferRHIRef CreateUniformBuffer(const void* Contents, const FRHI
 
 	bool bUseEmulatedUBs = GUseEmulatedUniformBuffers && !Layout->bNoEmulatedUniformBuffer;
 
+	check(IsInRenderingThread());
 	FRHICommandListImmediate& RHICmdList = FRHICommandListExecutor::GetImmediateCommandList();
+
 	{
 		const bool bStreamDraw = (Usage == UniformBuffer_SingleDraw || Usage == UniformBuffer_SingleFrame);
 		
@@ -580,18 +570,19 @@ static FUniformBufferRHIRef CreateUniformBuffer(const void* Contents, const FRHI
 		}
 		else if (IsSuballocatingUBOs())
 		{
-			GLCreationFunc = [NewUniformBuffer, AllocatedSize, &PersistentlyMappedBuffer, &EmulatedUniformDataRef, bStreamDraw]() {
+			RHICmdList.EnqueueLambda([NewUniformBuffer, AllocatedSize, &PersistentlyMappedBuffer, &EmulatedUniformDataRef, bStreamDraw](FRHICommandListImmediate&)
+			{
 				GLuint AllocatedResource = 0;
 				uint32 OffsetInBuffer = 0;
 
 				SuballocateUBO(AllocatedSize, AllocatedResource, OffsetInBuffer, PersistentlyMappedBuffer);
 				NewUniformBuffer->SetGLUniformBufferParams(AllocatedResource, OffsetInBuffer, PersistentlyMappedBuffer, AllocatedSize, EmulatedUniformDataRef, bStreamDraw);
-			};
+			});
 		}
 		else
 		{
 			check(PersistentlyMappedBuffer == nullptr);
-			GLCreationFunc = [NewUniformBuffer, AllocatedSize, PersistentlyMappedBuffer, EmulatedUniformDataRef, bStreamDraw]()
+			RHICmdList.EnqueueLambda([NewUniformBuffer, AllocatedSize, PersistentlyMappedBuffer, EmulatedUniformDataRef, bStreamDraw](FRHICommandListImmediate&)
 			{
 				VERIFY_GL_SCOPE();
 				GLuint AllocatedResource = 0;
@@ -600,31 +591,18 @@ static FUniformBufferRHIRef CreateUniformBuffer(const void* Contents, const FRHI
 				::CachedBindUniformBuffer(AllocatedResource);
 				glBufferData(GL_UNIFORM_BUFFER, AllocatedSize, NULL, bStreamDraw ? GL_STREAM_DRAW : GL_STATIC_DRAW);
 				NewUniformBuffer->SetGLUniformBufferParams(AllocatedResource, OffsetInBuffer, nullptr, AllocatedSize, EmulatedUniformDataRef, bStreamDraw);
-			};
+			});
 		}
 	}
 
-	const bool bCanCreateOnThisThread = RHICmdList.Bypass() || (!IsRunningRHIInSeparateThread() && IsInRenderingThread()) || IsInRHIThread();
 	if(!bUseEmulatedUBs)
 	{
-		if (bCanCreateOnThisThread)
-		{
-			GLCreationFunc();
-		}
-		else
+		// flush for the UBO case
+		// as this path interacts with UBOPool, this hasnt been addressed for the RHI thread case.
+		if (RHICmdList.IsTopOfPipe() && IsSuballocatingUBOs())
 		{
-			NewUniformBuffer->AccessFence.Reset();
-			// Queue GL resource creation.
-			ALLOC_COMMAND_CL(RHICmdList, FRHICommandGLCommand)([=]() {GLCreationFunc(); NewUniformBuffer->AccessFence.WriteAssertFence(); });
-			NewUniformBuffer->AccessFence.SetRHIThreadFence();
-
-			// flush for the UBO case
-			// as this path interacts with UBOPool, this hasnt been addressed for the RHI thread case.
-			if (IsSuballocatingUBOs())
-			{
-				RHICmdList.ImmediateFlush(EImmediateFlushType::FlushRHIThread);
-				RHITHREAD_GLTRACE_BLOCKING;
-			}
+			RHICmdList.ImmediateFlush(EImmediateFlushType::FlushRHIThread);
+			RHITHREAD_GLTRACE_BLOCKING;
 		}
 	}
 
@@ -634,7 +612,7 @@ static FUniformBufferRHIRef CreateUniformBuffer(const void* Contents, const FRHI
 
 	if (Contents)
 	{
-		CopyDataToUniformBuffer(bCanCreateOnThisThread, NewUniformBuffer, Contents, Layout->ConstantBufferSize);
+		CopyDataToUniformBuffer(RHICmdList, NewUniformBuffer, Contents, Layout->ConstantBufferSize);
 	}
 
 	// Initialize the resource table for this uniform buffer.
@@ -645,7 +623,9 @@ static FUniformBufferRHIRef CreateUniformBuffer(const void* Contents, const FRHI
 
 FUniformBufferRHIRef FOpenGLDynamicRHI::RHICreateUniformBuffer(const void* Contents, const FRHIUniformBufferLayout* Layout, EUniformBufferUsage Usage, EUniformBufferValidation Validation)
 {
+	check(IsInRenderingThread());
 	FRHICommandListImmediate& RHICmdList = FRHICommandListExecutor::GetImmediateCommandList();
+
 	// This should really be synchronized, if there's a chance it'll be used from more than one buffer. Luckily, uniform buffers
 	// are only used for drawing/shader usage, not for loading resources or framebuffer blitting, so no synchronization primitives for now.
 
@@ -666,8 +646,6 @@ FUniformBufferRHIRef FOpenGLDynamicRHI::RHICreateUniformBuffer(const void* Conte
 	uint32 AllocatedSize = 0;
 	FOpenGLEUniformBufferDataRef EmulatedUniformDataRef;
 
-	const bool bCanCreateOnThisThread = RHICmdList.Bypass() || (!IsRunningRHIInSeparateThread() && IsInRenderingThread()) || IsInRHIThread();
-
 	// If the uniform buffer contains constants, allocate a uniform buffer resource from GL.
 	if (Layout->ConstantBufferSize > 0)
 	{
@@ -717,20 +695,11 @@ FUniformBufferRHIRef FOpenGLDynamicRHI::RHICreateUniformBuffer(const void* Conte
 				}
 				else
 				{
-					auto CacheGLUniformBuffer = [AllocatedResource]()
+					RHICmdList.EnqueueLambda([AllocatedResource](FRHICommandListImmediate&)
 					{
 						VERIFY_GL_SCOPE();
 						::CachedBindUniformBuffer(AllocatedResource);
-					};
-
-					if (bCanCreateOnThisThread)
-					{
-						CacheGLUniformBuffer();
-					}
-					else
-					{
-						ALLOC_COMMAND_CL(RHICmdList, FRHICommandGLCommand)(MoveTemp(CacheGLUniformBuffer));
-					}
+					});
 				}
 			}
 			else
@@ -751,7 +720,7 @@ FUniformBufferRHIRef FOpenGLDynamicRHI::RHICreateUniformBuffer(const void* Conte
 	check(!bUseEmulatedUBs || (IsValidRef(EmulatedUniformDataRef) && (EmulatedUniformDataRef->Data.Num() * EmulatedUniformDataRef->Data.GetTypeSize() == AllocatedSize)));
 	if (Contents)
 	{
-		CopyDataToUniformBuffer(bCanCreateOnThisThread, NewUniformBuffer, Contents, Layout->ConstantBufferSize);
+		CopyDataToUniformBuffer(RHICmdList, NewUniformBuffer, Contents, Layout->ConstantBufferSize);
 	}
 
 	// Initialize the resource table for this uniform buffer.
diff --git a/Engine/Source/Runtime/OpenGLDrv/Public/OpenGLDrv.h b/Engine/Source/Runtime/OpenGLDrv/Public/OpenGLDrv.h
--- a/Engine/Source/Runtime/OpenGLDrv/Public/OpenGLDrv.h
+++ b/Engine/Source/Runtime/OpenGLDrv/Public/OpenGLDrv.h
@@ -241,7 +241,6 @@ struct FOpenGLGPUProfiler : public FGPUProfiler
 	class FOpenGLDynamicRHI* OpenGLRHI;
 	// count the number of beginframe calls without matching endframe calls.
 	int32 NestedFrameCount;
-	bool bIntialized;
 
 	uint32 ExternalGPUTime;
 
@@ -254,12 +253,7 @@ struct FOpenGLGPUProfiler : public FGPUProfiler
 	,	CurrentGPUFrameQueryIndex(0)
 	,	OpenGLRHI(InOpenGLRHI)
 	,	NestedFrameCount(0)
-	,	bIntialized(false)
 	,	ExternalGPUTime(0)
-	{
-	}
-
-	void InitResources()
 	{
 		FrameTiming.InitResources();
 		for (int32 Index = 0; Index < MAX_GPUFRAMEQUERIES; ++Index)
@@ -286,7 +280,7 @@ struct FOpenGLGPUProfiler : public FGPUProfiler
 PRAGMA_DISABLE_DEPRECATION_WARNINGS
 
 /** The interface which is implemented by the dynamically bound RHI. */
-class OPENGLDRV_API FOpenGLDynamicRHI  final : public IOpenGLDynamicRHI, public IRHICommandContextPSOFallback
+class OPENGLDRV_API FOpenGLDynamicRHI final : public IOpenGLDynamicRHI, public IRHICommandContextPSOFallback
 {
 	static inline FOpenGLDynamicRHI* Singleton = nullptr;
 
@@ -327,45 +321,13 @@ public:
 
 	// FDynamicRHI interface.
 	virtual void Init();
-	virtual void PostInit();
 
 	virtual void Shutdown();
 	virtual const TCHAR* GetName() override { return TEXT("OpenGL"); }
 
-	template<typename TRHIType>
-	static FORCEINLINE typename TOpenGLResourceTraits<TRHIType>::TConcreteType* ResourceCastProxy(TRHIType* Resource)
-	{
-		static_assert(TIsGLProxyObject_V<typename TOpenGLResourceTraits<TRHIType>::TConcreteType>, "Wrong type past");
-		return static_cast<typename TOpenGLResourceTraits<TRHIType>::TConcreteType*>(Resource);
-	}
-	template<typename TRHIType>
-	static FORCEINLINE const typename TOpenGLResourceTraits<TRHIType>::TConcreteType* ResourceCastProxy(const TRHIType* Resource)
-	{
-		static_assert(TIsGLProxyObject_V<typename TOpenGLResourceTraits<TRHIType>::TConcreteType>, "Wrong type past");
-		return static_cast<const typename TOpenGLResourceTraits<TRHIType>::TConcreteType*>(Resource);
-	}
-
-	// If using a Proxy object return the contained GL object rather than the proxy itself.
 	template<typename TRHIType>
 	static FORCEINLINE auto* ResourceCast(TRHIType* Resource)
 	{
-		if constexpr (TIsGLProxyObject_V<typename TOpenGLResourceTraits<TRHIType>::TConcreteType>)
-		{
-			auto GLProxy = static_cast<typename TOpenGLResourceTraits<TRHIType>::TConcreteType*>(Resource);
-			// rhi resource can be null.
-			return GLProxy ? GLProxy->GetGLResourceObject() : nullptr;
-		}
-		else
-		{
-			CheckRHITFence(static_cast<typename TOpenGLResourceTraits<TRHIType>::TConcreteType*>(Resource));
-			return static_cast<typename TOpenGLResourceTraits<TRHIType>::TConcreteType*>(Resource);
-		}
-	}
-
-	template<typename TRHIType>
-	static FORCEINLINE typename TOpenGLResourceTraits<TRHIType>::TConcreteType* ResourceCast_Unfenced(TRHIType* Resource)
-	{
-		static_assert(!TIsGLProxyObject_V<typename TOpenGLResourceTraits<TRHIType>::TConcreteType>, "Wrong type passed");
 		return static_cast<typename TOpenGLResourceTraits<TRHIType>::TConcreteType*>(Resource);
 	}
 
@@ -884,8 +846,6 @@ public:
 
 	GLuint GetOpenGLFramebuffer(uint32 NumSimultaneousRenderTargets, FOpenGLTexture** RenderTargets, const uint32* ArrayIndices, const uint32* MipmapLevels, FOpenGLTexture* DepthStencilTarget);
 	GLuint GetOpenGLFramebuffer(uint32 NumSimultaneousRenderTargets, FOpenGLTexture** RenderTargets, const uint32* ArrayIndices, const uint32* MipmapLevels, FOpenGLTexture* DepthStencilTarget, int32 NumRenderingSamples);
-
-	FRHIPixelShader* GetNULLPixelShader() const;
 	
 private:
 
@@ -967,9 +927,6 @@ private:
 	FCriticalSection CustomPresentSection;
 	TRefCountPtr<class FRHICustomPresent> CustomPresent;
 
-	// Cached NULLPS
-	FRHIPixelShader* NULLPixelShaderRHI;
-
 	void InitializeStateResources();
 
 	void SetupVertexArrays(FOpenGLContextState& ContextCache, uint32 BaseVertexIndex, FOpenGLStream* Streams, uint32 NumStreams, uint32 MaxVertices);
@@ -985,7 +942,7 @@ private:
 	void UpdateScissorRectInOpenGLContext( FOpenGLContextState& ContextState );
 	void UpdateViewportInOpenGLContext( FOpenGLContextState& ContextState );
 	
-	template <class ShaderType> void SetResourcesFromTables(ShaderType* RESTRICT);
+	template <class ShaderType> void SetResourcesFromTables(ShaderType* Shader);
 	FORCEINLINE void CommitGraphicsResourceTables()
 	{
 		if (PendingState.bAnyDirtyGraphicsUniformBuffers)
@@ -994,7 +951,7 @@ private:
 		}
 	}
 	void CommitGraphicsResourceTablesInner();
-	void CommitComputeResourceTables(FOpenGLComputeShaderProxy* ComputeShader);
+	void CommitComputeResourceTables(FOpenGLComputeShader* ComputeShader);
 	void CommitNonComputeShaderConstants();
 	void CommitComputeShaderConstants(FOpenGLComputeShader* ComputeShader);
 	void SetPendingBlendStateForActiveRenderTargets( FOpenGLContextState& ContextState );
@@ -1018,8 +975,6 @@ public:
 	void InitializeGLTextureInternal(FOpenGLTexture* Texture, void const* BulkDataPtr, uint64 BulkDataSize);
 private:
 
-	void SetupRecursiveResources();
-
 	void ApplyTextureStage(FOpenGLContextState& ContextState, GLint TextureIndex, const FTextureStage& TextureStage, FOpenGLSamplerState* SamplerState);
 
 	void ReadSurfaceDataRaw(FOpenGLContextState& ContextState, FRHITexture* TextureRHI,FIntRect Rect,TArray<uint8>& OutData, FReadSurfaceDataFlags InFlags);
diff --git a/Engine/Source/Runtime/OpenGLDrv/Public/OpenGLResources.h b/Engine/Source/Runtime/OpenGLDrv/Public/OpenGLResources.h
--- a/Engine/Source/Runtime/OpenGLDrv/Public/OpenGLResources.h
+++ b/Engine/Source/Runtime/OpenGLDrv/Public/OpenGLResources.h
@@ -76,295 +76,8 @@ void DecrementBufferMemory(GLenum Type, uint32 NumBytes);
 	#define DETAILED_QUICK_SCOPE_CYCLE_COUNTER(x)
 #endif
 
-#if UE_BUILD_TEST
-#define USE_REAL_RHI_FENCES (0)
-#define USE_CHEAP_ASSERTONLY_RHI_FENCES (1)
-#define GLAF_CHECK(x) \
-if (!(x)) \
-{  \
-	UE_LOG(LogRHI, Fatal, TEXT("AssertFence Fail on line %s."), TEXT(PREPROCESSOR_TO_STRING(__LINE__))); \
-	FPlatformMisc::LocalPrint(TEXT("Failed a check on line:\n")); FPlatformMisc::LocalPrint(TEXT(PREPROCESSOR_TO_STRING(__LINE__))); FPlatformMisc::LocalPrint(TEXT("\n")); *((int*)3) = 13; \
-}
-
-#elif DO_CHECK
-#define USE_REAL_RHI_FENCES (1)
-#define USE_CHEAP_ASSERTONLY_RHI_FENCES (1)
-#define GLAF_CHECK(x)  check(x)
-
-//#define GLAF_CHECK(x) if (!(x)) { FPlatformMisc::LocalPrint(TEXT("Failed a check on line:\n")); FPlatformMisc::LocalPrint(TEXT( PREPROCESSOR_TO_STRING(__LINE__))); FPlatformMisc::LocalPrint(TEXT("\n")); *((int*)3) = 13; }
-
-#else
-#define USE_REAL_RHI_FENCES (0)
-#define USE_CHEAP_ASSERTONLY_RHI_FENCES (0)
-
-#define GLAF_CHECK(x) 
-
-#endif
-
 #define GLDEBUG_LABELS_ENABLED (!UE_BUILD_SHIPPING)
 
-class FOpenGLRHIThreadResourceFence
-{
-	FGraphEventRef RealRHIFence;
-
-public:
-
-	FORCEINLINE_DEBUGGABLE void Reset()
-	{
-		if (IsRunningRHIInSeparateThread())
-		{
-			GLAF_CHECK(IsInRenderingThread());
-			GLAF_CHECK(!RealRHIFence.GetReference() || RealRHIFence->IsComplete());
-			RealRHIFence = nullptr;
-		}
-	}
-	FORCEINLINE_DEBUGGABLE void SetRHIThreadFence()
-	{
-		if (IsRunningRHIInSeparateThread())
-		{
-			GLAF_CHECK(IsInRenderingThread());
-
-			GLAF_CHECK(!RealRHIFence.GetReference() || RealRHIFence->IsComplete());
-			if (IsRunningRHIInSeparateThread())
-			{
-				RealRHIFence = FRHICommandListExecutor::GetImmediateCommandList().RHIThreadFence(false);
-			}
-		}
-	}
-	FORCEINLINE_DEBUGGABLE void WriteAssertFence()
-	{
-		if (IsRunningRHIInSeparateThread())
-		{
-			GLAF_CHECK((IsInRenderingThread() && !IsRunningRHIInSeparateThread()) || (IsInRHIThread() && IsRunningRHIInSeparateThread()));
-		}
-	}
-
-	FORCEINLINE_DEBUGGABLE void WaitFenceRenderThreadOnly()
-	{
-		if (IsRunningRHIInSeparateThread())
-		{
-			// Do not check if running on RHI thread.
-			// all rhi thread operations will be in order, check for RHIT isnt required.
-			if (IsInRenderingThread())
-			{
-				if (RealRHIFence.GetReference() && RealRHIFence->IsComplete())
-				{
-					RealRHIFence = nullptr;
-				}
-				else if (RealRHIFence.GetReference())
-				{
-					UE_LOG(LogRHI, Warning, TEXT("FOpenGLRHIThreadResourceFence waited."));
-					QUICK_SCOPE_CYCLE_COUNTER(STAT_FOpenGLRHIThreadResourceFence_Wait);
-					FRHICommandListExecutor::WaitOnRHIThreadFence(RealRHIFence);
-					RealRHIFence = nullptr;
-				}
-			}
-		}
-	}
-};
-
-class FOpenGLAssertRHIThreadFence
-{
-#if USE_REAL_RHI_FENCES
-	FGraphEventRef RealRHIFence;
-#endif
-#if USE_CHEAP_ASSERTONLY_RHI_FENCES
-	FThreadSafeCounter AssertFence;
-#endif
-
-public:
-
-	FORCEINLINE_DEBUGGABLE void Reset()
-	{
-		if (IsRunningRHIInSeparateThread())
-		{
-			check(IsInRenderingThread() || IsInRHIThread());
-#if USE_REAL_RHI_FENCES
-
-			GLAF_CHECK(!RealRHIFence.GetReference() || RealRHIFence->IsComplete());
-			RealRHIFence = nullptr;
-#endif
-#if USE_CHEAP_ASSERTONLY_RHI_FENCES
-			int32 AFenceVal = AssertFence.GetValue();
-			GLAF_CHECK(AFenceVal == 0 || AFenceVal == 2);
-			AssertFence.Set(1);
-#endif
-		}
-	}
-	FORCEINLINE_DEBUGGABLE void SetRHIThreadFence()
-	{
-		if (IsRunningRHIInSeparateThread())
-		{
-			check(IsInRenderingThread() || IsInRHIThread());
-
-#if USE_CHEAP_ASSERTONLY_RHI_FENCES
-			int32 AFenceVal = AssertFence.GetValue();
-			GLAF_CHECK(AFenceVal == 1 || AFenceVal == 2);
-#endif
-#if USE_REAL_RHI_FENCES
-			GLAF_CHECK(!RealRHIFence.GetReference() || RealRHIFence->IsComplete());
-			// Only get the fence if running on RT.
-			if (IsRunningRHIInSeparateThread() && IsInRenderingThread())
-			{
-				RealRHIFence = FRHICommandListExecutor::GetImmediateCommandList().RHIThreadFence(false);
-			}
-#endif
-		}
-	}
-	FORCEINLINE_DEBUGGABLE void WriteAssertFence()
-	{
-		if (IsRunningRHIInSeparateThread())
-		{
-			check((IsInRenderingThread() && !IsRunningRHIInSeparateThread()) || (IsInRHIThread() && IsRunningRHIInSeparateThread()));
-#if USE_CHEAP_ASSERTONLY_RHI_FENCES
-			int32 NewValue = AssertFence.Increment();
-			GLAF_CHECK(NewValue == 2);
-#endif
-		}
-	}
-
-	FORCEINLINE_DEBUGGABLE void WaitFenceRenderThreadOnly()
-	{
-		if (IsRunningRHIInSeparateThread())
-		{
-			// Do not check if running on RHI thread.
-			// all rhi thread operations will be in order, check for RHIT isnt required.
-			if (IsInRenderingThread())
-			{
-#if USE_CHEAP_ASSERTONLY_RHI_FENCES
-				GLAF_CHECK(AssertFence.GetValue() == 0 || AssertFence.GetValue() == 2);
-#endif
-#if USE_REAL_RHI_FENCES
-				GLAF_CHECK(!RealRHIFence.GetReference() || RealRHIFence->IsComplete());
-				if (RealRHIFence.GetReference())
-				{
-					FRHICommandListExecutor::WaitOnRHIThreadFence(RealRHIFence);
-					RealRHIFence = nullptr;
-				}
-#endif
-			}
-		}
-	}
-};
-
-
-//////////////////////////////////////////////////////////////////////////
-// Proxy object that fulfils immediate requirements of RHIResource creation whilst allowing deferment of GL resource creation on to the RHI thread.
-
-template<typename TRHIType, typename TOGLResourceType>
-class TOpenGLResourceProxy : public TRHIType
-{
-public:
-	TOpenGLResourceProxy(TFunction<TOGLResourceType*(TRHIType*)> CreateFunc)
-		: GLResourceObject(nullptr)
-	{
-		check((bool)CreateFunc);
-		FRHICommandListImmediate& RHICmdList = FRHICommandListExecutor::GetImmediateCommandList();
-		if (ShouldRunGLRenderContextOpOnThisThread(RHICmdList))
-		{
-			GLResourceObject = CreateFunc(this);
-			GLResourceObject->AddRef();
-			bQueuedCreation = false;
-		}
-		else
-		{
-			CreationFence.Reset();
-			ALLOC_COMMAND_CL(RHICmdList, FRHICommandGLCommand)([this, CreateFunc = MoveTemp(CreateFunc)]()
-			{
-				GLResourceObject = CreateFunc(this);
-				GLResourceObject->AddRef();
-				CreationFence.WriteAssertFence();
-			});
-			CreationFence.SetRHIThreadFence();
-			bQueuedCreation = true;
-		}
-	}
-
-	virtual ~TOpenGLResourceProxy()
-	{
-		// Wait for any queued creation calls.
-		WaitIfQueued();
-
-		check(GLResourceObject);
-
-		FRHICommandListImmediate& RHICmdList = FRHICommandListExecutor::GetImmediateCommandList();
-		if (ShouldRunGLRenderContextOpOnThisThread(RHICmdList))
-		{
-			GLResourceObject->Release();
-		}
-		else
-		{
-			RunOnGLRenderContextThread([GLResourceObject = GLResourceObject]()
-			{
-				GLResourceObject->Release();
-			});
-			GLResourceObject = nullptr;
-		}
-	}
-
-	TOpenGLResourceProxy(const TOpenGLResourceProxy&) = delete;
-	TOpenGLResourceProxy& operator = (const TOpenGLResourceProxy&) = delete;
-
-	TOGLResourceType* GetGLResourceObject()
-	{
-		CreationFence.WaitFenceRenderThreadOnly();
-		return GLResourceObject;
-	}
-
-	FORCEINLINE TOGLResourceType* GetGLResourceObject_OnRHIThread()
-	{
-		check(IsInRHIThread());
-		return GLResourceObject;
-	}
-
-	typedef TOGLResourceType ContainedGLType;
-private:
-	void WaitIfQueued()
-	{
-		if (bQueuedCreation)
-		{
-			CreationFence.WaitFenceRenderThreadOnly();
-		}
-	}
-
-	//FOpenGLRHIThreadResourceFence CreationFence;
-	FOpenGLAssertRHIThreadFence CreationFence;
-	TRefCountPtr<TOGLResourceType> GLResourceObject;
-	bool bQueuedCreation;
-};
-
-
-template<typename TRHIType, typename TOGLResourceType>
-class TOpenGLShaderProxy : public TOpenGLResourceProxy< TRHIType, TOGLResourceType>
-{
-public:
-	TOpenGLShaderProxy(const FOpenGLCompiledShaderKey& InKey, TFunction<TOGLResourceType* (TRHIType*)> CreateFunc)
-		: TOpenGLResourceProxy< TRHIType, TOGLResourceType>(CreateFunc), Key(InKey)
-	{
-	}
-	
-	const FOpenGLCompiledShaderKey& GetCompiledShaderKey() const { return Key; }
-
-private:
-
-	FOpenGLCompiledShaderKey Key;
-};
-
-typedef TOpenGLShaderProxy<FRHIVertexShader, FOpenGLVertexShader> FOpenGLVertexShaderProxy;
-typedef TOpenGLShaderProxy<FRHIPixelShader, FOpenGLPixelShader> FOpenGLPixelShaderProxy;
-typedef TOpenGLShaderProxy<FRHIGeometryShader, FOpenGLGeometryShader> FOpenGLGeometryShaderProxy;
-typedef TOpenGLShaderProxy<FRHIComputeShader, FOpenGLComputeShader> FOpenGLComputeShaderProxy;
-
-
-template <typename T>
-constexpr bool TIsGLProxyObject_V = false;
-
-template<typename TRHIType, typename TOGLResourceType>
-constexpr bool TIsGLProxyObject_V<TOpenGLResourceProxy<TRHIType, TOGLResourceType>> = true;
-
-template<typename TRHIType, typename TOGLResourceType>
-constexpr bool TIsGLProxyObject_V<TOpenGLShaderProxy<TRHIType, TOGLResourceType>> = true;
-
 typedef void (*BufferBindFunction)( GLenum Type, GLuint Buffer );
 
 template <typename BaseType, BufferBindFunction BufBind>
@@ -435,25 +148,25 @@ public:
 	{
 		RealSize = ResourceSize ? ResourceSize : InSize;
 
-		if (ShouldRunGLRenderContextOpOnThisThread(RHICmdList))
+		if (RHICmdList && RHICmdList->IsTopOfPipe() && InData)
+		{
+			void* LocalCopy = RHICmdList->Alloc(RealSize, 16);
+			FMemory::Memcpy(LocalCopy, InData, RealSize);
+			InData = LocalCopy;
+		}
+
+		auto InitLambda = [this, InData, ResourceToUse, ResourceSize]()
 		{
 			CreateGLBuffer(InData, ResourceToUse, ResourceSize);
+		};
+
+		if (RHICmdList)
+		{
+			RHICmdList->EnqueueLambda([Lambda = MoveTemp(InitLambda)](FRHICommandListBase&) { Lambda(); });
 		}
 		else
 		{
-			void* BuffData = nullptr;
-			if (InData)
-			{
-				BuffData = RHICmdList->Alloc(RealSize, 16);
-				FMemory::Memcpy(BuffData, InData, RealSize);
-			}
-			TransitionFence.Reset();
-			ALLOC_COMMAND_CL(*RHICmdList, FRHICommandGLCommand)([=]() 
-			{
-				CreateGLBuffer(BuffData, ResourceToUse, ResourceSize); 
-				TransitionFence.WriteAssertFence();
-			});
-			TransitionFence.SetRHIThreadFence();
+			InitLambda();
 		}
 	}
 
@@ -504,38 +217,32 @@ public:
 
 	virtual ~TOpenGLBuffer()
 	{
-		// this is a bit of a special case, normally the RT destroys all rhi resources...but this isn't an rhi resource
-		TransitionFence.WaitFenceRenderThreadOnly();
+		VERIFY_GL_SCOPE();
 
 		if (Resource != 0)
 		{
-			auto DeleteGLResources = [Resource=Resource, RealSize= RealSize, bStreamDraw= (bool)bStreamDraw, LockBuffer = LockBuffer, bLockBufferWasAllocated=bLockBufferWasAllocated]()
+			if (BaseType::OnDelete(Resource, RealSize, bStreamDraw, 0))
+			{
+				FOpenGL::DeleteBuffers(1, &Resource);
+			}
+
+			if (LockBuffer != NULL)
 			{
-				VERIFY_GL_SCOPE();
-				if (BaseType::OnDelete(Resource, RealSize, bStreamDraw, 0))
+				if (bLockBufferWasAllocated)
 				{
-					FOpenGL::DeleteBuffers(1, &Resource);
+					FMemory::Free(LockBuffer);
 				}
-				if (LockBuffer != NULL)
+				else
 				{
-					if (bLockBufferWasAllocated)
-					{
-						FMemory::Free(LockBuffer);
-					}
-					else
-					{
-						UE_LOG(LogRHI,Warning,TEXT("Destroying TOpenGLBuffer without returning memory to the driver; possibly called RHIMapStagingSurface() but didn't call RHIUnmapStagingSurface()? Resource %u"), Resource);
-					}
+					UE_LOG(LogRHI,Warning,TEXT("Destroying TOpenGLBuffer without returning memory to the driver; possibly called RHIMapStagingSurface() but didn't call RHIUnmapStagingSurface()? Resource %u"), Resource);
 				}
-			};
+			}
 
-			RunOnGLRenderContextThread(MoveTemp(DeleteGLResources));
 			LockBuffer = nullptr;
 			DecrementBufferMemory(Type, RealSize);
 
 			ReleaseCachedBuffer();
 		}
-
 	}
 
 	void Bind()
@@ -839,8 +546,6 @@ private:
 	GLuint CachedBufferSize = 0;
 
 	uint32 RealSize;	// sometimes (for example, for uniform buffer pool) we allocate more in OpenGL than is requested of us.
-
-	FOpenGLAssertRHIThreadFence TransitionFence;
 };
 
 class FOpenGLBasePixelBuffer : public FRefCountedObject
@@ -962,9 +667,6 @@ public:
 	/** Destructor. */
 	~FOpenGLUniformBuffer();
 
-	FOpenGLAssertRHIThreadFence AccessFence;
-	FOpenGLAssertRHIThreadFence CopyFence;
-
 	// Provides public non-const access to ResourceTable.
 	// @todo refactor uniform buffers to perform updates as a member function, so this isn't necessary.
 	TArray<TRefCountPtr<FRHIResource>>& GetResourceTable() { return ResourceTable; }
@@ -1119,9 +821,9 @@ public:
 
 	FOpenGLLinkedProgram* LinkedProgram;
 	TRefCountPtr<FOpenGLVertexDeclaration> VertexDeclaration;
-	TRefCountPtr<FOpenGLVertexShaderProxy> VertexShaderProxy;
-	TRefCountPtr<FOpenGLPixelShaderProxy> PixelShaderProxy;
-	TRefCountPtr<FOpenGLGeometryShaderProxy> GeometryShaderProxy;
+	TRefCountPtr<FOpenGLVertexShader> VertexShader;
+	TRefCountPtr<FOpenGLPixelShader> PixelShader;
+	TRefCountPtr<FOpenGLGeometryShader> GeometryShader;
 
 	/** Initialization constructor. */
 	FOpenGLBoundShaderState(
@@ -1140,21 +842,21 @@ public:
 	int32 MaxTextureStageUsed();
 	bool RequiresDriverInstantiation();
 
-	FOpenGLVertexShaderProxy* GetVertexShader()
+	FOpenGLVertexShader* GetVertexShader()
 	{
-		check(IsValidRef(VertexShaderProxy));
-		return VertexShaderProxy;
+		check(IsValidRef(VertexShader));
+		return VertexShader;
 	}
 
-	FOpenGLPixelShaderProxy* GetPixelShader()
+	FOpenGLPixelShader* GetPixelShader()
 	{
-		check(IsValidRef(PixelShaderProxy));
-		return PixelShaderProxy;
+		check(IsValidRef(PixelShader));
+		return PixelShader;
 	}
 
-	FOpenGLGeometryShaderProxy* GetGeometryShader()
+	FOpenGLGeometryShader* GetGeometryShader()
 	{
-		return GeometryShaderProxy;
+		return GeometryShader;
 	}
 
 	virtual ~FOpenGLBoundShaderState();
@@ -1451,27 +1153,6 @@ public:
 	uint8 const bMultisampleRenderbuffer : 1;
 };
 
-template <typename T>
-struct TIsGLResourceWithFence
-{
-	enum
-	{
-		Value = TOr<
-		TPointerIsConvertibleFromTo<T, const FOpenGLTexture>
-		//		,TIsDerivedFrom<T, FRHITexture>
-		>::Value
-	};
-};
-
-template<typename T>
-static void CheckRHITFence(T* Resource)
-{
-	if constexpr (TIsGLResourceWithFence<T>::Value)
-	{
-		Resource->CreationFence.WaitFenceRenderThreadOnly();
-	}
-}
-
 /** Given a pointer to a RHI texture that was created by the OpenGL RHI, returns a pointer to the FOpenGLTexture it encapsulates. */
 inline FOpenGLTexture* GetOpenGLTextureFromRHITexture(FRHITexture* TextureRHI)
 {
@@ -1498,8 +1179,6 @@ public:
 	/** The cached query result. */
 	GLuint64 Result;
 
-	FOpenGLAssertRHIThreadFence CreationFence;
-
 	FThreadSafeCounter TotalBegins;
 	FThreadSafeCounter TotalResults;
 
@@ -1750,22 +1429,22 @@ struct TOpenGLResourceTraits<FRHIVertexDeclaration>
 template<>
 struct TOpenGLResourceTraits<FRHIVertexShader>
 {
-	typedef FOpenGLVertexShaderProxy TConcreteType;
+	typedef FOpenGLVertexShader TConcreteType;
 };
 template<>
 struct TOpenGLResourceTraits<FRHIGeometryShader>
 {
-	typedef FOpenGLGeometryShaderProxy TConcreteType;
+	typedef FOpenGLGeometryShader TConcreteType;
 };
 template<>
 struct TOpenGLResourceTraits<FRHIPixelShader>
 {
-	typedef FOpenGLPixelShaderProxy TConcreteType;
+	typedef FOpenGLPixelShader TConcreteType;
 };
 template<>
 struct TOpenGLResourceTraits<FRHIComputeShader>
 {
-	typedef FOpenGLComputeShaderProxy TConcreteType;
+	typedef FOpenGLComputeShader TConcreteType;
 };
 template<>
 struct TOpenGLResourceTraits<FRHIBoundShaderState>
diff --git a/Engine/Source/Runtime/OpenGLDrv/Public/OpenGLShaderResources.h b/Engine/Source/Runtime/OpenGLDrv/Public/OpenGLShaderResources.h
--- a/Engine/Source/Runtime/OpenGLDrv/Public/OpenGLShaderResources.h
+++ b/Engine/Source/Runtime/OpenGLDrv/Public/OpenGLShaderResources.h
@@ -221,21 +221,47 @@ inline FArchive& operator<<(FArchive& Ar, FOpenGLCodeHeader& Header)
 
 class FOpenGLLinkedProgram;
 
+class FOpenGLCompiledShaderKey
+{
+public:
+	FOpenGLCompiledShaderKey() = default;
+	FOpenGLCompiledShaderKey(
+		GLenum InTypeEnum,
+		uint32 InCodeSize,
+		uint32 InCodeCRC
+	)
+		: TypeEnum(InTypeEnum)
+		, CodeSize(InCodeSize)
+		, CodeCRC(InCodeCRC)
+	{
+	}
+
+	friend bool operator == (const FOpenGLCompiledShaderKey& A, const FOpenGLCompiledShaderKey& B)
+	{
+		return A.TypeEnum == B.TypeEnum && A.CodeSize == B.CodeSize && A.CodeCRC == B.CodeCRC;
+	}
+
+	friend uint32 GetTypeHash(const FOpenGLCompiledShaderKey& Key)
+	{
+		return GetTypeHash(Key.TypeEnum) ^ GetTypeHash(Key.CodeSize) ^ GetTypeHash(Key.CodeCRC);
+	}
+
+	uint32 GetCodeCRC() const { return CodeCRC; }
+
+private:
+	GLenum TypeEnum = 0;
+	uint32 CodeSize = 0;
+	uint32 CodeCRC  = 0;
+};
+
 /**
  * OpenGL shader resource.
  */
-template <typename RHIResourceType, GLenum GLTypeEnum, EShaderFrequency FrequencyT>
-class TOpenGLShader : public RHIResourceType
+class FOpenGLShader
 {
 public:
-	enum
-	{
-		StaticFrequency = FrequencyT,
-		TypeEnum = GLTypeEnum,
-	};
-
 	/** The OpenGL resource ID. */
-	GLuint Resource;
+	GLuint Resource = 0;
 
 	/** External bindings for this shader. */
 	FOpenGLShaderBindings Bindings;
@@ -246,90 +272,73 @@ public:
 	// List of memory copies from RHIUniformBuffer to packed uniforms
 	TArray<CrossCompiler::FUniformBufferCopyInfo> UniformBuffersCopyInfo;
 
+	FOpenGLCompiledShaderKey ShaderCodeKey;
+
 #if DEBUG_GL_SHADERS
 	TArray<ANSICHAR> GlslCode;
 	const ANSICHAR*  GlslCodeString; // make it easier in VS to see shader code in debug mode; points to begin of GlslCode
 #endif
 
-	/** Constructor. */
-	TOpenGLShader()
-		: Resource(0)
-	{
-		FMemory::Memzero( &Bindings, sizeof(Bindings) );
-	}
+	FOpenGLShader(FRHICommandListImmediate& RHICmdList, TArrayView<const uint8> Code, const FSHAHash& Hash, GLenum TypeEnum);
 
-	/** Destructor. */
-	~TOpenGLShader()
+	~FOpenGLShader()
 	{
 //		if (Resource)
 //		{
 //			glDeleteShader(Resource);
 //		}
 	}
-};
-
-
-typedef TOpenGLShader<FRefCountedObject, GL_VERTEX_SHADER, SF_Vertex> FOpenGLVertexShader;
-typedef TOpenGLShader<FRefCountedObject, GL_FRAGMENT_SHADER, SF_Pixel> FOpenGLPixelShader;
-typedef TOpenGLShader<FRefCountedObject, GL_GEOMETRY_SHADER, SF_Geometry> FOpenGLGeometryShader;
 
+private:
+	void Compile(const FOpenGLCodeHeader& Header, const FSHAHash& LibraryHash, GLenum TypeEnum);
+};
 
-class FOpenGLComputeShader : public TOpenGLShader<FRefCountedObject, GL_COMPUTE_SHADER, SF_Compute>
+class FOpenGLVertexShader : public FRHIVertexShader, public FOpenGLShader
 {
 public:
-	FOpenGLComputeShader():
-		LinkedProgram(0)
-	{
+	static constexpr EShaderFrequency Frequency = SF_Vertex;
 
-	}
+	FOpenGLVertexShader(FRHICommandListImmediate& RHICmdList, TArrayView<const uint8> Code, const FSHAHash& Hash)
+		: FOpenGLShader(RHICmdList, Code, Hash, GL_VERTEX_SHADER)
+	{}
+};
 
-	bool NeedsTextureStage(int32 TextureStageIndex);
-	int32 MaxTextureStageUsed();
-	const TBitArray<>& GetTextureNeeds(int32& OutMaxTextureStageUsed);
-	const TBitArray<>& GetUAVNeeds(int32& OutMaxUAVUnitUsed) const;
-	bool NeedsUAVStage(int32 UAVStageIndex) const;
+class FOpenGLPixelShader : public FRHIPixelShader, public FOpenGLShader
+{
+public:
+	static constexpr EShaderFrequency Frequency = SF_Pixel;
 
-	FOpenGLLinkedProgram* LinkedProgram;
+	FOpenGLPixelShader(FRHICommandListImmediate& RHICmdList, TArrayView<const uint8> Code, const FSHAHash& Hash)
+		: FOpenGLShader(RHICmdList, Code, Hash, GL_FRAGMENT_SHADER)
+	{}
 };
 
-
-class FOpenGLCompiledShaderKey
+class FOpenGLGeometryShader : public FRHIGeometryShader, public FOpenGLShader
 {
 public:
-	FOpenGLCompiledShaderKey()
-		: TypeEnum(0)
-		, CodeSize(0)
-		, CodeCRC(0)
-	{
-	}
+	static constexpr EShaderFrequency Frequency = SF_Geometry;
 
-	FOpenGLCompiledShaderKey(
-		GLenum InTypeEnum,
-		uint32 InCodeSize,
-		uint32 InCodeCRC
-	)
-		: TypeEnum(InTypeEnum)
-		, CodeSize(InCodeSize)
-		, CodeCRC(InCodeCRC)
-	{
-	}
+	FOpenGLGeometryShader(FRHICommandListImmediate& RHICmdList, TArrayView<const uint8> Code, const FSHAHash& Hash)
+		: FOpenGLShader(RHICmdList, Code, Hash, GL_GEOMETRY_SHADER)
+	{}
+};
 
-	friend bool operator ==(const FOpenGLCompiledShaderKey& A, const FOpenGLCompiledShaderKey& B)
-	{
-		return A.TypeEnum == B.TypeEnum && A.CodeSize == B.CodeSize && A.CodeCRC == B.CodeCRC;
-	}
+class FOpenGLComputeShader : public FRHIComputeShader, public FOpenGLShader
+{
+public:
+	static constexpr EShaderFrequency Frequency = SF_Compute;
 
-	friend uint32 GetTypeHash(const FOpenGLCompiledShaderKey& Key)
-	{
-		return GetTypeHash(Key.TypeEnum) ^ GetTypeHash(Key.CodeSize) ^ GetTypeHash(Key.CodeCRC);
-	}
+	FOpenGLComputeShader(FRHICommandListImmediate& RHICmdList, TArrayView<const uint8> Code, const FSHAHash& Hash)
+		: FOpenGLShader(RHICmdList, Code, Hash, GL_COMPUTE_SHADER)
+	{}
 
-	uint32 GetCodeCRC() const { return CodeCRC; }
+	bool NeedsTextureStage(int32 TextureStageIndex);
+	int32 MaxTextureStageUsed();
+	const TBitArray<>& GetTextureNeeds(int32& OutMaxTextureStageUsed);
+	const TBitArray<>& GetUAVNeeds(int32& OutMaxUAVUnitUsed) const;
+	bool NeedsUAVStage(int32 UAVStageIndex) const;
 
-private:
-	GLenum TypeEnum;
-	uint32 CodeSize;
-	uint32 CodeCRC;
+	FOpenGLLinkedProgram* LinkedProgram = nullptr;
 };
 
 /**
diff --git a/Engine/Source/Runtime/OpenGLDrv/Public/OpenGLState.h b/Engine/Source/Runtime/OpenGLDrv/Public/OpenGLState.h
--- a/Engine/Source/Runtime/OpenGLDrv/Public/OpenGLState.h
+++ b/Engine/Source/Runtime/OpenGLDrv/Public/OpenGLState.h
@@ -68,7 +68,6 @@ class FOpenGLSamplerState : public FRHISamplerState
 public:
 	GLuint Resource;
 	FOpenGLSamplerStateData Data;
-	FOpenGLAssertRHIThreadFence CreationFence;
 
 	~FOpenGLSamplerState();
 };
diff --git a/Engine/Source/Runtime/RenderCore/Private/ClearReplacementShaders.cpp b/Engine/Source/Runtime/RenderCore/Private/ClearReplacementShaders.cpp
--- a/Engine/Source/Runtime/RenderCore/Private/ClearReplacementShaders.cpp
+++ b/Engine/Source/Runtime/RenderCore/Private/ClearReplacementShaders.cpp
@@ -53,66 +53,52 @@ IMPLEMENT_SHADER_TYPE4_WITH_TEMPLATE_PREFIX(template<>, RENDERCORE_API, FClearRe
 
 void CreateClearReplacementShaders()
 {
-#if !WITH_EDITOR	// in editor, this function overall has a lesser sence due to the ability to switch the preview modes on the fly and recompile the shaders, which both can execute the code on GT while RT is already created.
-	// if RHI supports MT shader creation, then we don't care to init here
-	ensureMsgf(!GRHISupportsMultithreadedShaderCreation, TEXT("InitClearReplacementShaders() is called while GRHISupportsMultithreadedShaderCreation is true. This is an unnecessary call."));
-	ensureMsgf(IsInRenderingThread() || GRHISupportsMultithreadedShaderCreation, TEXT("InitClearReplacementShaders() is expected to be called from the render thread if GRHISupportsMultithreadedShaderCreation is false."));
+	ForceInitGlobalShaderType<FClearReplacementCS_Buffer_Uint_Bounds>();
 
-#define UE_CREATE_SHADER(TShaderType) \
-	{ \
-		TShaderMapRef<TShaderType> ComputeShader(GetGlobalShaderMap(GMaxRHIFeatureLevel)); \
-		FRHIComputeShader* ShaderRHI = ComputeShader.GetComputeShader(); \
-	}
-
-	UE_CREATE_SHADER(FClearReplacementCS_Buffer_Uint_Bounds);
-
-	UE_CREATE_SHADER(FClearReplacementCS_Buffer_Float_Bounds);
-	UE_CREATE_SHADER(FClearReplacementCS_Buffer_Float4_Bounds);
-	UE_CREATE_SHADER(FClearReplacementCS_LargeBuffer_Uint_Bounds);
-	UE_CREATE_SHADER(FClearReplacementCS_LargeBuffer_Float_Bounds);
-	UE_CREATE_SHADER(FClearReplacementCS_LargeBuffer_Float4_Bounds);
-	UE_CREATE_SHADER(FClearReplacementCS_StructuredBuffer_Uint_Bounds);
-	UE_CREATE_SHADER(FClearReplacementCS_StructuredBuffer_Float_Bounds);
-	UE_CREATE_SHADER(FClearReplacementCS_StructuredBuffer_Float4_Bounds);
-	UE_CREATE_SHADER(FClearReplacementCS_Texture3D_Float4_Bounds);
-	UE_CREATE_SHADER(FClearReplacementCS_Texture2D_Float4_Bounds);
-	UE_CREATE_SHADER(FClearReplacementCS_Texture2DArray_Float4_Bounds);
+	ForceInitGlobalShaderType<FClearReplacementCS_Buffer_Float_Bounds>();
+	ForceInitGlobalShaderType<FClearReplacementCS_Buffer_Float4_Bounds>();
+	ForceInitGlobalShaderType<FClearReplacementCS_LargeBuffer_Uint_Bounds>();
+	ForceInitGlobalShaderType<FClearReplacementCS_LargeBuffer_Float_Bounds>();
+	ForceInitGlobalShaderType<FClearReplacementCS_LargeBuffer_Float4_Bounds>();
+	ForceInitGlobalShaderType<FClearReplacementCS_StructuredBuffer_Uint_Bounds>();
+	ForceInitGlobalShaderType<FClearReplacementCS_StructuredBuffer_Float_Bounds>();
+	ForceInitGlobalShaderType<FClearReplacementCS_StructuredBuffer_Float4_Bounds>();
+	ForceInitGlobalShaderType<FClearReplacementCS_Texture3D_Float4_Bounds>();
+	ForceInitGlobalShaderType<FClearReplacementCS_Texture2D_Float4_Bounds>();
+	ForceInitGlobalShaderType<FClearReplacementCS_Texture2DArray_Float4_Bounds>();
 
 	// Compute shaders for clearing each resource type. No bounds checks enabled.
-	UE_CREATE_SHADER(FClearReplacementCS_Buffer_Uint_Zero);
-	UE_CREATE_SHADER(FClearReplacementCS_StructuredBuffer_Uint_Zero);
-	UE_CREATE_SHADER(FClearReplacementCS_Texture2DArray_Uint_Zero);
-	UE_CREATE_SHADER(FClearReplacementCS_Buffer_Uint);
-	UE_CREATE_SHADER(FClearReplacementCS_LargeBuffer_Uint);
-	UE_CREATE_SHADER(FClearReplacementCS_StructuredBuffer_Uint);
-	UE_CREATE_SHADER(FClearReplacementCS_Texture2DArray_Uint);
+	ForceInitGlobalShaderType<FClearReplacementCS_Buffer_Uint_Zero>();
+	ForceInitGlobalShaderType<FClearReplacementCS_StructuredBuffer_Uint_Zero>();
+	ForceInitGlobalShaderType<FClearReplacementCS_Texture2DArray_Uint_Zero>();
+	ForceInitGlobalShaderType<FClearReplacementCS_Buffer_Uint>();
+	ForceInitGlobalShaderType<FClearReplacementCS_LargeBuffer_Uint>();
+	ForceInitGlobalShaderType<FClearReplacementCS_StructuredBuffer_Uint>();
+	ForceInitGlobalShaderType<FClearReplacementCS_Texture2DArray_Uint>();
 
-	UE_CREATE_SHADER(FClearReplacementCS_Texture3D_Float4);
-	UE_CREATE_SHADER(FClearReplacementCS_Texture2D_Float4);
-	UE_CREATE_SHADER(FClearReplacementCS_Texture2DArray_Float4);
+	ForceInitGlobalShaderType<FClearReplacementCS_Texture3D_Float4>();
+	ForceInitGlobalShaderType<FClearReplacementCS_Texture2D_Float4>();
+	ForceInitGlobalShaderType<FClearReplacementCS_Texture2DArray_Float4>();
 
 	// Used by ClearUAV_T in ClearQuad.cpp
-	UE_CREATE_SHADER(FClearReplacementCS_Texture3D_Uint4);
-	UE_CREATE_SHADER(FClearReplacementCS_Texture2D_Uint4);
-	UE_CREATE_SHADER(FClearReplacementCS_Texture2DArray_Uint4);
-	UE_CREATE_SHADER(FClearReplacementCS_Buffer_Uint4);
-	UE_CREATE_SHADER(FClearReplacementCS_LargeBuffer_Uint4);
-	UE_CREATE_SHADER(FClearReplacementCS_StructuredBuffer_Uint4);
-
-	UE_CREATE_SHADER(FClearReplacementCS_Buffer_Uint4_Bounds);
-	UE_CREATE_SHADER(FClearReplacementCS_LargeBuffer_Uint4_Bounds);
-	UE_CREATE_SHADER(FClearReplacementCS_StructuredBuffer_Uint4_Bounds);
-	UE_CREATE_SHADER(FClearReplacementCS_Texture3D_Uint4_Bounds);
-	UE_CREATE_SHADER(FClearReplacementCS_Texture2D_Uint4_Bounds);
-	UE_CREATE_SHADER(FClearReplacementCS_Texture2DArray_Uint4_Bounds);
+	ForceInitGlobalShaderType<FClearReplacementCS_Texture3D_Uint4>();
+	ForceInitGlobalShaderType<FClearReplacementCS_Texture2D_Uint4>();
+	ForceInitGlobalShaderType<FClearReplacementCS_Texture2DArray_Uint4>();
+	ForceInitGlobalShaderType<FClearReplacementCS_Buffer_Uint4>();
+	ForceInitGlobalShaderType<FClearReplacementCS_LargeBuffer_Uint4>();
+	ForceInitGlobalShaderType<FClearReplacementCS_StructuredBuffer_Uint4>();
 
-	UE_CREATE_SHADER(FClearReplacementCS_Buffer_Sint4_Bounds);
-	UE_CREATE_SHADER(FClearReplacementCS_LargeBuffer_Sint4_Bounds);
-	UE_CREATE_SHADER(FClearReplacementCS_StructuredBuffer_Sint4_Bounds);
-	UE_CREATE_SHADER(FClearReplacementCS_Texture3D_Sint4_Bounds);
-	UE_CREATE_SHADER(FClearReplacementCS_Texture2D_Sint4_Bounds);
-	UE_CREATE_SHADER(FClearReplacementCS_Texture2DArray_Sint4_Bounds);
+	ForceInitGlobalShaderType<FClearReplacementCS_Buffer_Uint4_Bounds>();
+	ForceInitGlobalShaderType<FClearReplacementCS_LargeBuffer_Uint4_Bounds>();
+	ForceInitGlobalShaderType<FClearReplacementCS_StructuredBuffer_Uint4_Bounds>();
+	ForceInitGlobalShaderType<FClearReplacementCS_Texture3D_Uint4_Bounds>();
+	ForceInitGlobalShaderType<FClearReplacementCS_Texture2D_Uint4_Bounds>();
+	ForceInitGlobalShaderType<FClearReplacementCS_Texture2DArray_Uint4_Bounds>();
 
-#undef UE_CREATE_SHADER
-#endif
+	ForceInitGlobalShaderType<FClearReplacementCS_Buffer_Sint4_Bounds>();
+	ForceInitGlobalShaderType<FClearReplacementCS_LargeBuffer_Sint4_Bounds>();
+	ForceInitGlobalShaderType<FClearReplacementCS_StructuredBuffer_Sint4_Bounds>();
+	ForceInitGlobalShaderType<FClearReplacementCS_Texture3D_Sint4_Bounds>();
+	ForceInitGlobalShaderType<FClearReplacementCS_Texture2D_Sint4_Bounds>();
+	ForceInitGlobalShaderType<FClearReplacementCS_Texture2DArray_Sint4_Bounds>();
 }
diff --git a/Engine/Source/Runtime/RenderCore/Private/GlobalShader.cpp b/Engine/Source/Runtime/RenderCore/Private/GlobalShader.cpp
--- a/Engine/Source/Runtime/RenderCore/Private/GlobalShader.cpp
+++ b/Engine/Source/Runtime/RenderCore/Private/GlobalShader.cpp
@@ -809,3 +809,46 @@ void FGlobalShaderMap::LoadFromGlobalArchive(FArchive& Ar)
 		}
 	}
 }
+
+RENDERCORE_API ERecursiveShader GRequiredRecursiveShaders = ERecursiveShader::None;
+
+void ForceInitGlobalShaderType(FShaderType& ShaderType)
+{
+	FGlobalShaderMap* GlobalShaderMap = GetGlobalShaderMap(GMaxRHIFeatureLevel);
+	for (int32 Permutation = 0; Permutation < ShaderType.GetPermutationCount(); ++Permutation)
+	{
+		TShaderRef<FShader> ShaderRef = GlobalShaderMap->GetShader(&ShaderType, Permutation);
+
+		if (ShaderRef.IsValid())
+		{
+			FShaderMapResource& MapResource = ShaderRef.GetResourceChecked();
+			for (int32 Index = 0; Index < MapResource.GetNumShaders(); ++Index)
+			{
+				MapResource.GetShader(Index);
+			}
+		}
+	}
+}
+
+RENDERCORE_API void CreateRecursiveShaders()
+{
+	ensureMsgf(!GRHISupportsMultithreadedShaderCreation, TEXT("CreateRecursiveShaders() is called while GRHISupportsMultithreadedShaderCreation is true. This is an unnecessary call."));
+	ensureMsgf(IsInRenderingThread(), TEXT("CreateRecursiveShaders() is expected to be called from the render thread."));
+
+	if (EnumHasAnyFlags(GRequiredRecursiveShaders, ERecursiveShader::Resolve))
+	{
+		extern void CreateResolveShaders();
+		CreateResolveShaders();
+	}
+
+	if (EnumHasAnyFlags(GRequiredRecursiveShaders, ERecursiveShader::Clear))
+	{
+		extern void CreateClearReplacementShaders();
+		CreateClearReplacementShaders();
+	}
+
+	if (EnumHasAnyFlags(GRequiredRecursiveShaders, ERecursiveShader::Null))
+	{
+		ForceInitGlobalShaderType<FNULLPS>();
+	}
+}
diff --git a/Engine/Source/Runtime/RenderCore/Private/ResolveShader.cpp b/Engine/Source/Runtime/RenderCore/Private/ResolveShader.cpp
--- a/Engine/Source/Runtime/RenderCore/Private/ResolveShader.cpp
+++ b/Engine/Source/Runtime/RenderCore/Private/ResolveShader.cpp
@@ -247,3 +247,19 @@ void FResolveArrayVS::SetParameters(FRHICommandList& RHICmdList, const FResolveR
 {
 	return FResolveVS::SetParameters(RHICmdList, SrcBounds, DstBounds, DstSurfaceWidth, DstSurfaceHeight);
 }
+
+void CreateResolveShaders()
+{
+	ForceInitGlobalShaderType<FResolveDepthPS>();
+	ForceInitGlobalShaderType<FResolveDepth2XPS>();
+	ForceInitGlobalShaderType<FResolveDepth4XPS>();
+	ForceInitGlobalShaderType<FResolveDepth8XPS>();
+	ForceInitGlobalShaderType<FResolveDepthArrayPS>();
+	ForceInitGlobalShaderType<FResolveDepthArray2XPS>();
+	ForceInitGlobalShaderType<FResolveDepthArray4XPS>();
+	ForceInitGlobalShaderType<FResolveDepthArray8XPS>();
+	ForceInitGlobalShaderType<FResolveSingleSamplePS>();
+
+	ForceInitGlobalShaderType<FResolveVS>();
+	ForceInitGlobalShaderType<FResolveArrayVS>();
+}
diff --git a/Engine/Source/Runtime/RenderCore/Public/ClearReplacementShaders.h b/Engine/Source/Runtime/RenderCore/Public/ClearReplacementShaders.h
--- a/Engine/Source/Runtime/RenderCore/Public/ClearReplacementShaders.h
+++ b/Engine/Source/Runtime/RenderCore/Public/ClearReplacementShaders.h
@@ -317,11 +317,6 @@ typedef TClearReplacementCS<EClearReplacementResourceType::Texture3D,        FCl
 typedef TClearReplacementCS<EClearReplacementResourceType::Texture2D,        FClearReplacementBase_Sint4_Bounds>  FClearReplacementCS_Texture2D_Sint4_Bounds;
 typedef TClearReplacementCS<EClearReplacementResourceType::Texture2DArray,   FClearReplacementBase_Sint4_Bounds>  FClearReplacementCS_Texture2DArray_Sint4_Bounds;
 
-/**
- * Force creation of the above clear replacement shaders. This is sometimes useful if multi-threaded creation (i.e. on demand) isn't supported.
- */
-void RENDERCORE_API CreateClearReplacementShaders();
-
 /**
  * Helper functions for running the clear replacement shader for specific resource types, values types and number of channels.
  * Can be used from inside RHIs via FRHICommandList_RecursiveHazardous. ResourceBindCallback is provided to allow the RHI to override
diff --git a/Engine/Source/Runtime/RenderCore/Public/GlobalShader.h b/Engine/Source/Runtime/RenderCore/Public/GlobalShader.h
--- a/Engine/Source/Runtime/RenderCore/Public/GlobalShader.h
+++ b/Engine/Source/Runtime/RenderCore/Public/GlobalShader.h
@@ -446,3 +446,39 @@ inline FGlobalShaderMap* GetGlobalShaderMap(ERHIFeatureLevel::Type FeatureLevel)
 #define DECLARE_GLOBAL_SHADER(ShaderClass) DECLARE_SHADER_TYPE(ShaderClass, Global)
 #define IMPLEMENT_GLOBAL_SHADER(ShaderClass,SourceFilename,FunctionName,Frequency) IMPLEMENT_SHADER_TYPE(,ShaderClass,TEXT(SourceFilename),TEXT(FunctionName),Frequency)
 
+/**
+* Groups of global shaders which can be used in recursive
+* RHI command lists by platform RHI implementations.
+*/
+enum class ERecursiveShader
+{
+	None = 0,
+
+	// Shaders used for resolving MSAA textures (defined in ResolveShader.h)
+	Resolve = 1 << 0,
+
+	// Shaders used for clearing textures and buffers (defined in ClearReplacementShaders.h)
+	Clear = 1 << 1,
+
+	// A null pixel shader (defined above)
+	Null = 1 << 2
+};
+ENUM_CLASS_FLAGS(ERecursiveShader)
+
+/*
+* Set by the platform RHI implementations to indicate which groups of global shaders will be used recursively.
+*/
+extern RENDERCORE_API ERecursiveShader GRequiredRecursiveShaders;
+
+/*
+* Called on the render thread when the global shader map is available to force-init certain global shaders as specified by GRequiredRecursiveShaders.
+*/
+extern RENDERCORE_API void CreateRecursiveShaders();
+
+void ForceInitGlobalShaderType(FShaderType& ShaderType);
+
+template <typename TShaderType>
+inline void ForceInitGlobalShaderType()
+{
+	ForceInitGlobalShaderType(TShaderType::StaticType);
+}
diff --git a/Engine/Source/Runtime/Renderer/Private/SceneRendering.h b/Engine/Source/Runtime/Renderer/Private/SceneRendering.h
--- a/Engine/Source/Runtime/Renderer/Private/SceneRendering.h
+++ b/Engine/Source/Runtime/Renderer/Private/SceneRendering.h
@@ -2563,7 +2563,7 @@ UE::Tasks::FTask LaunchSceneRenderTask(const TCHAR* DebugName, LambdaType&& Lamb
 
 	return UE::Tasks::Launch(DebugName, [Lambda = MoveTemp(Lambda)]
 	{
-		FTaskTagScope Scope(ETaskTag::EParallelRenderingThread);
+		FOptionalTaskTagScope Scope(ETaskTag::EParallelRenderingThread);
 		Lambda();
 	},
 	TaskPriority, bExecuteInParallel ? UE::Tasks::EExtendedTaskPriority::None : UE::Tasks::EExtendedTaskPriority::Inline);
diff --git a/Engine/Source/Runtime/VulkanRHI/Private/VulkanRHI.cpp b/Engine/Source/Runtime/VulkanRHI/Private/VulkanRHI.cpp
--- a/Engine/Source/Runtime/VulkanRHI/Private/VulkanRHI.cpp
+++ b/Engine/Source/Runtime/VulkanRHI/Private/VulkanRHI.cpp
@@ -287,9 +287,6 @@ void FVulkanDynamicRHI::Init()
 
 void FVulkanDynamicRHI::PostInit()
 {
-	//work around layering violation
-	TShaderMapRef<FNULLPS>(GetGlobalShaderMap(GMaxRHIFeatureLevel)).GetPixelShader();
-
 #if VULKAN_RHI_RAYTRACING
 	if (GRHISupportsRayTracing)
 	{
diff --git a/Engine/Source/Runtime/VulkanRHI/Private/VulkanViewport.cpp b/Engine/Source/Runtime/VulkanRHI/Private/VulkanViewport.cpp
--- a/Engine/Source/Runtime/VulkanRHI/Private/VulkanViewport.cpp
+++ b/Engine/Source/Runtime/VulkanRHI/Private/VulkanViewport.cpp
@@ -1087,24 +1087,12 @@ void FVulkanDynamicRHI::RHITick(float DeltaTime)
 {
 	check(IsInGameThread());
 	FVulkanDevice* VulkanDevice = GetDevice();
-	static bool bRequestNULLPixelShader = true;
-	bool bRequested = bRequestNULLPixelShader;
-	ENQUEUE_RENDER_COMMAND(TempFrameReset)(
-		[VulkanDevice, bRequested](FRHICommandListImmediate& RHICmdList)
-		{
-			if (bRequested)
-			{
-				//work around layering violation
-				TShaderMapRef<FNULLPS>(GetGlobalShaderMap(GMaxRHIFeatureLevel)).GetPixelShader();
-			}
 
-			VulkanDevice->GetImmediateContext().GetTempFrameAllocationBuffer().Reset();
-		});
-
-	if (bRequestNULLPixelShader)
+	ENQUEUE_RENDER_COMMAND(TempFrameReset)(
+		[VulkanDevice](FRHICommandListImmediate& RHICmdList)
 	{
-		bRequestNULLPixelShader = false;
-	}
+		VulkanDevice->GetImmediateContext().GetTempFrameAllocationBuffer().Reset();
+	});
 }
 
 FTexture2DRHIRef FVulkanDynamicRHI::RHIGetViewportBackBuffer(FRHIViewport* ViewportRHI)
diff --git a/Engine/Source/Runtime/Windows/D3D11RHI/Private/D3D11RHIPrivate.h b/Engine/Source/Runtime/Windows/D3D11RHI/Private/D3D11RHIPrivate.h
--- a/Engine/Source/Runtime/Windows/D3D11RHI/Private/D3D11RHIPrivate.h
+++ b/Engine/Source/Runtime/Windows/D3D11RHI/Private/D3D11RHIPrivate.h
@@ -426,7 +426,6 @@ public:
 
 	// FDynamicRHI interface.
 	virtual void Init() override;
-	virtual void PostInit() override;
 	virtual void Shutdown() override;
 	virtual const TCHAR* GetName() override { return TEXT("D3D11"); }
 
diff --git a/Engine/Source/Runtime/Windows/D3D11RHI/Private/Windows/WindowsD3D11Device.cpp b/Engine/Source/Runtime/Windows/D3D11RHI/Private/Windows/WindowsD3D11Device.cpp
--- a/Engine/Source/Runtime/Windows/D3D11RHI/Private/Windows/WindowsD3D11Device.cpp
+++ b/Engine/Source/Runtime/Windows/D3D11RHI/Private/Windows/WindowsD3D11Device.cpp
@@ -1148,20 +1148,6 @@ void FD3D11DynamicRHI::Init()
 	InitD3DDevice();
 }
 
-void FD3D11DynamicRHI::PostInit()
-{
-	if (!FPlatformProperties::RequiresCookedData())
-	{
-		// Make sure all global shaders are complete at this point
-		extern RENDERCORE_API const int32 GlobalShaderMapId;
-
-		TArray<int32> ShaderMapIds;
-		ShaderMapIds.Add(GlobalShaderMapId);
-
-		GShaderCompilingManager->FinishCompilation(TEXT("Global"), ShaderMapIds);
-	}
-}
-
 bool FD3D11DynamicRHI::IsQuadBufferStereoEnabled()
 {
 	return bIsQuadBufferStereoEnabled;
@@ -1731,7 +1717,8 @@ void FD3D11DynamicRHI::InitD3DDevice()
 		GTexturePoolSize = 0;
 
 		// turn off creation on other threads for NVidia since a driver heuristic will notice that and make the creation synchronous, and that is not desirable given that large number of shaders will still be created on a single thread
-		GRHISupportsMultithreadedShaderCreation = !IsRHIDeviceNVIDIA(); 
+		GRHISupportsMultithreadedShaderCreation = !IsRHIDeviceNVIDIA();
+		GRequiredRecursiveShaders = ERecursiveShader::Resolve | ERecursiveShader::Clear;
 
 		UE_LOG(LogD3D11RHI, Log, TEXT("    GPU DeviceId: 0x%x (for the marketing name, search the web for \"GPU Device Id\")"), Adapter.DXGIAdapterDesc.DeviceId);
 
