Prevent Ogl's blendstate from producing different RHI Initializers from those used to create it.
This stops  'PSO from the cache was not found in the cache!' warning from being emited when processing the PSO cache.

diff --git a/Engine/Source/Runtime/OpenGLDrv/Private/OpenGLState.cpp b/Engine/Source/Runtime/OpenGLDrv/Private/OpenGLState.cpp
--- a/Engine/Source/Runtime/OpenGLDrv/Private/OpenGLState.cpp
+++ b/Engine/Source/Runtime/OpenGLDrv/Private/OpenGLState.cpp
@@ -401,44 +401,10 @@ bool FOpenGLDepthStencilState::GetInitializer(FDepthStencilStateInitializerRHI&
 	return true;
 }
 
-bool FOpenGLBlendState::GetInitializer(FBlendStateInitializerRHI& Init)
-{
-	Init.bUseIndependentRenderTargetBlendStates = true;
-	Init.bUseAlphaToCoverage = Data.bUseAlphaToCoverage;
-	for(uint32 RenderTargetIndex = 0;RenderTargetIndex < MaxSimultaneousRenderTargets;++RenderTargetIndex)
-	{
-		FOpenGLBlendStateData::FRenderTarget const& RenderTarget = Data.RenderTargets[RenderTargetIndex];
-		FBlendStateInitializerRHI::FRenderTarget& RenderTargetInitializer = Init.RenderTargets[RenderTargetIndex];
-		
-		RenderTargetInitializer.ColorBlendOp = TranslateBlendOp(RenderTarget.ColorBlendOperation);
-		RenderTargetInitializer.ColorSrcBlend = TranslateBlendFactor(RenderTarget.ColorSourceBlendFactor);
-		RenderTargetInitializer.ColorDestBlend = TranslateBlendFactor(RenderTarget.ColorDestBlendFactor);
-		Init.bUseIndependentRenderTargetBlendStates &= (RenderTargetInitializer.ColorBlendOp == Init.RenderTargets[0].ColorBlendOp);
-		Init.bUseIndependentRenderTargetBlendStates &= (RenderTargetInitializer.ColorSrcBlend == Init.RenderTargets[0].ColorSrcBlend);
-		Init.bUseIndependentRenderTargetBlendStates &= (RenderTargetInitializer.ColorDestBlend == Init.RenderTargets[0].ColorDestBlend);
-		
-		RenderTargetInitializer.AlphaBlendOp = TranslateBlendOp(RenderTarget.AlphaBlendOperation);
-		RenderTargetInitializer.AlphaSrcBlend = TranslateBlendFactor(RenderTarget.AlphaSourceBlendFactor);
-		RenderTargetInitializer.AlphaDestBlend = TranslateBlendFactor(RenderTarget.AlphaDestBlendFactor);
-		Init.bUseIndependentRenderTargetBlendStates &= (RenderTargetInitializer.AlphaBlendOp == Init.RenderTargets[0].AlphaBlendOp);
-		Init.bUseIndependentRenderTargetBlendStates &= (RenderTargetInitializer.AlphaSrcBlend == Init.RenderTargets[0].AlphaSrcBlend);
-		Init.bUseIndependentRenderTargetBlendStates &= (RenderTargetInitializer.AlphaDestBlend == Init.RenderTargets[0].AlphaDestBlend);
-		
-		uint32 Mask = CW_NONE;
-		Mask |= (RenderTarget.ColorWriteMaskR) ? CW_RED : 0;
-		Mask |= (RenderTarget.ColorWriteMaskG) ? CW_GREEN : 0;
-		Mask |= (RenderTarget.ColorWriteMaskB) ? CW_BLUE : 0;
-		Mask |= (RenderTarget.ColorWriteMaskA) ? CW_ALPHA : 0;
-		RenderTargetInitializer.ColorWriteMask = (EColorWriteMask)Mask;
-		
-		Init.bUseIndependentRenderTargetBlendStates &= (RenderTargetInitializer.ColorWriteMask == Init.RenderTargets[0].ColorWriteMask);
-	}
-	return true;
-}
-
 FBlendStateRHIRef FOpenGLDynamicRHI::RHICreateBlendState(const FBlendStateInitializerRHI& Initializer)
 {
-	FOpenGLBlendState* BlendState = new FOpenGLBlendState;
+	FOpenGLBlendState* BlendState = new FOpenGLBlendState(Initializer);
+
 	BlendState->Data.bUseAlphaToCoverage = Initializer.bUseAlphaToCoverage;
 	for(uint32 RenderTargetIndex = 0;RenderTargetIndex < MaxSimultaneousRenderTargets;++RenderTargetIndex)
 	{
diff --git a/Engine/Source/Runtime/OpenGLDrv/Public/OpenGLState.h b/Engine/Source/Runtime/OpenGLDrv/Public/OpenGLState.h
--- a/Engine/Source/Runtime/OpenGLDrv/Public/OpenGLState.h
+++ b/Engine/Source/Runtime/OpenGLDrv/Public/OpenGLState.h
@@ -183,9 +183,15 @@ struct FOpenGLBlendStateData
 
 class FOpenGLBlendState : public FRHIBlendState
 {
+	FBlendStateInitializerRHI RHIInitializer;
 public:
-	virtual bool GetInitializer(FBlendStateInitializerRHI& Init) override final;
-	
+	FOpenGLBlendState(const FBlendStateInitializerRHI& Initializer) : RHIInitializer(Initializer) {}
+	virtual bool GetInitializer(FBlendStateInitializerRHI& Init) override final
+	{ 
+		Init = RHIInitializer; 
+		return true;
+	}
+
 	FOpenGLBlendStateData Data;
 };
 
