Fix GLES binary cache creation failure when the intended dest dir is already occupied by a file.

diff --git a/Engine/Source/Runtime/OpenGLDrv/Private/OpenGLProgramBinaryFileCache.cpp b/Engine/Source/Runtime/OpenGLDrv/Private/OpenGLProgramBinaryFileCache.cpp
--- a/Engine/Source/Runtime/OpenGLDrv/Private/OpenGLProgramBinaryFileCache.cpp
+++ b/Engine/Source/Runtime/OpenGLDrv/Private/OpenGLProgramBinaryFileCache.cpp
@@ -187,12 +187,11 @@ FOpenGLProgramBinaryCache::FOpenGLProgramBinaryCache(const FString& InCachePathR
 	IFileManager::Get().FindFiles(FoundFiles, *(CachePathRoot/TEXT("*")), true, true);
 	for(FString& FoundFile : FoundFiles)
 	{
-		if (FoundFile != CacheSubDir)
+		const FString FullPath = (CachePathRoot / FoundFile);
+		const bool bIsDir = PlatformFile.DirectoryExists(*FullPath);
+		if (FoundFile != CacheSubDir || !bIsDir)
 		{
-			const FString FullPath = (CachePathRoot / FoundFile);
-			const bool bIsDir = PlatformFile.DirectoryExists(*FullPath);
-			bool bSuccess = false;
-
+			bool bSuccess;
 			if(bIsDir)
 			{
 				bSuccess = PlatformFile.DeleteDirectoryRecursively(*FullPath);
