Fixed a threading issue around GL timing queries
Downgrade GL debug reporting from Fatal to Error as most of messages are not fatal

diff --git a/Engine/Source/Runtime/OpenGLDrv/Private/OpenGLDevice.cpp b/Engine/Source/Runtime/OpenGLDrv/Private/OpenGLDevice.cpp
--- a/Engine/Source/Runtime/OpenGLDrv/Private/OpenGLDevice.cpp
+++ b/Engine/Source/Runtime/OpenGLDrv/Private/OpenGLDevice.cpp
@@ -391,7 +391,7 @@ static void APIENTRY OpenGLDebugMessageCallbackARB(
 	ELogVerbosity::Type Verbosity = ELogVerbosity::Warning;
 	if (Type == GL_DEBUG_TYPE_ERROR_ARB && Severity == GL_DEBUG_SEVERITY_HIGH_ARB)
 	{
-		Verbosity = ELogVerbosity::Fatal;
+		Verbosity = ELogVerbosity::Error;
 	}
 
 	if ((Verbosity & ELogVerbosity::VerbosityMask) <= FLogCategoryLogRHI::CompileTimeVerbosity)
diff --git a/Engine/Source/Runtime/OpenGLDrv/Private/OpenGLQuery.cpp b/Engine/Source/Runtime/OpenGLDrv/Private/OpenGLQuery.cpp
--- a/Engine/Source/Runtime/OpenGLDrv/Private/OpenGLQuery.cpp
+++ b/Engine/Source/Runtime/OpenGLDrv/Private/OpenGLQuery.cpp
@@ -784,8 +784,6 @@ static FOpenGLRenderQuery* GetTimeQuery()
 
 void FOpenGLBufferedGPUTiming::InitResources()
 {
-	VERIFY_GL_SCOPE();
-
 	StaticInitialize(OpenGLRHI, PlatformStaticInitialize);
 
 	CurrentTimestamp = 0;
@@ -1100,10 +1098,11 @@ bool FOpenGLDisjointTimeStampQuery::GetResult( uint64* OutResult/*=NULL*/ )
 
 void FOpenGLDisjointTimeStampQuery::InitResources()
 {
-	VERIFY_GL_SCOPE();
-	if ( IsSupported() )
+	if (IsSupported())
 	{
-		PlatformGetNewRenderQuery(&DisjointQuery, &Context);
+		FRHICommandListExecutor::GetImmediateCommandList().EnqueueLambda([this](FRHICommandListImmediate&) { 
+				PlatformGetNewRenderQuery(&DisjointQuery, &Context); 
+		});
 	}
 }
 
