Fallback to R16F as a replacement for a missing R16Unorm pixel format on GLES

diff --git a/Engine/Source/Runtime/OpenGLDrv/Private/OpenGLDevice.cpp b/Engine/Source/Runtime/OpenGLDrv/Private/OpenGLDevice.cpp
--- a/Engine/Source/Runtime/OpenGLDrv/Private/OpenGLDevice.cpp
+++ b/Engine/Source/Runtime/OpenGLDrv/Private/OpenGLDevice.cpp
@@ -1134,8 +1134,10 @@ static void InitRHICapabilitiesForGL()
 	bool bNeedsToSwizzleRedBlue = GL5551Format != GL_UNSIGNED_SHORT_5_5_5_1;
 	SetupTextureFormat( PF_B5G5R5A1_UNORM,		FOpenGLTextureFormat( GL_RGB5_A1,               GL_RGB5_A1,             GL_RGBA,			GL5551Format,					false,          bNeedsToSwizzleRedBlue));
 
+	// Unfortunatelly most OpenGLES devices do not support R16Unorm pixel format, fallback to a less precise R16F
+	SetupTextureFormat( PF_G16,					FOpenGLTextureFormat( GL_R16F,					GL_R16F,				GL_RED,				GL_HALF_FLOAT,						false,	false));
+	//SetupTextureFormat( PF_G16,				FOpenGLTextureFormat( GL_R16,					GL_R16,					GL_RED,				GL_UNSIGNED_SHORT,					false,	false));
 
-	SetupTextureFormat( PF_G16,					FOpenGLTextureFormat( GL_R16,					GL_R16,					GL_RED,				GL_UNSIGNED_SHORT,					false,	false));
 	SetupTextureFormat( PF_R32_FLOAT,			FOpenGLTextureFormat( GL_R32F,					GL_R32F,				GL_RED,				GL_FLOAT,							false,	false));
 	SetupTextureFormat( PF_G16R16F,				FOpenGLTextureFormat( GL_RG16F,					GL_RG16F,				GL_RG,				GL_HALF_FLOAT,						false,	false));
 	SetupTextureFormat( PF_G16R16F_FILTER,		FOpenGLTextureFormat( GL_RG16F,					GL_RG16F,				GL_RG,				GL_HALF_FLOAT,						false,	false));
