Fallback to R16F on Android as a replacement for a missing R16Unorm pixel format on GLES

diff --git a/Engine/Source/Developer/Android/AndroidTargetPlatform/Private/AndroidTargetPlatform.cpp b/Engine/Source/Developer/Android/AndroidTargetPlatform/Private/AndroidTargetPlatform.cpp
--- a/Engine/Source/Developer/Android/AndroidTargetPlatform/Private/AndroidTargetPlatform.cpp
+++ b/Engine/Source/Developer/Android/AndroidTargetPlatform/Private/AndroidTargetPlatform.cpp
@@ -77,6 +77,7 @@ namespace AndroidTexFormat
 	const static FName NameG8(TEXT("G8"));
 	const static FName NameRGBA16F(TEXT("RGBA16F"));
 	const static FName NameR16F(TEXT("R16F"));
+	const static FName NameG16(TEXT("G16"));
 
 	//A1RGB555 is mapped to RGB555A1, because OpenGL GL_RGB5_A1 only supports alpha on the lowest bit
 	const static FName NameA1RGB555(TEXT("A1RGB555"));
@@ -84,7 +85,8 @@ namespace AndroidTexFormat
 
 	const static FName GenericRemap[][2] =
 	{
-		{ NameA1RGB555,		NameRGB555A1				},
+		{ NameA1RGB555,		NameRGB555A1			},
+		{ NameG16,			NameR16F				}, // GLES does not support R16Unorm, fallback all Android to R16F
 	};
 	
 	static const FName NameASTC_RGB_HDR(TEXT("ASTC_RGB_HDR"));
diff --git a/Engine/Source/Runtime/OpenGLDrv/Private/OpenGLDevice.cpp b/Engine/Source/Runtime/OpenGLDrv/Private/OpenGLDevice.cpp
--- a/Engine/Source/Runtime/OpenGLDrv/Private/OpenGLDevice.cpp
+++ b/Engine/Source/Runtime/OpenGLDrv/Private/OpenGLDevice.cpp
@@ -1133,11 +1133,13 @@ static void InitRHICapabilitiesForGL()
 	//the last 5 bits of GL_UNSIGNED_SHORT_1_5_5_5_REV is Red, while the last 5 bits of DXGI_FORMAT_B5G5R5A1_UNORM is Blue
 	bool bNeedsToSwizzleRedBlue = GL5551Format != GL_UNSIGNED_SHORT_5_5_5_1;
 	SetupTextureFormat( PF_B5G5R5A1_UNORM,		FOpenGLTextureFormat( GL_RGB5_A1,               GL_RGB5_A1,             GL_RGBA,			GL5551Format,					false,          bNeedsToSwizzleRedBlue));
-
+	
+#if PLATFORM_ANDROID
 	// Unfortunatelly most OpenGLES devices do not support R16Unorm pixel format, fallback to a less precise R16F
 	SetupTextureFormat( PF_G16,					FOpenGLTextureFormat( GL_R16F,					GL_R16F,				GL_RED,				GL_HALF_FLOAT,						false,	false));
-	//SetupTextureFormat( PF_G16,				FOpenGLTextureFormat( GL_R16,					GL_R16,					GL_RED,				GL_UNSIGNED_SHORT,					false,	false));
-
+#else
+	SetupTextureFormat( PF_G16,					FOpenGLTextureFormat( GL_R16,					GL_R16,					GL_RED,				GL_UNSIGNED_SHORT,					false,	false));
+#endif
 	SetupTextureFormat( PF_R32_FLOAT,			FOpenGLTextureFormat( GL_R32F,					GL_R32F,				GL_RED,				GL_FLOAT,							false,	false));
 	SetupTextureFormat( PF_G16R16F,				FOpenGLTextureFormat( GL_RG16F,					GL_RG16F,				GL_RG,				GL_HALF_FLOAT,						false,	false));
 	SetupTextureFormat( PF_G16R16F_FILTER,		FOpenGLTextureFormat( GL_RG16F,					GL_RG16F,				GL_RG,				GL_HALF_FLOAT,						false,	false));
diff --git a/Engine/Source/Runtime/VulkanRHI/Private/Android/VulkanAndroidPlatform.h b/Engine/Source/Runtime/VulkanRHI/Private/Android/VulkanAndroidPlatform.h
--- a/Engine/Source/Runtime/VulkanRHI/Private/Android/VulkanAndroidPlatform.h
+++ b/Engine/Source/Runtime/VulkanRHI/Private/Android/VulkanAndroidPlatform.h
@@ -71,6 +71,9 @@ public:
 	static bool SupportsBCTextureFormats() { return false; }
 	static bool SupportsASTCTextureFormats() { return true; }
 	static bool SupportsETC2TextureFormats() { return true; }
+	// GLES does not support R16Unorm, so all Android has to fallback to R16F instead
+	static bool SupportsR16UnormTextureFormat() { return false; }
+	
 	static bool SupportsQuerySurfaceProperties() { return false; }
 
 	static void SetupFeatureLevels()
diff --git a/Engine/Source/Runtime/VulkanRHI/Private/VulkanDevice.cpp b/Engine/Source/Runtime/VulkanRHI/Private/VulkanDevice.cpp
--- a/Engine/Source/Runtime/VulkanRHI/Private/VulkanDevice.cpp
+++ b/Engine/Source/Runtime/VulkanRHI/Private/VulkanDevice.cpp
@@ -603,7 +603,6 @@ void FVulkanDevice::SetupFormats()
 	// Default formats
 	MapFormatSupport(PF_B8G8R8A8, { VK_FORMAT_B8G8R8A8_UNORM }, ComponentMappingRGBA);
 	MapFormatSupport(PF_G8, { VK_FORMAT_R8_UNORM }, ComponentMappingR001);
-	MapFormatSupport(PF_G16, { VK_FORMAT_R16_UNORM, VK_FORMAT_R16_SFLOAT }, ComponentMappingR001);
 	MapFormatSupport(PF_FloatRGB, { VK_FORMAT_B10G11R11_UFLOAT_PACK32, VK_FORMAT_R16G16B16_SFLOAT, VK_FORMAT_R16G16B16A16_SFLOAT }, ComponentMappingRGB1, ColorRenderTargetRequiredCapabilities);
 	MapFormatSupport(PF_FloatRGBA, { VK_FORMAT_R16G16B16A16_SFLOAT }, ComponentMappingRGBA, 8);
 	MapFormatSupport(PF_ShadowDepth, { VK_FORMAT_D16_UNORM }, ComponentMappingRIII);
@@ -707,6 +706,14 @@ void FVulkanDevice::SetupFormats()
 		MapFormatSupport(PF_ETC2_R11_EAC,  { VK_FORMAT_EAC_R11_UNORM_BLOCK },       ComponentMappingR001);
 		MapFormatSupport(PF_ETC2_RG11_EAC, { VK_FORMAT_EAC_R11G11_UNORM_BLOCK },    ComponentMappingRG01);
 	}
+	if (FVulkanPlatform::SupportsR16UnormTextureFormat())
+	{
+		MapFormatSupport(PF_G16, { VK_FORMAT_R16_UNORM, VK_FORMAT_R16_SFLOAT }, ComponentMappingR001);
+	}
+	else
+	{
+		MapFormatSupport(PF_G16, { VK_FORMAT_R16_SFLOAT, VK_FORMAT_R16_UNORM }, ComponentMappingR001);
+	}
 
 	if (GetOptionalExtensions().HasEXTASTCDecodeMode)
 	{
diff --git a/Engine/Source/Runtime/VulkanRHI/Private/VulkanGenericPlatform.h b/Engine/Source/Runtime/VulkanRHI/Private/VulkanGenericPlatform.h
--- a/Engine/Source/Runtime/VulkanRHI/Private/VulkanGenericPlatform.h
+++ b/Engine/Source/Runtime/VulkanRHI/Private/VulkanGenericPlatform.h
@@ -47,6 +47,7 @@ public:
 	static bool SupportsBCTextureFormats() { return true; }
 	static bool SupportsASTCTextureFormats() { return false; }
 	static bool SupportsETC2TextureFormats() { return false; }
+	static bool SupportsR16UnormTextureFormat() { return true; }
 
 	// most platforms can query the surface for the present mode, and size, etc
 	static bool SupportsQuerySurfaceProperties() { return true; }
