Override fma behaviour in OpenGL to use saperate multiply/add

diff --git a/Engine/Source/Runtime/OpenGLDrv/Private/OpenGLShaders.cpp b/Engine/Source/Runtime/OpenGLDrv/Private/OpenGLShaders.cpp
--- a/Engine/Source/Runtime/OpenGLDrv/Private/OpenGLShaders.cpp
+++ b/Engine/Source/Runtime/OpenGLDrv/Private/OpenGLShaders.cpp
@@ -757,6 +757,8 @@ void OPENGLDRV_API GLSLToDeviceCompatibleGLSL(FAnsiCharArray& GlslCodeOriginal,
 		ReplaceCString(GlslCodeOriginal, "layout(early_fragment_tests) in;", "");
 	}
 
+	AppendCString(GlslCode, "#define fma(A, B, C) ((A) * (B) + (C))\n");
+
 	// The incoming glsl may have preprocessor code that is dependent on defines introduced via the engine.
 	// This is the place to insert such engine preprocessor defines, immediately after the glsl version declaration.
 	if (TypeEnum == GL_FRAGMENT_SHADER)
