Unshelved from pending changelist '24460796':

Saves some GLES shader memory by stripping unused shader binding data at load time.

diff --git a/Engine/Source/Runtime/OpenGLDrv/Public/OpenGLShaderResources.h b/Engine/Source/Runtime/OpenGLDrv/Public/OpenGLShaderResources.h
--- a/Engine/Source/Runtime/OpenGLDrv/Public/OpenGLShaderResources.h
+++ b/Engine/Source/Runtime/OpenGLDrv/Public/OpenGLShaderResources.h
@@ -81,6 +81,8 @@ struct FOpenGLShaderBindings
 	uint8	NumUAVs;
 	bool	bFlattenUB;
 
+	FSHAHash VaryingHash; // Not serialized, built during load to allow us to diff varying info but avoid the memory overhead.
+
 	FOpenGLShaderBindings() :
 		NumSamplers(0),
 		NumUniformBuffers(0),
@@ -103,6 +105,7 @@ struct FOpenGLShaderBindings
 		bEqual &= A.InputVaryings.Num() == B.InputVaryings.Num();
 		bEqual &= A.OutputVaryings.Num() == B.OutputVaryings.Num();
 		bEqual &= A.ShaderResourceTable == B.ShaderResourceTable;
+		bEqual &= A.VaryingHash == B.VaryingHash;
 
 		if ( !bEqual )
 		{
@@ -160,6 +163,9 @@ struct FOpenGLShaderBindings
 		{
 			Hash ^= GetTypeHash(Binding.OutputVaryings[Item]);
 		}
+
+		Hash ^= GetTypeHash(Binding.VaryingHash);
+
 		return Hash;
 	}
 };
@@ -176,6 +182,29 @@ inline FArchive& operator<<(FArchive& Ar, FOpenGLShaderBindings& Bindings)
 	Ar << Bindings.NumUniformBuffers;
 	Ar << Bindings.NumUAVs;
 	Ar << Bindings.bFlattenUB;
+
+	if (Ar.IsLoading())
+	{
+		// hash then strip out the Input/OutputVaryings at load time.
+		// The hash ensures varying diffs still affect operator== and GetTypeHash()
+		FSHA1 HashState;
+		auto HashVarying = [&](FSHA1& HashStateIN, const TArray<FOpenGLShaderVarying>& InputVaryings)
+		{
+			for (const FOpenGLShaderVarying& Varying : InputVaryings)
+			{
+				HashStateIN.Update((const uint8*)&Varying.Location, sizeof(Varying.Location));
+				HashStateIN.Update((const uint8*)Varying.Varying.GetData(), Varying.Varying.Num() * sizeof(ANSICHAR));
+			}
+		};
+		HashVarying(HashState, Bindings.InputVaryings);
+		HashVarying(HashState, Bindings.OutputVaryings);
+		HashState.Final();
+		HashState.GetHash(&Bindings.VaryingHash.Hash[0]);
+
+		Bindings.InputVaryings.Empty();
+		Bindings.OutputVaryings.Empty();
+	}
+
 	return Ar;
 }
 
