Fix c++20 volatile compound assignment warning

#rnx

diff --git a/Engine/Source/Runtime/OpenGLDrv/Private/OpenGLTexture.cpp b/Engine/Source/Runtime/OpenGLDrv/Private/OpenGLTexture.cpp
--- a/Engine/Source/Runtime/OpenGLDrv/Private/OpenGLTexture.cpp
+++ b/Engine/Source/Runtime/OpenGLDrv/Private/OpenGLTexture.cpp
@@ -118,7 +118,7 @@ void FOpenGLTexture::UpdateTextureStats(FOpenGLTexture* Texture, bool bAllocatin
 		case ETextureDimension::TextureCubeArray: INC_MEMORY_STAT_BY(STAT_RenderTargetMemoryCube, TextureSize); break;
 		}
 
-		GCurrentRendertargetMemorySize += TextureSizeInKBs;
+		GCurrentRendertargetMemorySize = GCurrentRendertargetMemorySize + TextureSizeInKBs;
 		LLM_SCOPED_PAUSE_TRACKING_WITH_ENUM_AND_AMOUNT(ELLMTag::GraphicsPlatform, TextureSize, ELLMTracker::Platform, ELLMAllocType::None);
 		LLM_SCOPED_PAUSE_TRACKING_WITH_ENUM_AND_AMOUNT(ELLMTag::RenderTargets   , TextureSize, ELLMTracker::Default , ELLMAllocType::None);
 	}
@@ -134,7 +134,7 @@ void FOpenGLTexture::UpdateTextureStats(FOpenGLTexture* Texture, bool bAllocatin
 		case ETextureDimension::TextureCubeArray: INC_MEMORY_STAT_BY(STAT_TextureMemoryCube, TextureSize); break;
 		}
 
-		GCurrentTextureMemorySize += TextureSizeInKBs;
+		GCurrentTextureMemorySize = GCurrentTextureMemorySize + TextureSizeInKBs;
 		LLM_SCOPED_PAUSE_TRACKING_WITH_ENUM_AND_AMOUNT(ELLMTag::GraphicsPlatform, TextureSize, ELLMTracker::Platform, ELLMAllocType::None);
 		LLM_SCOPED_PAUSE_TRACKING_WITH_ENUM_AND_AMOUNT(ELLMTag::Textures        , TextureSize, ELLMTracker::Default , ELLMAllocType::None);
 	}
