Fixed: Missing implementation for 3D, 2D Arrays render targets in OpenGL ES, that causes a crash on some content

diff --git a/Engine/Source/Runtime/OpenGLDrv/Private/OpenGLES.cpp b/Engine/Source/Runtime/OpenGLDrv/Private/OpenGLES.cpp
--- a/Engine/Source/Runtime/OpenGLDrv/Private/OpenGLES.cpp
+++ b/Engine/Source/Runtime/OpenGLDrv/Private/OpenGLES.cpp
@@ -55,6 +55,7 @@ namespace GLFuncPointers
 	PFNGLBLENDFUNCIEXTPROC					glBlendFunciEXT = nullptr;
 	PFNGLBLENDFUNCSEPARATEIEXTPROC			glBlendFuncSeparateiEXT = nullptr;
 	PFNGLCOLORMASKIEXTPROC					glColorMaskiEXT = nullptr;
+	PFNGLFRAMEBUFFERTEXTUREPROC				glFramebufferTexture = nullptr;
 
 	PFNGLFRAMEBUFFERTEXTUREMULTIVIEWOVRPROC				glFramebufferTextureMultiviewOVR = NULL;
 	PFNGLFRAMEBUFFERTEXTUREMULTISAMPLEMULTIVIEWOVRPROC	glFramebufferTextureMultisampleMultiviewOVR = NULL;
@@ -305,6 +306,7 @@ void FOpenGLES::ProcessExtensions(const FString& ExtensionsString)
 		glBlendFunciEXT = (PFNGLBLENDFUNCIEXTPROC)((void*)eglGetProcAddress("glBlendFunci"));
 		glBlendFuncSeparateiEXT = (PFNGLBLENDFUNCSEPARATEIEXTPROC)((void*)eglGetProcAddress("glBlendFuncSeparatei"));
 		glColorMaskiEXT = (PFNGLCOLORMASKIEXTPROC)((void*)eglGetProcAddress("glColorMaski"));
+		glFramebufferTexture = (PFNGLFRAMEBUFFERTEXTUREPROC)((void*)eglGetProcAddress("glFramebufferTexture"));
 	}
 	
 	if (!glEnableiEXT && ExtensionsString.Contains(TEXT("GL_EXT_draw_buffers_indexed")))
diff --git a/Engine/Source/Runtime/OpenGLDrv/Public/OpenGLES.h b/Engine/Source/Runtime/OpenGLDrv/Public/OpenGLES.h
--- a/Engine/Source/Runtime/OpenGLDrv/Public/OpenGLES.h
+++ b/Engine/Source/Runtime/OpenGLDrv/Public/OpenGLES.h
@@ -122,6 +122,9 @@ typedef GLfloat GLdouble;
 #define GL_DEBUG_SEVERITY_NOTIFICATION GL_DEBUG_SEVERITY_NOTIFICATION_KHR
 #endif
 
+// FIXME: include gl32.h
+typedef void (GL_APIENTRYP PFNGLFRAMEBUFFERTEXTUREPROC) (GLenum target, GLenum attachment, GLuint texture, GLint level);
+
 namespace GLFuncPointers
 {
 	// GL_EXT_multisampled_render_to_texture
@@ -161,6 +164,7 @@ namespace GLFuncPointers
 	extern PFNGLBLENDFUNCIEXTPROC				glBlendFunciEXT;
 	extern PFNGLBLENDFUNCSEPARATEIEXTPROC		glBlendFuncSeparateiEXT;
 	extern PFNGLCOLORMASKIEXTPROC				glColorMaskiEXT;
+	extern PFNGLFRAMEBUFFERTEXTUREPROC			glFramebufferTexture;
 
 	// Mobile multi-view
 	extern PFNGLFRAMEBUFFERTEXTUREMULTIVIEWOVRPROC glFramebufferTextureMultiviewOVR;
@@ -639,12 +643,14 @@ struct FOpenGLES : public FOpenGLBase
 
 	static FORCEINLINE void FramebufferTexture(GLenum Target, GLenum Attachment, GLuint Texture, GLint Level)
 	{
-		check(0);
+		// ES 3.2
+		glFramebufferTexture(Target, Attachment, Texture, Level);
 	}
 
 	static FORCEINLINE void FramebufferTexture3D(GLenum Target, GLenum Attachment, GLenum TexTarget, GLuint Texture, GLint Level, GLint ZOffset)
 	{
-		check(0);
+		// glFramebufferTexture3D is not supported on GLES
+		glFramebufferTextureLayer(Target, Attachment, Texture, Level, ZOffset);
 	}
 
 	static FORCEINLINE void FramebufferTextureLayer(GLenum Target, GLenum Attachment, GLuint Texture, GLint Level, GLint Layer)
