Fix a bug of missing CommitComputeResourceTables for dispatching indirect compute shader on OpenGL.

#jire none

diff --git a/Engine/Source/Runtime/OpenGLDrv/Private/OpenGLCommands.cpp b/Engine/Source/Runtime/OpenGLDrv/Private/OpenGLCommands.cpp
--- a/Engine/Source/Runtime/OpenGLDrv/Private/OpenGLCommands.cpp
+++ b/Engine/Source/Runtime/OpenGLDrv/Private/OpenGLCommands.cpp
@@ -2916,6 +2916,8 @@ void FOpenGLDynamicRHI::RHIDispatchIndirectComputeShader(FRHIBuffer* ArgumentBuf
 
 	BindPendingComputeShaderState(ContextState, ComputeShader);
 
+	CommitComputeResourceTables(ComputeShader);
+
 	SetupTexturesForDraw(ContextState, ComputeShader, FOpenGL::GetMaxComputeTextureImageUnits());
 
 	SetupUAVsForCompute(ContextState, ComputeShader);
