fix OpenGLDrv build warning for windows platform extensions.

#rnx

diff --git a/Engine/Source/Runtime/OpenGLDrv/OpenGLDrv.Build.cs b/Engine/Source/Runtime/OpenGLDrv/OpenGLDrv.Build.cs
--- a/Engine/Source/Runtime/OpenGLDrv/OpenGLDrv.Build.cs
+++ b/Engine/Source/Runtime/OpenGLDrv/OpenGLDrv.Build.cs
@@ -48,7 +48,7 @@ public class OpenGLDrv : ModuleRules
 			PrivateIncludePathModuleNames.Add("Launch");
 		}
 
-        if (Target.Platform != UnrealTargetPlatform.Win64
+        if (!Target.Platform.IsInGroup(UnrealPlatformGroup.Windows)
 			&& Target.Platform != UnrealTargetPlatform.IOS && Target.Platform != UnrealTargetPlatform.Android
 			&& !Target.IsInPlatformGroup(UnrealPlatformGroup.Linux)
 			&& Target.Platform != UnrealTargetPlatform.TVOS)
