expanding some simple module declarations to include windows platform extensions by default.

#rnx

diff --git a/Engine/Source/Runtime/AVEncoder/AVEncoder.Build.cs b/Engine/Source/Runtime/AVEncoder/AVEncoder.Build.cs
--- a/Engine/Source/Runtime/AVEncoder/AVEncoder.Build.cs
+++ b/Engine/Source/Runtime/AVEncoder/AVEncoder.Build.cs
@@ -4,7 +4,8 @@ using UnrealBuildTool;
 using System.IO;
 using System;
 
-[SupportedPlatforms("Win64", "Linux")]
+[SupportedPlatforms("Linux")]
+[SupportedPlatformGroups("Windows")]
 public class AVEncoder : ModuleRules
 {
 	public AVEncoder(ReadOnlyTargetRules Target) : base(Target)
diff --git a/Engine/Source/Runtime/GameplayMediaEncoder/GameplayMediaEncoder.Build.cs b/Engine/Source/Runtime/GameplayMediaEncoder/GameplayMediaEncoder.Build.cs
--- a/Engine/Source/Runtime/GameplayMediaEncoder/GameplayMediaEncoder.Build.cs
+++ b/Engine/Source/Runtime/GameplayMediaEncoder/GameplayMediaEncoder.Build.cs
@@ -3,7 +3,7 @@
 using System.IO;
 using UnrealBuildTool;
 
-[SupportedPlatforms("Win64")]
+[SupportedPlatformGroups("Windows")]
 public class GameplayMediaEncoder : ModuleRules
 {
 	public GameplayMediaEncoder(ReadOnlyTargetRules Target) : base(Target)
@@ -27,7 +27,7 @@ public class GameplayMediaEncoder : ModuleRules
 			"AVEncoder"
 		});
 
-		if (Target.Platform == UnrealTargetPlatform.Win64)
+		if (Target.IsInPlatformGroup(UnrealPlatformGroup.Windows))
 		{
 			PrivateDependencyModuleNames.Add("D3D11RHI");
 			AddEngineThirdPartyPrivateStaticDependencies(Target, "DX11");
diff --git a/Engine/Source/Runtime/OpenGLDrv/OpenGLDrv.Build.cs b/Engine/Source/Runtime/OpenGLDrv/OpenGLDrv.Build.cs
--- a/Engine/Source/Runtime/OpenGLDrv/OpenGLDrv.Build.cs
+++ b/Engine/Source/Runtime/OpenGLDrv/OpenGLDrv.Build.cs
@@ -2,7 +2,8 @@
 
 using UnrealBuildTool;
 
-[SupportedPlatforms("Win64", "Linux", "Android", "LinuxArm64")]
+[SupportedPlatforms("Linux", "Android", "LinuxArm64")]
+[SupportedPlatformGroups("Windows")]
 public class OpenGLDrv : ModuleRules
 {
 	public OpenGLDrv(ReadOnlyTargetRules Target) : base(Target)
@@ -48,7 +49,7 @@ public class OpenGLDrv : ModuleRules
 			PrivateIncludePathModuleNames.Add("Launch");
 		}
 
-        if (!Target.Platform.IsInGroup(UnrealPlatformGroup.Windows)
+        if (!Target.IsInPlatformGroup(UnrealPlatformGroup.Windows)
 			&& Target.Platform != UnrealTargetPlatform.IOS && Target.Platform != UnrealTargetPlatform.Android
 			&& !Target.IsInPlatformGroup(UnrealPlatformGroup.Linux)
 			&& Target.Platform != UnrealTargetPlatform.TVOS)
diff --git a/Engine/Source/Runtime/Windows/D3D11RHI/D3D11RHI.Build.cs b/Engine/Source/Runtime/Windows/D3D11RHI/D3D11RHI.Build.cs
--- a/Engine/Source/Runtime/Windows/D3D11RHI/D3D11RHI.Build.cs
+++ b/Engine/Source/Runtime/Windows/D3D11RHI/D3D11RHI.Build.cs
@@ -3,7 +3,7 @@
 using System.IO;
 using UnrealBuildTool;
 
-[SupportedPlatforms("Win64")]
+[SupportedPlatformGroups("Windows")]
 public class D3D11RHI : ModuleRules
 {
 	protected virtual bool bIncludeExtensions { get => true; }
@@ -37,7 +37,7 @@ public class D3D11RHI : ModuleRules
 			AddEngineThirdPartyPrivateStaticDependencies(Target, "IntelExtensionsFramework");
 		}
 
-		if (Target.Platform == UnrealTargetPlatform.Win64)
+		if (Target.IsInPlatformGroup(UnrealPlatformGroup.Windows))
 		{
 			PrivateIncludePaths.Add(Path.Combine(this.ModuleDirectory, "Private", "Windows"));
 		}
diff --git a/Engine/Source/Runtime/Windows/WindowsPlatformFeatures/WindowsPlatformFeatures.Build.cs b/Engine/Source/Runtime/Windows/WindowsPlatformFeatures/WindowsPlatformFeatures.Build.cs
--- a/Engine/Source/Runtime/Windows/WindowsPlatformFeatures/WindowsPlatformFeatures.Build.cs
+++ b/Engine/Source/Runtime/Windows/WindowsPlatformFeatures/WindowsPlatformFeatures.Build.cs
@@ -3,7 +3,7 @@
 using System.IO;
 using UnrealBuildTool;
 
-[SupportedPlatforms("Win64")]
+[SupportedPlatformGroups("Windows")]
 public class WindowsPlatformFeatures : ModuleRules
 {
 	public WindowsPlatformFeatures(ReadOnlyTargetRules Target) : base(Target)
