Defer GL shader compilation until we need to link GL program

diff --git a/Engine/Source/Runtime/OpenGLDrv/Private/OpenGLShaders.cpp b/Engine/Source/Runtime/OpenGLDrv/Private/OpenGLShaders.cpp
--- a/Engine/Source/Runtime/OpenGLDrv/Private/OpenGLShaders.cpp
+++ b/Engine/Source/Runtime/OpenGLDrv/Private/OpenGLShaders.cpp
@@ -20,8 +20,6 @@
 #include "OpenGLBinaryProgramUtils.h"
 #include "HAL/PlatformFramePacer.h"
 
-#define CHECK_FOR_GL_SHADERS_TO_REPLACE 0
-
 #if PLATFORM_WINDOWS && PLATFORM_CPU_X86_FAMILY
 #include <mmintrin.h>
 #endif
@@ -268,10 +266,8 @@ class FOpenGLCompiledShaderValue
 	const FName CompressionMethod = NAME_Oodle;
 
 public:
-	FOpenGLCompiledShaderValue(GLuint ResourceIN, const TArray<ANSICHAR>& GlslCodeIN)
-		: Resource(ResourceIN)
+	FOpenGLCompiledShaderValue()
 	{
-		CompressShader(GlslCodeIN);
 	}
 	~FOpenGLCompiledShaderValue()
 	{
@@ -310,10 +306,8 @@ public:
 
 	static TAtomic<uint32> StatTotalStoredSize;
 	static TAtomic<uint32> StatTotalUncompressedSize;
-	bool bHasCompiled = false;
-private:
-	TArray<ANSICHAR> GlslCode;
-	int32 UncompressedSize = -1;
+
+	bool HasCode() const { return !GlslCode.IsEmpty(); };
 
 	void CompressShader(const TArray<ANSICHAR>& InGlslCode)
 	{
@@ -351,12 +345,15 @@ private:
  		//UE_LOG(LogRHI, Warning, TEXT("Shader sizes: %d %d"), StatTotalStoredSize.Load(EMemoryOrder::Relaxed), StatTotalUncompressedSize.Load(EMemoryOrder::Relaxed));
 	}
 
+private:
+	TArray<ANSICHAR> GlslCode;
+	int32 UncompressedSize = -1;
 };
 
 TAtomic<uint32> FOpenGLCompiledShaderValue::StatTotalStoredSize = 0;
 TAtomic<uint32> FOpenGLCompiledShaderValue::StatTotalUncompressedSize = 0;
 
-typedef TMap<FOpenGLCompiledShaderKey, TSharedPtr<FOpenGLCompiledShaderValue> > FOpenGLCompiledShaderCache;
+typedef TMap<FOpenGLCompiledShaderKey, FOpenGLCompiledShaderValue> FOpenGLCompiledShaderCache;
 
 static FCriticalSection GCompiledShaderCacheCS;
 
@@ -366,27 +363,6 @@ static FOpenGLCompiledShaderCache& GetOpenGLCompiledShaderCache()
 	return CompiledShaderCache;
 }
 
-struct FLibraryShaderCacheValue
-{
-	FOpenGLCodeHeader* Header;
-	uint32 ShaderCrc;
-	GLuint GLShader;
-	TArray<FUniformBufferStaticSlot> StaticSlots;
-
-#if DEBUG_GL_SHADERS
-	TArray<ANSICHAR> GlslCode;
-	const ANSICHAR*  GlslCodeString; // make it easier in VS to see shader code in debug mode; points to begin of GlslCode
-#endif
-};
-
-typedef TMap<FSHAHash, FLibraryShaderCacheValue> FOpenGLCompiledLibraryShaderCache;
-
-static FOpenGLCompiledLibraryShaderCache& GetOpenGLCompiledLibraryShaderCache()
-{
-	static FOpenGLCompiledLibraryShaderCache CompiledShaderCache;
-	return CompiledShaderCache;
-}
-
 // ============================================================================================================================
 
 
@@ -533,17 +509,11 @@ inline uint32 GetTypeHash(FAnsiCharArray const& CharArray)
 	return FCrc::MemCrc32(CharArray.GetData(), CharArray.Num() * sizeof(ANSICHAR));
 }
 
-// Helper to compile a shader 
+// Helper to verify a compiled shader 
 // returns true if shader was compiled without any errors or errors should be ignored
-static bool CompileCurrentShader(const GLuint Resource, const FAnsiCharArray& GlslCode)
+static bool VerifyShaderCompilation(GLuint Resource, const ANSICHAR* GlslCodeString)
 {
 	VERIFY_GL_SCOPE();
-	const ANSICHAR * GlslCodeString = GlslCode.GetData();
-	int32 GlslCodeLength = GlslCode.Num() - 1;
-
-	glShaderSource(Resource, 1, (const GLchar**)&GlslCodeString, &GlslCodeLength);
-	glCompileShader(Resource);
-
 	// Verify that an OpenGL shader has compiled successfully.
 	SCOPE_CYCLE_COUNTER(STAT_OpenGLShaderCompileVerifyTime);
 	{
@@ -572,12 +542,11 @@ static bool CompileCurrentShader(const GLuint Resource, const FAnsiCharArray& Gl
 					glGetShaderInfoLog(Resource, LogLength, NULL, CompileLog);
 				}
 
-#if DEBUG_GL_SHADERS
 				if (GlslCodeString)
 				{
-					UE_LOG(LogRHI,Error,TEXT("Shader:\n%s"),ANSI_TO_TCHAR(GlslCodeString));
+					UE_LOG(LogRHI,Error,TEXT("Shader:\n%s"), ANSI_TO_TCHAR(GlslCodeString));
 				}
-#endif
+
 				UE_LOG(LogRHI,Error,TEXT("Failed to compile shader. Compile log:\n%s"), ANSI_TO_TCHAR(CompileLog));
 				if (LogLength > 1)
 				{
@@ -615,79 +584,27 @@ static void GLSLToPlatform(const FOpenGLCodeHeader& Header, GLenum TypeEnum, FAn
 /**
  * Compiles an OpenGL shader using the given GLSL microcode.
  */
-void FOpenGLShader::Compile(const FOpenGLCodeHeader& Header, const FSHAHash& LibraryHash, GLenum TypeEnum)
+void FOpenGLShader::Compile(GLenum TypeEnum)
 {
-	SCOPE_CYCLE_COUNTER(STAT_OpenGLShaderCompileTime);
 	VERIFY_GL_SCOPE();
-
-	{
-		FLibraryShaderCacheValue *Val = GetOpenGLCompiledLibraryShaderCache().Find(LibraryHash);
-		if (Val)
-		{
-			Resource = Val->GLShader;
-			Bindings = Val->Header->Bindings;
-			UniformBuffersCopyInfo = Val->Header->UniformBuffersCopyInfo;
-			StaticSlots = Val->StaticSlots;
-#if DEBUG_GL_SHADERS
-			GlslCode = Val->GlslCode;
-			GlslCodeString = (ANSICHAR*)GlslCode.GetData();
-#endif
-			return;
-		}
-	}
-
-	// Find the existing compiled shader in the cache.
+	
 	FScopeLock Lock(&GCompiledShaderCacheCS);
-	TSharedPtr<FOpenGLCompiledShaderValue> FoundShader = GetOpenGLCompiledShaderCache().FindRef(ShaderCodeKey);
-	check(FoundShader);
-	Resource = FoundShader->Resource;
-
+	FOpenGLCompiledShaderValue& FoundShader = GetOpenGLCompiledShaderCache().FindOrAdd(ShaderCodeKey);
+	Resource = FoundShader.Resource;
 	if (Resource == 0)
 	{
+		SCOPE_CYCLE_COUNTER(STAT_OpenGLShaderCompileTime);
 		Resource = FOpenGL::CreateShader(TypeEnum);
-		const FOpenGLShaderDeviceCapabilities& Capabilities = GetOpenGLShaderDeviceCapabilities();
-
-		// Save the code and defer compilation if our device supports program binaries and we're not checking for shader compatibility.
-		const bool bDeferredCompilation = FOpenGLProgramBinaryCache::IsEnabled();
-		if (bDeferredCompilation)
-		{
-			// do nothing...
-		}
-		else
-		{
-			const FAnsiCharArray GlslPlatformCode = FoundShader->GetUncompressedShader();
-			const bool bSuccessfullyCompiled = CompileCurrentShader(Resource, GlslPlatformCode);
-			ensure(bSuccessfullyCompiled);
-		}
-
-		FoundShader->Resource = Resource;
-	}
 
-	Bindings = Header.Bindings;
-	UniformBuffersCopyInfo = Header.UniformBuffersCopyInfo;
-
-	UE::RHICore::InitStaticUniformBufferSlots(StaticSlots, Bindings.ShaderResourceTable);
-
-#if DEBUG_GL_SHADERS
-	{
-		const FAnsiCharArray glslPlatformCode = FoundShader->GetUncompressedShader();
-		GlslCode = glslPlatformCode;
-		GlslCodeString = (ANSICHAR*)GlslCode.GetData();
-	}
-#endif
-	if (LibraryHash != FSHAHash() && !GetOpenGLCompiledLibraryShaderCache().Contains(LibraryHash))
-	{
-		FLibraryShaderCacheValue Val;
-		Val.GLShader = Resource;
-		Val.Header = new FOpenGLCodeHeader;
-		*Val.Header = Header;
-		Val.ShaderCrc = ShaderCodeKey.GetCodeCRC();
-		Val.StaticSlots = StaticSlots;
-#if DEBUG_GL_SHADERS
-		Val.GlslCode = FoundShader->GetUncompressedShader();
-		Val.GlslCodeString = (ANSICHAR*)GlslCode.GetData();
-#endif
-		GetOpenGLCompiledLibraryShaderCache().Add(LibraryHash, Val);
+		TArray<ANSICHAR> UncompressedShaderCode = FoundShader.GetUncompressedShader();
+		int32 GlslCodeLength = UncompressedShaderCode.Num() - 1;
+		const ANSICHAR* UncompressedGlslCodeString = UncompressedShaderCode.GetData();
+		glShaderSource(Resource, 1, (const GLchar**)&UncompressedGlslCodeString, &GlslCodeLength);
+		glCompileShader(Resource);
+		const bool bSuccessfullyCompiled = VerifyShaderCompilation(Resource, UncompressedGlslCodeString);
+		ensure(bSuccessfullyCompiled);
+		
+		FoundShader.Resource = Resource;
 	}
 }
 
@@ -726,6 +643,8 @@ void OPENGLDRV_API GLSLToDeviceCompatibleGLSL(FAnsiCharArray& GlslCodeOriginal,
 		return; // platform extension overrides
 	}
 
+	GlslCode.Reserve(GlslCodeOriginal.Num());
+
 	// Whether we need to emit mobile multi-view code or not.
 	const bool bEmitMobileMultiView = (FCStringAnsi::Strstr(GlslCodeOriginal.GetData(), "gl_ViewID_OVR") != nullptr);
 
@@ -833,27 +752,6 @@ void OPENGLDRV_API GLSLToDeviceCompatibleGLSL(FAnsiCharArray& GlslCodeOriginal,
 	// Move version tag & extensions before beginning all other operations
 	MoveHashLines(GlslCode, GlslCodeOriginal);
 
-	// OpenGL SM5 shader platforms require location declarations for the layout, but don't necessarily use SSOs
-	if (Capabilities.TargetPlatform == EOpenGLShaderTargetPlatform::OGLSTP_Desktop)
-	{
-		AppendCString(GlslCode, "#define INTERFACE_BLOCK(Pos, Interp, Modifiers, Semantic, PreType, PostType) layout(location=Pos) Interp Modifiers struct { PreType PostType; }\n");
-	}
-	else
-	{
-		AppendCString(GlslCode, "#define INTERFACE_BLOCK(Pos, Interp, Modifiers, Semantic, PreType, PostType) layout(location=Pos) Modifiers Semantic { PreType PostType; }\n");
-	}
-
-	if (Capabilities.TargetPlatform == EOpenGLShaderTargetPlatform::OGLSTP_Desktop)
-	{
-		// If we're running <= featurelevel es3.1 shaders then enable this extension which adds support for uintBitsToFloat etc.
-		if ((FCStringAnsi::Strstr(GlslCode.GetData(), "#version 150") != nullptr))
-		{
-			AppendCString(GlslCode, "\n\n");
-			AppendCString(GlslCode, "#extension GL_ARB_gpu_shader5 : enable\n");
-			AppendCString(GlslCode, "\n\n");
-		}
-	}
-
 #if	DEBUG_GL_SHADERS
 	if (ShaderName.IsEmpty() == false)
 	{
@@ -915,7 +813,7 @@ static ANSICHAR* SetIndex(ANSICHAR* Str, int32 Offset, int32 Index)
 	return Str;
 }
 
-FOpenGLShader::FOpenGLShader(FRHICommandListImmediate& RHICmdList, TArrayView<const uint8> Code, const FSHAHash& Hash, GLenum TypeEnum)
+FOpenGLShader::FOpenGLShader(TArrayView<const uint8> Code, const FSHAHash& Hash, GLenum TypeEnum)
 {
 	FMemory::Memzero(&Bindings, sizeof(Bindings));
 
@@ -943,6 +841,10 @@ FOpenGLShader::FOpenGLShader(FRHICommandListImmediate& RHICmdList, TArrayView<co
 		return;
 	}
 
+	Bindings = Header.Bindings;
+	UniformBuffersCopyInfo = Header.UniformBuffersCopyInfo;
+	UE::RHICore::InitStaticUniformBufferSlots(StaticSlots, Bindings.ShaderResourceTable);
+
 	int32 CodeOffset = Ar.Tell();
 
 	// The code as given to us.
@@ -955,54 +857,26 @@ FOpenGLShader::FOpenGLShader(FRHICommandListImmediate& RHICmdList, TArrayView<co
 	uint32 CodeCRC = FCrc::MemCrc32(GlslCodeOriginal.GetData(), GlslCodeOriginal.Num());
 	ShaderCodeKey = FOpenGLCompiledShaderKey(TypeEnum, GlslCodeOriginal.Num(), CodeCRC);
 
-#if CHECK_FOR_GL_SHADERS_TO_REPLACE
-	{
-		// 1. Check for specific file
-		FString PotentialShaderFileName = FString::Printf(TEXT("%s-%d-0x%x.txt"), ShaderNameFromShaderType(TypeEnum), glslPlatformCode.Num(), CodeCRC);
-		FString PotentialShaderFile = FPaths::ProfilingDir();
-		PotentialShaderFile *= PotentialShaderFileName;
-
-		UE_LOG(LogRHI, Log, TEXT("Looking for shader file '%s' for potential replacement."), *PotentialShaderFileName);
-
-		int64 FileSize = IFileManager::Get().FileSize(*PotentialShaderFile);
-		if (FileSize > 0)
-		{
-			FArchive* Ar = IFileManager::Get().CreateFileReader(*PotentialShaderFile);
-			if (Ar != NULL)
-			{
-				UE_LOG(LogRHI, Log, TEXT("Replacing %s shader with length %d and CRC 0x%x with the one from a file."), (TypeEnum == GL_VERTEX_SHADER) ? TEXT("vertex") : ((TypeEnum == GL_FRAGMENT_SHADER) ? TEXT("fragment") : TEXT("geometry")), GlslCodeOriginal.Num(), GlslCodeOriginalCRC);
-
-				// read in the file
-				GlslCodeOriginal.Empty();
-				GlslCodeOriginal.AddUninitialized(FileSize + 1);
-				Ar->Serialize(GlslCodeOriginal.GetData(), FileSize);
-				delete Ar;
-				GlslCodeOriginal[FileSize] = 0;
-			}
-		}
-	}
-#endif
-
+	FAnsiCharArray GlslCodeFinal;
 	{
 		FScopeLock Lock(&GCompiledShaderCacheCS);
+		FOpenGLCompiledShaderValue& FoundShader = GetOpenGLCompiledShaderCache().FindOrAdd(ShaderCodeKey);
+		Resource = FoundShader.Resource;
 
-		TSharedPtr<FOpenGLCompiledShaderValue> FoundShader = GetOpenGLCompiledShaderCache().FindRef(ShaderCodeKey);
-		if (!FoundShader)
+		if (FoundShader.Resource == 0 && !FoundShader.HasCode())
 		{
-			FAnsiCharArray Finalglsl;
-			GLSLToPlatform(Header, TypeEnum, GlslCodeOriginal, Finalglsl);
-			GetOpenGLCompiledShaderCache().Add(ShaderCodeKey, MakeShareable(new FOpenGLCompiledShaderValue(0, Finalglsl)));
+			GLSLToPlatform(Header, TypeEnum, GlslCodeOriginal, GlslCodeFinal);
+			FoundShader.CompressShader(GlslCodeFinal);
 		}
 		// With debug shaders we insert a shader name into the source and that can make it unique failing CRC check
 #if (UE_BUILD_DEBUG || UE_BUILD_DEVELOPMENT) && !DEBUG_GL_SHADERS 
 		else
 		{
 			QUICK_SCOPE_CYCLE_COUNTER(STAT_GLCheckShaderCodeCRC);
-			FAnsiCharArray Finalglsl;
-			GLSLToPlatform(Header, TypeEnum, GlslCodeOriginal, Finalglsl);
-			TArray<ANSICHAR> FoundShaderCode = FoundShader->GetUncompressedShader();
-			if (FoundShaderCode.Num() != Finalglsl.Num()
-				|| FMemory::Memcmp(FoundShaderCode.GetData(), Finalglsl.GetData(), FoundShaderCode.Num())
+			GLSLToPlatform(Header, TypeEnum, GlslCodeOriginal, GlslCodeFinal);
+			TArray<ANSICHAR> FoundShaderCode = FoundShader.GetUncompressedShader();
+			if (FoundShaderCode.Num() != GlslCodeFinal.Num()
+				|| FMemory::Memcmp(FoundShaderCode.GetData(), GlslCodeFinal.GetData(), FoundShaderCode.Num())
 				)
 			{
 				UE_LOG(LogRHI, Fatal, TEXT("SHADER CRC CLASH!"));
@@ -1011,28 +885,32 @@ FOpenGLShader::FOpenGLShader(FRHICommandListImmediate& RHICmdList, TArrayView<co
 #endif
 	}
 
-	RHICmdList.EnqueueLambda([this, Header = MoveTemp(Header), Hash, TypeEnum](FRHICommandListImmediate&)
+#if DEBUG_GL_SHADERS
+	if (GlslCodeFinal.Num() == 0)
 	{
-		Compile(Header, Hash, TypeEnum);
-	});
+		GLSLToPlatform(Header, TypeEnum, GlslCodeOriginal, GlslCodeFinal);
+		check(GlslCodeFinal.Num());
+	}
+	GlslCode = MoveTemp(GlslCodeFinal);
+	GlslCodeString = GlslCodeFinal.GetData();
+#endif
+
+	// The shader is compiled when we link program
 }
 
 FVertexShaderRHIRef FOpenGLDynamicRHI::RHICreateVertexShader(TArrayView<const uint8> Code, const FSHAHash& Hash)
 {
-	check(IsInRenderingThread());
-	return new FOpenGLVertexShader(FRHICommandListExecutor::GetImmediateCommandList(), Code, Hash);
+	return new FOpenGLVertexShader(Code, Hash);
 }
 
 FPixelShaderRHIRef FOpenGLDynamicRHI::RHICreatePixelShader(TArrayView<const uint8> Code, const FSHAHash& Hash)
 {
-	check(IsInRenderingThread());
-	return new FOpenGLPixelShader(FRHICommandListExecutor::GetImmediateCommandList(), Code, Hash);
+	return new FOpenGLPixelShader(Code, Hash);
 }
 
 FGeometryShaderRHIRef FOpenGLDynamicRHI::RHICreateGeometryShader(TArrayView<const uint8> Code, const FSHAHash& Hash)
 {
-	check(IsInRenderingThread());
-	return new FOpenGLGeometryShader(FRHICommandListExecutor::GetImmediateCommandList(), Code, Hash);
+	return new FOpenGLGeometryShader(Code, Hash);
 }
 
 static void MarkShaderParameterCachesDirty(FOpenGLShaderParameterCache* ShaderParameters, bool UpdateCompute)
@@ -1319,7 +1197,7 @@ static void ConfigureStageStates(FOpenGLLinkedProgram* LinkedProgram)
 {
 	const FOpenGLLinkedProgramConfiguration &Config = LinkedProgram->Config;
 
-	if (Config.Shaders[CrossCompiler::SHADER_STAGE_VERTEX].Resource)
+	if (Config.Shaders[CrossCompiler::SHADER_STAGE_VERTEX].bValid)
 	{
 		LinkedProgram->ConfigureShaderStage(
 			CrossCompiler::SHADER_STAGE_VERTEX,
@@ -1328,7 +1206,7 @@ static void ConfigureStageStates(FOpenGLLinkedProgram* LinkedProgram)
 		check(LinkedProgram->StagePackedUniformInfo[CrossCompiler::SHADER_STAGE_VERTEX].PackedUniformInfos.Num() <= Config.Shaders[CrossCompiler::SHADER_STAGE_VERTEX].Bindings.PackedGlobalArrays.Num());
 	}
 
-	if (Config.Shaders[CrossCompiler::SHADER_STAGE_PIXEL].Resource)
+	if (Config.Shaders[CrossCompiler::SHADER_STAGE_PIXEL].bValid)
 	{
 		LinkedProgram->ConfigureShaderStage(
 			CrossCompiler::SHADER_STAGE_PIXEL,
@@ -1338,7 +1216,7 @@ static void ConfigureStageStates(FOpenGLLinkedProgram* LinkedProgram)
 		check(LinkedProgram->StagePackedUniformInfo[CrossCompiler::SHADER_STAGE_PIXEL].PackedUniformInfos.Num() <= Config.Shaders[CrossCompiler::SHADER_STAGE_PIXEL].Bindings.PackedGlobalArrays.Num());
 	}
 
-	if (Config.Shaders[CrossCompiler::SHADER_STAGE_GEOMETRY].Resource)
+	if (Config.Shaders[CrossCompiler::SHADER_STAGE_GEOMETRY].bValid)
 	{
 		LinkedProgram->ConfigureShaderStage(
 			CrossCompiler::SHADER_STAGE_GEOMETRY,
@@ -1349,7 +1227,7 @@ static void ConfigureStageStates(FOpenGLLinkedProgram* LinkedProgram)
 		check(LinkedProgram->StagePackedUniformInfo[CrossCompiler::SHADER_STAGE_GEOMETRY].PackedUniformInfos.Num() <= Config.Shaders[CrossCompiler::SHADER_STAGE_GEOMETRY].Bindings.PackedGlobalArrays.Num());
 	}
 
-	if (Config.Shaders[CrossCompiler::SHADER_STAGE_COMPUTE].Resource)
+	if (Config.Shaders[CrossCompiler::SHADER_STAGE_COMPUTE].bValid)
 	{
 		LinkedProgram->ConfigureShaderStage(
 			CrossCompiler::SHADER_STAGE_COMPUTE,
@@ -1421,25 +1299,6 @@ namespace UE
 
 			return bSuccess;
 		}
-
-        static void CompilePendingShaders(const FOpenGLLinkedProgramConfiguration& Config)
-        {
-	        VERIFY_GL_SCOPE();
-        
-	        // Find the existing compiled shader in the cache.
-	        FScopeLock Lock(&GCompiledShaderCacheCS);
-	        for (int32 StageIdx = 0; StageIdx < UE_ARRAY_COUNT(Config.Shaders); ++StageIdx)
-	        {
-		        TSharedPtr<FOpenGLCompiledShaderValue> FoundShader = GetOpenGLCompiledShaderCache().FindRef(Config.Shaders[StageIdx].ShaderKey);
-		        if (FoundShader && FoundShader->bHasCompiled == false)
-		        {
-			        TArray<ANSICHAR> GlslCode = FoundShader->GetUncompressedShader();
-			        CompileCurrentShader(Config.Shaders[StageIdx].Resource, GlslCode);
-			        FoundShader->bHasCompiled = true;
-		        }
-	        }
-        }
-
 	}
 }
 
@@ -2458,8 +2317,6 @@ FOpenGLLinkedProgram* FOpenGLDynamicRHI::LinkProgram( const FOpenGLLinkedProgram
 	MarkShaderParameterCachesDirty(PendingState.ShaderParameters, false);
 	PendingState.LinkedProgramAndDirtyFlag = nullptr;
 
-	ANSICHAR Buf[32] = {0};
-
 	SCOPE_CYCLE_COUNTER(STAT_OpenGLShaderLinkTime);
 	VERIFY_GL_SCOPE();
 
@@ -2467,11 +2324,6 @@ FOpenGLLinkedProgram* FOpenGLDynamicRHI::LinkProgram( const FOpenGLLinkedProgram
 	check( (Config.Shaders[CrossCompiler::SHADER_STAGE_VERTEX].Resource == 0) != (Config.Shaders[CrossCompiler::SHADER_STAGE_COMPUTE].Resource == 0));
 	check( (Config.Shaders[CrossCompiler::SHADER_STAGE_PIXEL].Resource == 0) != (Config.Shaders[CrossCompiler::SHADER_STAGE_COMPUTE].Resource == 0));
 
-	if (FOpenGLProgramBinaryCache::IsEnabled())
-	{
-		UE::OpenGL::CompilePendingShaders(Config);
-	}
-
 	GLuint Program = 0;
 	FOpenGL::GenProgramPipelines(1, &Program);
 
@@ -2522,19 +2374,23 @@ FOpenGLLinkedProgram* FOpenGLDynamicRHI::LinkProgram( const FOpenGLLinkedProgram
 bool FOpenGLDynamicRHI::LinkComputeShader(FRHIComputeShader* ComputeShaderRHI, FOpenGLComputeShader* ComputeShader)
 {
 	check(ComputeShader);
-	check(ComputeShader->Resource != 0);
 	check(ComputeShaderRHI->GetHash() != FSHAHash());
 
 	FOpenGLLinkedProgramConfiguration Config;
-	Config.Shaders[CrossCompiler::SHADER_STAGE_COMPUTE].Resource = ComputeShader->Resource;
-	Config.Shaders[CrossCompiler::SHADER_STAGE_COMPUTE].Bindings = ComputeShader->Bindings;
 	Config.ProgramKey.ShaderHashes[CrossCompiler::SHADER_STAGE_COMPUTE] = ComputeShaderRHI->GetHash();
-	Config.Shaders[CrossCompiler::SHADER_STAGE_COMPUTE].ShaderKey = FOpenGLDynamicRHI::ResourceCast(ComputeShaderRHI)->ShaderCodeKey;
 
 	ComputeShader->LinkedProgram = GetOpenGLProgramsCache().Find(Config.ProgramKey, true);
 
 	if (ComputeShader->LinkedProgram == nullptr)
 	{
+		ComputeShader->ConditionalyCompile();
+		check(ComputeShader->Resource != 0);
+
+		Config.Shaders[CrossCompiler::SHADER_STAGE_COMPUTE].Resource = ComputeShader->Resource;
+		Config.Shaders[CrossCompiler::SHADER_STAGE_COMPUTE].Bindings = ComputeShader->Bindings;
+		Config.Shaders[CrossCompiler::SHADER_STAGE_COMPUTE].ShaderKey = FOpenGLDynamicRHI::ResourceCast(ComputeShaderRHI)->ShaderCodeKey;
+		Config.Shaders[CrossCompiler::SHADER_STAGE_COMPUTE].bValid = true;
+		
 		ComputeShader->LinkedProgram = LinkProgram(Config);
 		if(ComputeShader->LinkedProgram == nullptr)
 		{
@@ -2561,9 +2417,6 @@ FOpenGLLinkedProgram* FOpenGLDynamicRHI::GetLinkedComputeProgram(FRHIComputeShad
 	FOpenGLComputeShader* ComputeShader = ResourceCast(ComputeShaderRHI);
 
 	FOpenGLLinkedProgramConfiguration Config;
-
-	Config.Shaders[CrossCompiler::SHADER_STAGE_COMPUTE].Resource = ComputeShader->Resource;
-	Config.Shaders[CrossCompiler::SHADER_STAGE_COMPUTE].Bindings = ComputeShader->Bindings;
 	Config.ProgramKey.ShaderHashes[CrossCompiler::SHADER_STAGE_COMPUTE] = ComputeShaderRHI->GetHash();
 
 	FOpenGLLinkedProgram* LinkedProgram = GetOpenGLProgramsCache().Find(Config.ProgramKey, true);
@@ -2587,6 +2440,9 @@ FOpenGLLinkedProgram* FOpenGLDynamicRHI::GetLinkedComputeProgram(FRHIComputeShad
 	else if (!LinkedProgram->bConfigIsInitalized)
 	{
 		// this has been loaded via binary program cache, properly initialize it here:
+		Config.Shaders[CrossCompiler::SHADER_STAGE_COMPUTE].Resource = ComputeShader->Resource; // potentially could be zero, but an already linked program should not need a compiled shader
+		Config.Shaders[CrossCompiler::SHADER_STAGE_COMPUTE].Bindings = ComputeShader->Bindings;
+		Config.Shaders[CrossCompiler::SHADER_STAGE_COMPUTE].bValid = true;
 		LinkedProgram->SetConfig(Config);
 		// We now have the config for this program, we must configure the program for use.
 		ConfigureGLProgramStageStates(LinkedProgram);
@@ -2597,8 +2453,7 @@ FOpenGLLinkedProgram* FOpenGLDynamicRHI::GetLinkedComputeProgram(FRHIComputeShad
 
 FComputeShaderRHIRef FOpenGLDynamicRHI::RHICreateComputeShader(TArrayView<const uint8> Code, const FSHAHash& Hash)
 {
-	check(IsInRenderingThread());
-	return new FOpenGLComputeShader(FRHICommandListExecutor::GetImmediateCommandList(), Code, Hash);
+	return new FOpenGLComputeShader(Code, Hash);
 }
 
 template<class TOpenGLStage>
@@ -2673,18 +2528,24 @@ static FOpenGLLinkedProgramConfiguration CreateConfig(FRHIVertexShader* VertexSh
 	check(VertexShaderRHI);
 	check(PixelShaderRHI);
 
+	VertexShader->ConditionalyCompile();
+	PixelShader->ConditionalyCompile();
+
 	// Fill-in the configuration
 	Config.Shaders[CrossCompiler::SHADER_STAGE_VERTEX].Bindings = VertexShader->Bindings;
 	Config.Shaders[CrossCompiler::SHADER_STAGE_VERTEX].Resource = VertexShader->Resource;
 	Config.Shaders[CrossCompiler::SHADER_STAGE_VERTEX].ShaderKey = FOpenGLDynamicRHI::ResourceCast(VertexShaderRHI)->ShaderCodeKey;
+	Config.Shaders[CrossCompiler::SHADER_STAGE_VERTEX].bValid = true;
 	Config.ProgramKey.ShaderHashes[CrossCompiler::SHADER_STAGE_VERTEX] = VertexShaderRHI->GetHash();
 
 	if (GeometryShaderRHI)
 	{
 		check(VertexShader);
+		GeometryShader->ConditionalyCompile();
 		BindShaderStage(Config, CrossCompiler::SHADER_STAGE_GEOMETRY, GeometryShaderRHI, CrossCompiler::SHADER_STAGE_VERTEX, VertexShaderRHI);
 		Config.ProgramKey.ShaderHashes[CrossCompiler::SHADER_STAGE_GEOMETRY] = GeometryShaderRHI->GetHash();
 		Config.Shaders[CrossCompiler::SHADER_STAGE_GEOMETRY].ShaderKey = FOpenGLDynamicRHI::ResourceCast(GeometryShaderRHI)->ShaderCodeKey;
+		Config.Shaders[CrossCompiler::SHADER_STAGE_GEOMETRY].bValid = true;
 	}
 
 	check(GeometryShaderRHI || VertexShaderRHI);
@@ -2698,6 +2559,7 @@ static FOpenGLLinkedProgramConfiguration CreateConfig(FRHIVertexShader* VertexSh
 	}
 	Config.ProgramKey.ShaderHashes[CrossCompiler::SHADER_STAGE_PIXEL] = PixelShaderRHI->GetHash();
 	Config.Shaders[CrossCompiler::SHADER_STAGE_PIXEL].ShaderKey = FOpenGLDynamicRHI::ResourceCast(PixelShaderRHI)->ShaderCodeKey;
+	Config.Shaders[CrossCompiler::SHADER_STAGE_PIXEL].bValid = true;
 
 
 	return Config;
@@ -3013,15 +2875,7 @@ void DestroyShadersAndPrograms()
 		FOpenGLCompiledShaderCache& ShaderCache = GetOpenGLCompiledShaderCache();
 		for (FOpenGLCompiledShaderCache::TIterator It(ShaderCache); It; ++It)
 		{
-			FOpenGL::DeleteShader(It.Value()->Resource);
-		}
-		ShaderCache.Empty();
-	}
-	{
-		FOpenGLCompiledLibraryShaderCache& ShaderCache = GetOpenGLCompiledLibraryShaderCache();
-		for (FOpenGLCompiledLibraryShaderCache::TIterator It(ShaderCache); It; ++It)
-		{
-			delete It.Value().Header;
+			FOpenGL::DeleteShader(It.Value().Resource);
 		}
 		ShaderCache.Empty();
 	}
diff --git a/Engine/Source/Runtime/OpenGLDrv/Public/OpenGLShaderResources.h b/Engine/Source/Runtime/OpenGLDrv/Public/OpenGLShaderResources.h
--- a/Engine/Source/Runtime/OpenGLDrv/Public/OpenGLShaderResources.h
+++ b/Engine/Source/Runtime/OpenGLDrv/Public/OpenGLShaderResources.h
@@ -308,7 +308,7 @@ public:
 	const ANSICHAR*  GlslCodeString; // make it easier in VS to see shader code in debug mode; points to begin of GlslCode
 #endif
 
-	FOpenGLShader(FRHICommandListImmediate& RHICmdList, TArrayView<const uint8> Code, const FSHAHash& Hash, GLenum TypeEnum);
+	FOpenGLShader(TArrayView<const uint8> Code, const FSHAHash& Hash, GLenum TypeEnum);
 
 	~FOpenGLShader()
 	{
@@ -318,8 +318,8 @@ public:
 //		}
 	}
 
-private:
-	void Compile(const FOpenGLCodeHeader& Header, const FSHAHash& LibraryHash, GLenum TypeEnum);
+protected:
+	void Compile(GLenum TypeEnum);
 };
 
 class FOpenGLVertexShader : public FRHIVertexShader, public FOpenGLShader
@@ -327,9 +327,17 @@ class FOpenGLVertexShader : public FRHIVertexShader, public FOpenGLShader
 public:
 	static constexpr EShaderFrequency Frequency = SF_Vertex;
 
-	FOpenGLVertexShader(FRHICommandListImmediate& RHICmdList, TArrayView<const uint8> Code, const FSHAHash& Hash)
-		: FOpenGLShader(RHICmdList, Code, Hash, GL_VERTEX_SHADER)
+	FOpenGLVertexShader(TArrayView<const uint8> Code, const FSHAHash& Hash)
+		: FOpenGLShader(Code, Hash, GL_VERTEX_SHADER)
 	{}
+
+	void ConditionalyCompile()
+	{
+		if (Resource == 0)
+		{
+			Compile(GL_VERTEX_SHADER);
+		}
+	}
 };
 
 class FOpenGLPixelShader : public FRHIPixelShader, public FOpenGLShader
@@ -337,9 +345,17 @@ class FOpenGLPixelShader : public FRHIPixelShader, public FOpenGLShader
 public:
 	static constexpr EShaderFrequency Frequency = SF_Pixel;
 
-	FOpenGLPixelShader(FRHICommandListImmediate& RHICmdList, TArrayView<const uint8> Code, const FSHAHash& Hash)
-		: FOpenGLShader(RHICmdList, Code, Hash, GL_FRAGMENT_SHADER)
+	FOpenGLPixelShader(TArrayView<const uint8> Code, const FSHAHash& Hash)
+		: FOpenGLShader(Code, Hash, GL_FRAGMENT_SHADER)
 	{}
+
+	void ConditionalyCompile()
+	{
+		if (Resource == 0)
+		{
+			Compile(GL_FRAGMENT_SHADER);
+		}
+	}
 };
 
 class FOpenGLGeometryShader : public FRHIGeometryShader, public FOpenGLShader
@@ -347,9 +363,17 @@ class FOpenGLGeometryShader : public FRHIGeometryShader, public FOpenGLShader
 public:
 	static constexpr EShaderFrequency Frequency = SF_Geometry;
 
-	FOpenGLGeometryShader(FRHICommandListImmediate& RHICmdList, TArrayView<const uint8> Code, const FSHAHash& Hash)
-		: FOpenGLShader(RHICmdList, Code, Hash, GL_GEOMETRY_SHADER)
+	FOpenGLGeometryShader(TArrayView<const uint8> Code, const FSHAHash& Hash)
+		: FOpenGLShader(Code, Hash, GL_GEOMETRY_SHADER)
 	{}
+
+	void ConditionalyCompile()
+	{
+		if (Resource == 0)
+		{
+			Compile(GL_GEOMETRY_SHADER);
+		}
+	}
 };
 
 class FOpenGLComputeShader : public FRHIComputeShader, public FOpenGLShader
@@ -357,10 +381,18 @@ class FOpenGLComputeShader : public FRHIComputeShader, public FOpenGLShader
 public:
 	static constexpr EShaderFrequency Frequency = SF_Compute;
 
-	FOpenGLComputeShader(FRHICommandListImmediate& RHICmdList, TArrayView<const uint8> Code, const FSHAHash& Hash)
-		: FOpenGLShader(RHICmdList, Code, Hash, GL_COMPUTE_SHADER)
+	FOpenGLComputeShader(TArrayView<const uint8> Code, const FSHAHash& Hash)
+		: FOpenGLShader(Code, Hash, GL_COMPUTE_SHADER)
 	{}
 
+	void ConditionalyCompile()
+	{
+		if (Resource == 0)
+		{
+			Compile(GL_COMPUTE_SHADER);
+		}
+	}
+
 	bool NeedsTextureStage(int32 TextureStageIndex);
 	int32 MaxTextureStageUsed();
 	const TBitArray<>& GetTextureNeeds(int32& OutMaxTextureStageUsed);
@@ -497,8 +529,8 @@ public:
 	{
 		FOpenGLShaderBindings Bindings;
 		GLuint Resource;
-
 		FOpenGLCompiledShaderKey ShaderKey; // This is the key to the shader within FOpenGLCompiledShader container
+		bool bValid; // To mark that stage is valid for this program, even when shader Resource could be zero
 	}
 	Shaders[CrossCompiler::NUM_SHADER_STAGES];
 	FOpenGLProgramKey ProgramKey;
@@ -508,6 +540,7 @@ public:
 		for (int32 Stage = 0; Stage < CrossCompiler::NUM_SHADER_STAGES; Stage++)
 		{
 			Shaders[Stage].Resource = 0;
+			Shaders[Stage].bValid = false;
 		}
 	}
 
@@ -517,6 +550,7 @@ public:
 		for (int32 Stage = 0; Stage < CrossCompiler::NUM_SHADER_STAGES && bEqual; Stage++)
 		{
 			bEqual &= A.Shaders[Stage].Resource == B.Shaders[Stage].Resource;
+			bEqual &= A.Shaders[Stage].bValid == B.Shaders[Stage].bValid;
 			bEqual &= A.Shaders[Stage].Bindings == B.Shaders[Stage].Bindings;
 		}
 		return bEqual;
