Fixed: Android compile errors in OpenGLShaders.cpp

diff --git a/Engine/Source/Runtime/OpenGLDrv/Private/OpenGLShaders.cpp b/Engine/Source/Runtime/OpenGLDrv/Private/OpenGLShaders.cpp
--- a/Engine/Source/Runtime/OpenGLDrv/Private/OpenGLShaders.cpp
+++ b/Engine/Source/Runtime/OpenGLDrv/Private/OpenGLShaders.cpp
@@ -2590,8 +2590,8 @@ static FOpenGLProgramBinary ExternalProgramCompile(const FOpenGLProgramKey& Prog
 	TArray<ANSICHAR> ComputeGlslCode;
 	{
 		FScopeLock Lock(&GCompiledShaderCacheCS);
-		VSCode = GetOpenGLCompiledShaderCache().FindRef(VSKey)->GetUncompressedShader();
-		PSCode = GetOpenGLCompiledShaderCache().FindRef(PSKey)->GetUncompressedShader();
+		VSCode = GetOpenGLCompiledShaderCache().FindRef(VSKey).GetUncompressedShader();
+		PSCode = GetOpenGLCompiledShaderCache().FindRef(PSKey).GetUncompressedShader();
 	}
 
 	FString FailLog;
