Add GLES PSO precache support.
 - When PSO precaching is enabled:
   - the binary program cache is continually appended with new PSOs.
   - all encountered PSOs are added to the binary cache.
   - The binary program cache is mmapped after the cache has grown beyond a cvar specified size.
   - rudimentary cache size controls (deletes after going over cvar controlled size threshold on startup).
   - any non-precache caches are deleted on startup.

diff --git a/Engine/Source/Runtime/Core/Private/Android/AndroidPlatformFile.cpp b/Engine/Source/Runtime/Core/Private/Android/AndroidPlatformFile.cpp
--- a/Engine/Source/Runtime/Core/Private/Android/AndroidPlatformFile.cpp
+++ b/Engine/Source/Runtime/Core/Private/Android/AndroidPlatformFile.cpp
@@ -372,6 +372,12 @@ public:
 				bSuccess = false;
 				break;
 			}
+#if LOG_ANDROID_FILE
+			FPlatformMisc::LowLevelOutputDebugStringf(
+				TEXT("(%d/%d) FFileHandleAndroid:Write => Path = %s, this size = %d, CurrentOffset = %d, Source = %p"),
+				FAndroidTLS::GetCurrentThreadId(), File->Handle,
+				*(File->Path), int32(ThisSize), CurrentOffset, Source);
+#endif
 			CurrentOffset += ThisSize;
 			Source += ThisSize;
 			BytesToWrite -= ThisSize;
@@ -379,7 +385,12 @@ public:
 		
 		// Update the cached file length
 		Length = FMath::Max(Length, CurrentOffset);
-
+#if LOG_ANDROID_FILE
+		FPlatformMisc::LowLevelOutputDebugStringf(
+			TEXT("(%d/%d) FFileHandleAndroid:Write => Path = %s, final size %d"),
+			FAndroidTLS::GetCurrentThreadId(), File->Handle,
+			*(File->Path), Length);
+#endif
 		return bSuccess;
 	}
 
@@ -1009,8 +1020,9 @@ public:
 	virtual IMappedFileRegion* MapRegion(int64 Offset = 0, int64 BytesToMap = MAX_int64, bool bPreloadHint = false) override
 	{
 		LLM_PLATFORM_SCOPE(ELLMTag::PlatformMMIO);
-		check(Offset < GetFileSize()); // don't map zero bytes and don't map off the end of the file
-		BytesToMap = FMath::Min<int64>(BytesToMap, GetFileSize() - Offset);
+		const int64 CurrentFileSize = GetCurrentFileSize();
+		check(Offset < CurrentFileSize); // don't map zero bytes and don't map off the end of the file
+		BytesToMap = FMath::Min<int64>(BytesToMap, CurrentFileSize - Offset);
 		check(BytesToMap > 0); // don't map zero bytes
 
 		const int64 AlignedOffset = AlignDown(Offset, FileMappingAlignment);
@@ -1052,7 +1064,8 @@ public:
 #if LOG_ANDROID_FILE
 		FPlatformMisc::LowLevelOutputDebugStringf(TEXT("Failed to unmap region from %s, errno=%s"), *Filename, UTF8_TO_TCHAR(strerror(errno)));
 #endif
-		checkf(Res == 0, TEXT("Failed to unmap, error is %d, errno is %d [params: %x, %d]"), Res, errno, MappedPtr, GetFileSize());
+		const int64 CurrentFileSize = GetCurrentFileSize();
+		checkf(Res == 0, TEXT("Failed to unmap, error is %d, errno is %d [params: %x, %d]"), Res, errno, MappedPtr, CurrentFileSize);
 	}
 
 private:
@@ -1060,6 +1073,22 @@ private:
 	FString Filename;
 	int32 NumOutstandingRegions;
 	int FileHandle;
+
+	int64 GetCurrentFileSize() const
+	{
+		struct stat FileInfo;
+		FileInfo.st_size = -1;
+		const int StatResult = fstat(FileHandle, &FileInfo);
+		if (StatResult == -1)
+		{
+			const int ErrNo = errno;
+#if LOG_ANDROID_FILE
+			FPlatformMisc::LowLevelOutputDebugStringf(TEXT("FAndroidPlatformFile::FAndroidMappedFileHandle fstat failed: ('%s') failed: errno=%d (%s)"), *Filename, ErrNo, UTF8_TO_TCHAR(strerror(ErrNo)));
+#endif
+			return GetFileSize();
+		}
+		return FileInfo.st_size;
+	}
 };
 
 FAndroidMappedFileRegion::~FAndroidMappedFileRegion()
diff --git a/Engine/Source/Runtime/Core/Public/Async/MappedFileHandle.h b/Engine/Source/Runtime/Core/Public/Async/MappedFileHandle.h
--- a/Engine/Source/Runtime/Core/Public/Async/MappedFileHandle.h
+++ b/Engine/Source/Runtime/Core/Public/Async/MappedFileHandle.h
@@ -100,7 +100,7 @@ public:
 	/**
 	* Return the size of the mapped file.
 	**/
-	FORCEINLINE int64 GetFileSize()
+	FORCEINLINE int64 GetFileSize() const
 	{
 		return MappedFileSize;
 	}
diff --git a/Engine/Source/Runtime/OpenGLDrv/Private/OpenGLProgramBinaryFileCache.cpp b/Engine/Source/Runtime/OpenGLDrv/Private/OpenGLProgramBinaryFileCache.cpp
--- a/Engine/Source/Runtime/OpenGLDrv/Private/OpenGLProgramBinaryFileCache.cpp
+++ b/Engine/Source/Runtime/OpenGLDrv/Private/OpenGLProgramBinaryFileCache.cpp
@@ -63,6 +63,33 @@ static FAutoConsoleVariableRef CVarMaxBinaryProgramLoadTime(
 	ECVF_RenderThreadSafe
 );
 
+static int32 GBinaryCachePeriodicFlushProgramCount = 20;
+static FAutoConsoleVariableRef CVarBinaryCachePeriodicFlushProgramCount(
+	TEXT("r.OpenGL.BinaryCachePeriodicFlushProgramCount"),
+	GBinaryCachePeriodicFlushProgramCount,
+	TEXT("When r.PSOPrecaching is active this value\n")
+	TEXT("is the number of appended programs to accumulate before the cache is flushed to storage."),
+	ECVF_RenderThreadSafe
+);
+
+static int32 GBinaryCacheMMapAfterEveryMB = 50;
+static FAutoConsoleVariableRef CVarBinaryCacheMMapAfterEveryMB(
+	TEXT("r.OpenGL.BinaryCacheMMapAfterEveryMB"),
+	GBinaryCacheMMapAfterEveryMB,
+	TEXT("When r.PSOPrecaching is active this value\n")
+	TEXT("specifies the size program binary cache can grow before it is memory mapped, the mmapped programs replace the allocated programs and potentially frees up memory for unused precached programs."),
+	ECVF_RenderThreadSafe
+);
+
+static int32 GBinaryCacheMaxPermittedSizeMB = 350;
+static FAutoConsoleVariableRef CVarBinaryCacheMaxPermittedSizeMB(
+	TEXT("r.OpenGL.BinaryCacheMaxPermittedSize"),
+	GBinaryCacheMaxPermittedSizeMB,
+	TEXT("When r.PSOPrecaching is active and the binary cache's size is greater\n")
+	TEXT("than this value the cache will be deleted at startup. The precaching cache is rebuilt from empty."),
+	ECVF_RenderThreadSafe
+);
+
 namespace UE
 {
 	namespace OpenGL
@@ -80,42 +107,58 @@ namespace UE
 			return StoreCompressedBinariesCVar->GetInt() != 0;
 		}
 
-		static const uint32 GBinaryProgramFileVersion = 5;
+		static const uint32 GBinaryProgramFileVersion = 6;
 		struct FBinaryCacheFileHeader
 		{
-			uint32 Version = 0;
+			uint32 Version = 0xFFFFFFFF;
 			FGuid BinaryCacheGuid;
 			bool bCacheUsesCompressedBinaries;
-			uint32 ProgramCount;
+			uint32 ProgramCount = 0;
+			uint32 ValidSize = 0;
 
-			static FBinaryCacheFileHeader CreateHeader(const FGuid& BinaryCacheGuidIn, uint32 NumPrograms)
+			static FBinaryCacheFileHeader CreateHeader(const FGuid& BinaryCacheGuidIn, uint32 NumPrograms, uint32 ValidSize)
 			{
 				FBinaryCacheFileHeader NewHeader;
 				NewHeader.Version = GBinaryProgramFileVersion;
 				NewHeader.BinaryCacheGuid = BinaryCacheGuidIn;
 				NewHeader.bCacheUsesCompressedBinaries = UE::OpenGL::AreBinaryProgramsCompressed();
 				NewHeader.ProgramCount = NumPrograms;
+				NewHeader.ValidSize = ValidSize;
 				return NewHeader;
 			}
 
 			friend FArchive& operator<<(FArchive& Ar, FBinaryCacheFileHeader& Header)
 			{
 				Ar << Header.Version;
-				check(Ar.IsLoading() || Header.Version == GBinaryProgramFileVersion); // This should always be correct when saving.
-				if (Header.Version == GBinaryProgramFileVersion)
+				check(Ar.IsLoading() || Header.IsValidVersion()); // This should always be correct when saving.
+				if (Header.IsValidVersion())
 				{
 					Ar << Header.BinaryCacheGuid;
 					Ar << Header.bCacheUsesCompressedBinaries;
 					Ar << Header.ProgramCount;
+					Ar << Header.ValidSize;
 				}
 
 				return Ar;
 			}
 
-			bool IsValid() const
+			bool IsValidVersion() const
 			{
 				return Version == GBinaryProgramFileVersion;
 			}
+
+			bool IsValid(const FGuid* OptionalWantedGuid) const
+			{
+				return IsValidVersion()
+					&& (UE::OpenGL::AreBinaryProgramsCompressed() == bCacheUsesCompressedBinaries)
+					&& (OptionalWantedGuid == nullptr || (*OptionalWantedGuid == BinaryCacheGuid))
+					&& (ProgramCount > 0);
+			}
+
+			FString ToString() const
+			{
+				return FString::Format(TEXT("{0}, {1}, {2}, {3}"), { Version, bCacheUsesCompressedBinaries, *BinaryCacheGuid.ToString(), ProgramCount });
+			}
 		};
 	}
 }
@@ -125,24 +168,48 @@ namespace UE
 class FOpenGLProgramBinaryMapping : public FThreadSafeRefCountedObject
 {
 public:
-	FOpenGLProgramBinaryMapping(TUniquePtr<IMappedFileHandle> MappedCacheFileIn, TUniquePtr<IMappedFileRegion> MappedRegionIn, uint32 ProgramCountIfKnown) : MappedCacheFile(MoveTemp(MappedCacheFileIn)), MappedRegion(MoveTemp(MappedRegionIn)) { Content.Reserve(ProgramCountIfKnown); }
+	FOpenGLProgramBinaryMapping(TUniquePtr<IMappedFileHandle> MappedCacheFileIn, uint32 ProgramCountIfKnown) : MappedCacheFile(MoveTemp(MappedCacheFileIn)) { Content.Reserve(ProgramCountIfKnown); }
 
 	TArrayView<const uint8> GetView(uint32 FileOffset, uint32 NumBytes) const
 	{
-		check(FileOffset + NumBytes <= MappedRegion->GetMappedSize());
-		return TArrayView<const uint8>(MappedRegion->GetMappedPtr() + FileOffset, NumBytes);
+		check(FileOffset >= CurrentMappingRegionOffset);
+		uint32 OffsetWithinMapping = FileOffset - CurrentMappingRegionOffset;
+		check(OffsetWithinMapping + NumBytes <= GetCurrentMappedRegion()->GetMappedSize());
+		return TArrayView<const uint8>(GetCurrentMappedRegion()->GetMappedPtr() + OffsetWithinMapping, NumBytes);
 	}
 
+	void AddMapping(uint32 MappingRegionOffset, TUniquePtr<IMappedFileRegion> NewMappedRegionIn)
+	{
+		check(MappingRegionOffset != 0 || MappedRegions.IsEmpty());
+		CurrentMappingRegionOffset = MappingRegionOffset;
+		MappedRegions.Add(MoveTemp(NewMappedRegionIn));
+	};
+
 	void AddProgramKey(const class FOpenGLProgramKey& KeyIn) { check(!Content.Contains(KeyIn)); Content.Add(KeyIn); }
-	bool HasValidMapping() const { return MappedRegion.IsValid() && MappedCacheFile.IsValid(); }
+	bool HasValidMapping() const { return MappedCacheFile.IsValid() && !MappedRegions.IsEmpty() && MappedRegions.Last().IsValid(); }
 	int32 NumPrograms() const { return Content.Num(); }
+
+	TUniquePtr<IMappedFileHandle>& GetMappedCacheFile() { return MappedCacheFile; }
+
+	TUniquePtr<IMappedFileRegion>& GetCurrentMappedRegion()
+	{
+		check(!MappedRegions.IsEmpty());
+		return MappedRegions.Last();
+	};
+
+	const TUniquePtr<IMappedFileRegion>& GetCurrentMappedRegion() const
+	{
+		check(!MappedRegions.IsEmpty());
+		return MappedRegions.Last();
+	};
+
 private:
 	TUniquePtr<IMappedFileHandle> MappedCacheFile;
-	TUniquePtr<IMappedFileRegion> MappedRegion;
+	TArray<TUniquePtr<IMappedFileRegion>> MappedRegions;
+	uint32 CurrentMappingRegionOffset = 0;
 	TSet<FOpenGLProgramKey> Content;
 };
 
-
 static FCriticalSection GProgramBinaryFileCacheCS;
 
 // guards the container that collects scanned programs and send to RHIT
@@ -208,6 +275,13 @@ FOpenGLProgramBinaryCache::FOpenGLProgramBinaryCache(const FString& InCachePathR
 
 FOpenGLProgramBinaryCache::~FOpenGLProgramBinaryCache()
 {
+#if PLATFORM_ANDROID
+	if (FAndroidOpenGL::AreRemoteCompileServicesActive())
+	{
+		FAndroidOpenGL::StopRemoteCompileServices();
+	}
+#endif
+
 	if (BinaryCacheWriteFileHandle)
 	{
 		delete BinaryCacheWriteFileHandle;
@@ -238,6 +312,8 @@ bool FOpenGLProgramBinaryCache::IsBuildingCache()
 	return false;
 }
 
+extern bool IsPrecachingEnabled();
+
 void FOpenGLProgramBinaryCache::Initialize()
 {
 	check(CachePtr == nullptr);
@@ -286,10 +362,18 @@ void FOpenGLProgramBinaryCache::Initialize()
 	CachePtr = new FOpenGLProgramBinaryCache(CacheFolderPathRoot);
 	UE_LOG(LogRHI, Log, TEXT("Enabling program binary cache dir at %s"), *CachePtr->GetProgramBinaryCacheDir());
 
-	// Add delegates for the ShaderPipelineCache precompile.
-	UE_LOG(LogRHI, Log, TEXT("FOpenGLProgramBinaryCache will be initialized when ShaderPipelineCache opens its file"));
-	CachePtr->OnShaderPipelineCacheOpenedDelegate = FShaderPipelineCache::GetCacheOpenedDelegate().AddRaw(CachePtr, &FOpenGLProgramBinaryCache::OnShaderPipelineCacheOpened);
-	CachePtr->OnShaderPipelineCachePrecompilationCompleteDelegate = FShaderPipelineCache::GetPrecompilationCompleteDelegate().AddRaw(CachePtr, &FOpenGLProgramBinaryCache::OnShaderPipelineCachePrecompilationComplete);
+
+	if (IsPrecachingEnabled())
+	{
+		CachePtr->InitPrecaching();
+	}
+	else
+	{
+		// Add delegates for the ShaderPipelineCache precompile.
+		UE_LOG(LogRHI, Log, TEXT("FOpenGLProgramBinaryCache will be initialized when ShaderPipelineCache opens its file"));
+		CachePtr->OnShaderPipelineCacheOpenedDelegate = FShaderPipelineCache::GetCacheOpenedDelegate().AddRaw(CachePtr, &FOpenGLProgramBinaryCache::OnShaderPipelineCacheOpened);
+		CachePtr->OnShaderPipelineCachePrecompilationCompleteDelegate = FShaderPipelineCache::GetPrecompilationCompleteDelegate().AddRaw(CachePtr, &FOpenGLProgramBinaryCache::OnShaderPipelineCachePrecompilationComplete);
+	}
 }
 
 #if PLATFORM_ANDROID
@@ -298,28 +382,153 @@ static FAutoConsoleVariableRef CVarNumRemoteProgramCompileServices(
 	TEXT("Android.OpenGL.NumRemoteProgramCompileServices"),
 	GNumRemoteProgramCompileServices,
 	TEXT("The number of separate processes to make available to compile opengl programs.\n")
-	TEXT("0 to disable use of separate processes to precompile PSOs\n")
+	TEXT("0 to disable use of separate processes to precompile Programs\n")
 	TEXT("valid range is 1-8 (4 default).")
 	,
 	ECVF_RenderThreadSafe | ECVF_ReadOnly
 );
 #endif
 
+
+namespace UE
+{
+	namespace OpenGL
+	{
+		static bool IsBinaryCacheValid(const FString& CachePath, const FGuid* OptionalGuidCheck, bool bLogWhenInvalid, UE::OpenGL::FBinaryCacheFileHeader* HeaderOUT)
+		{
+			TUniquePtr<FArchive> BinaryProgramReader = nullptr;
+			BinaryProgramReader = TUniquePtr<FArchive>(IFileManager::Get().CreateFileReader(*CachePath));
+			if (BinaryProgramReader.IsValid())
+			{
+				FArchive& Ar = *BinaryProgramReader;
+				UE::OpenGL::FBinaryCacheFileHeader BinaryCacheHeader;
+				if (Ar.TotalSize() > 0)
+				{
+					Ar << BinaryCacheHeader;
+				}
+
+				BinaryProgramReader->Close();
+
+				if (BinaryCacheHeader.IsValid(OptionalGuidCheck))
+				{
+					if (HeaderOUT)
+					{
+						*HeaderOUT = BinaryCacheHeader;
+					}
+					// we could validate file content.
+					return true;
+				}
+				UE_CLOG(bLogWhenInvalid, LogRHI, Warning, TEXT("FOpenGLProgramBinaryCache: %s is an invalid binary cache (%s)"), *CachePath, *BinaryCacheHeader.ToString());
+			}
+			else
+			{
+				UE_CLOG(bLogWhenInvalid, LogRHI, Warning, TEXT("FOpenGLProgramBinaryCache: %s not found."), *CachePath);
+			}
+
+			return false;
+		}
+	}
+}
+
+FShaderPipelineCache::FShaderCachePrecompileContext GShaderCachePrecompileContext;
+
+void FOpenGLProgramBinaryCache::InitPrecaching()
+{
+	UE_LOG(LogRHI, Log, TEXT("FOpenGLProgramBinaryCache: using precache system."));
+
+	FScopeLock Lock(&GProgramBinaryFileCacheCS);
+
+	const FString PreCacheName(TEXT("PrecacheBinaries"));
+
+	const FString RootDir = CachePathRoot / CacheSubDir;
+
+	TArray<FString> CacheFileNames;
+	IFileManager::Get().FindFiles(CacheFileNames, *RootDir, TEXT("*"));
+	for (FString& Filename : CacheFileNames)
+	{
+		Filename = RootDir / Filename;
+	}
+	
+	// remove old caches (such as those used by the PSO file cache) or precache caches that are oversized.
+	UE::OpenGL::FBinaryCacheFileHeader PrecacheBinaryHeader;
+
+	for (const FString& CacheFileName : CacheFileNames)
+	{
+		FString CacheName = FPaths::GetBaseFilename(CacheFileName);
+		IPlatformFile& PlatformFile = FPlatformFileManager::Get().GetPlatformFile();
+		FFileStatData CacheStatData = PlatformFile.GetStatData(*CacheFileName);
+
+		bool bDeleteMe;
+		if (CacheName.Equals(PreCacheName) && CacheStatData.FileSize <= GBinaryCacheMaxPermittedSizeMB * 1024 * 1024)
+		{
+			bDeleteMe = !UE::OpenGL::IsBinaryCacheValid(CacheFileName, nullptr, true, &PrecacheBinaryHeader);
+		}
+		else
+		{
+			bDeleteMe = true;
+		}
+
+		if (bDeleteMe)
+		{
+			UE_LOG(LogRHI, Log, TEXT("Deleting binary cache file %s (%d bytes)"), *CacheFileName, CacheStatData.FileSize);
+			PlatformFile.DeleteFile(*CacheName);
+		}
+	}
+
+	// Use the existing cache's guid to load it.
+	FGuid CacheGuid = PrecacheBinaryHeader.IsValid(nullptr) ? PrecacheBinaryHeader.BinaryCacheGuid : FGuid::NewGuid();
+	GShaderCachePrecompileContext = FShaderPipelineCache::FShaderCachePrecompileContext(PreCacheName);
+
+	UE_LOG(LogRHI, Log, TEXT("InitPrecaching : Beginning new cache = %s"), *GShaderCachePrecompileContext.GetCacheName());
+	OnShaderPipelineCacheOpened(
+		PreCacheName,
+		GetFeatureLevelShaderPlatform(FOpenGL::GetFeatureLevel()),
+		1,	// count must be > 0. 0 sized caches are ignored.
+		CacheGuid,
+		GShaderCachePrecompileContext);
+}
+
+// if the file has sufficiently grown flush the program cache and mmap the extra programs.
+void FOpenGLProgramBinaryCache::UpdatePrecacheMapping()
+{
+	QUICK_SCOPE_CYCLE_COUNTER(STAT_OpenGLUpdatePrecacheMapping);
+	// TODO: avoid this lock.
+	FScopeLock Lock(&GProgramBinaryFileCacheCS);
+	FArchive& DestAr = *BinaryCacheWriteFileHandle;
+	int64 DestPos = DestAr.Tell();
+
+	if (DestPos >= CurrentShaderPipelineProperties.LastMappedPosition + (GBinaryCacheMMapAfterEveryMB * 1024 * 1024))
+	{
+		int32 ProgramCount = ProgramsInCurrentCache.Num();
+		MarkValidContent(ProgramCount);
+
+		UE_LOG(LogRHI, Log, TEXT("UpdatePrecacheMapping : Mapping program cache from %d to %d "), CurrentShaderPipelineProperties.LastMappedPosition, DestPos);
+
+		TUniquePtr<FArchive> BinaryProgramReader = TUniquePtr<FArchive>(IFileManager::Get().CreateFileReader(*GetProgramBinaryCacheFilePath()));
+ 		BinaryProgramReader->Seek(CurrentShaderPipelineProperties.LastMappedPosition);
+		int32 ProgramsToRead = ProgramCount - CurrentShaderPipelineProperties.MappedPrograms;
+ 		bool bSuccess = ReadProgramFile_Internal(ProgramsToRead, DestPos, GetProgramBinaryCacheFilePath(), *BinaryProgramReader);
+		CurrentShaderPipelineProperties.MappedPrograms = ProgramCount;
+		check(bSuccess);
+	}
+}
+
 void FOpenGLProgramBinaryCache::OnShaderPipelineCacheOpened(FString const& Name, EShaderPlatform Platform, uint32 Count, const FGuid& VersionGuid, FShaderPipelineCache::FShaderCachePrecompileContext& ShaderCachePrecompileContext)
 {
 	FScopeLock Lock(&GProgramBinaryFileCacheCS);
-	checkf(CurrentShaderPipelineCacheVersionGuid == FGuid(), TEXT("OGL: OnShaderPipelineCacheOpened, previous PSO cache %s (%s) has not completed!"), *CurrentShaderPipelineCacheName, *CurrentShaderPipelineCacheVersionGuid.ToString());
-	CurrentShaderPipelineCacheVersionGuid = VersionGuid;
-	CurrentShaderPipelineCacheName = ShaderCachePrecompileContext.GetCacheName();
+
+	checkf(CurrentShaderPipelineProperties.CacheVersionGuid == FGuid(), TEXT("OGL: OnShaderPipelineCacheOpened, previous PSO cache %s (%s) has not completed!"), *CurrentShaderPipelineProperties.PipelineCacheName, *CurrentShaderPipelineProperties.CacheVersionGuid.ToString());
+	CurrentShaderPipelineProperties.CacheVersionGuid = VersionGuid;
+	CurrentShaderPipelineProperties.PipelineCacheName = ShaderCachePrecompileContext.GetCacheName();
 
 	if (Count == 0)
 	{		
 		check(CurrentBinaryFileState == EBinaryFileState::Uninitialized);
-		UE_LOG(LogRHI, Verbose, TEXT("OnShaderPipelineCacheOpened, Ignoring empty PSO cache. %s (%s)"), *CurrentShaderPipelineCacheName, *CurrentShaderPipelineCacheVersionGuid.ToString());
+		UE_LOG(LogRHI, Verbose, TEXT("OnShaderPipelineCacheOpened, Ignoring empty PSO cache. %s (%s)"), *CurrentShaderPipelineProperties.PipelineCacheName, *CurrentShaderPipelineProperties.CacheVersionGuid.ToString());
 		return;
 	}
 	 
-	UE_LOG(LogRHI, Log, TEXT("Scanning Binary program cache, using Shader Pipeline Cache %s (%s)"), *CurrentShaderPipelineCacheName, *CurrentShaderPipelineCacheVersionGuid.ToString());
+	UE_LOG(LogRHI, Log, TEXT("Scanning Binary program cache, using Shader Pipeline Cache %s (%s)"), *CurrentShaderPipelineProperties.PipelineCacheName, *CurrentShaderPipelineProperties.CacheVersionGuid.ToString());
 
 	ScanProgramCacheFile();
 	if (IsBuildingCache_internal())
@@ -327,7 +536,11 @@ void FOpenGLProgramBinaryCache::OnShaderPipelineCacheOpened(FString const& Name,
 #if PLATFORM_ANDROID
 		if (GNumRemoteProgramCompileServices)
 		{
-			FAndroidOpenGL::StartAndWaitForRemoteCompileServices(GNumRemoteProgramCompileServices);
+			check(!FAndroidOpenGL::AreRemoteCompileServicesActive() || IsPrecachingEnabled());
+			if(FAndroidOpenGL::AreRemoteCompileServicesActive() == false)
+			{
+				FAndroidOpenGL::StartAndWaitForRemoteCompileServices(GNumRemoteProgramCompileServices);
+			}
 		}
 #endif
 		ShaderCachePrecompileContext.SetPrecompilationIsSlowTask();
@@ -350,7 +563,7 @@ void FOpenGLProgramBinaryCache::OnShaderPipelineCachePrecompilationComplete(uint
 	// if 0 programs were cached when Count >0, then we suffer a performance penalty for invoking the services for no reason.
 	
 	const bool bIsBuildingCache = IsBuildingCache_internal();
-	TRefCountPtr<FOpenGLProgramBinaryMapping>*FoundCache = MappedCacheFiles.Find(CurrentShaderPipelineCacheVersionGuid);
+	TRefCountPtr<FOpenGLProgramBinaryMapping>*FoundCache = MappedCacheFiles.Find(CurrentShaderPipelineProperties.CacheVersionGuid);
 	const int32 ProgramsInLoadedCache = FoundCache ? (*FoundCache)->NumPrograms() : 0;
 	const int32 ProgramsCached = bIsBuildingCache ? ProgramsInCurrentCache.Num() : ProgramsInLoadedCache;
 
@@ -362,18 +575,18 @@ void FOpenGLProgramBinaryCache::OnShaderPipelineCachePrecompilationComplete(uint
 		:
 		(Count == 0 ? TEXT("ignored empty cache") : TEXT("cache loaded"));
 
-	UE_LOG(LogRHI, Log, TEXT("OnShaderPipelineCachePrecompilationComplete: %s(%s) - %s %d program binaries (%d requested)"), *CurrentShaderPipelineCacheName, *CurrentShaderPipelineCacheVersionGuid.ToString(), CacheStatusText, ProgramsCached, Count );
+	UE_LOG(LogRHI, Log, TEXT("OnShaderPipelineCachePrecompilationComplete: %s(%s) - %s %d program binaries (%d requested)"), *CurrentShaderPipelineProperties.PipelineCacheName, *CurrentShaderPipelineProperties.CacheVersionGuid.ToString(), CacheStatusText, ProgramsCached, Count );
 
 	if (bIsBuildingCache)
 	{
 #if PLATFORM_ANDROID
-		if (GNumRemoteProgramCompileServices)
+		if ( !IsPrecachingEnabled() && FAndroidOpenGL::AreRemoteCompileServicesActive())
 		{
 			FAndroidOpenGL::StopRemoteCompileServices();
 		}
 #endif
 
-		const bool bSuccess = CloseCacheWriteHandle();
+		const bool bSuccess = CloseCacheWriteHandle(ProgramsInCurrentCache.Num());
 
 		if(bSuccess && CVarRestartAndroidAfterPrecompile.GetValueOnAnyThread() == 1)
 		{
@@ -393,9 +606,9 @@ void FOpenGLProgramBinaryCache::OnShaderPipelineCachePrecompilationComplete(uint
 			check(!IsBuildingCache_internal());
 			if (IsBuildingCache_internal())
 			{
-				UE_LOG( LogRHI, Error, TEXT("Failed to load just completed cache! : %s(%s)"), *CurrentShaderPipelineCacheName, *CurrentShaderPipelineCacheVersionGuid.ToString());
+				UE_LOG( LogRHI, Error, TEXT("Failed to load just completed cache! : %s(%s)"), *CurrentShaderPipelineProperties.PipelineCacheName, *CurrentShaderPipelineProperties.CacheVersionGuid.ToString());
 				// The cache we've just written is invalid. This is extremely unlikely.
-				CloseCacheWriteHandle();				
+				CloseCacheWriteHandle(ProgramsInCurrentCache.Num());
 			}
 		}
 	}
@@ -403,8 +616,52 @@ void FOpenGLProgramBinaryCache::OnShaderPipelineCachePrecompilationComplete(uint
 	// unset the completed cache.
 	Reset();
 
-	CurrentShaderPipelineCacheVersionGuid = FGuid();
-	CurrentShaderPipelineCacheName.Reset();
+	CurrentShaderPipelineProperties.CacheVersionGuid = FGuid();
+	CurrentShaderPipelineProperties.PipelineCacheName.Reset();
+	CurrentShaderPipelineProperties.LastMappedPosition = 0;
+	CurrentShaderPipelineProperties.MappedPrograms = 0;
+	CurrentShaderPipelineProperties.NumProgramsFlushed = 0;
+}
+
+TRefCountPtr<FOpenGLProgramBinaryMapping> FOpenGLProgramBinaryCache::GetOrAddFileMapping(const FString& ProgramCacheFilename, int32 ProgramCount, int64 Offset, int64 Size)
+{
+	TRefCountPtr<FOpenGLProgramBinaryMapping> CurrentMapping;
+	if (!UE::OpenGL::CanMemoryMapGLProgramCache())
+	{
+		return CurrentMapping;
+	}
+
+	TRefCountPtr<FOpenGLProgramBinaryMapping>* FoundMapping = MappedCacheFiles.Find(CurrentShaderPipelineProperties.CacheVersionGuid);
+	if( FoundMapping == nullptr )
+	{
+		TUniquePtr<IMappedFileHandle> MappedCacheFile;
+		if(UE::OpenGL::CanMemoryMapGLProgramCache())
+		{
+			MappedCacheFile = TUniquePtr<IMappedFileHandle>(FPlatformFileManager::Get().GetPlatformFile().OpenMapped(*ProgramCacheFilename));
+		}
+		CurrentMapping = new FOpenGLProgramBinaryMapping(MoveTemp(MappedCacheFile), ProgramCount);
+		MappedCacheFiles.Add(CurrentShaderPipelineProperties.CacheVersionGuid, CurrentMapping);		
+	}
+	else
+	{
+		CurrentMapping = *FoundMapping;
+	}
+
+	if(UE::OpenGL::CanMemoryMapGLProgramCache())
+	{
+		TUniquePtr<IMappedFileHandle>& MappedCacheFile = CurrentMapping->GetMappedCacheFile();
+
+		// add the new mapping
+		if (ensure(MappedCacheFile.IsValid()))
+		{
+			check(Size);
+			TUniquePtr<IMappedFileRegion> MappedRegion = TUniquePtr<IMappedFileRegion>(MappedCacheFile->MapRegion(Offset, Size));
+			check(MappedRegion.IsValid());
+			CurrentMapping->AddMapping(Offset, MoveTemp(MappedRegion));
+		}
+	}
+
+	return CurrentMapping;
 }
 
 // Scan the binary cache file and build a record of all programs.
@@ -422,130 +679,53 @@ void FOpenGLProgramBinaryCache::ScanProgramCacheFile()
 	check(ProgramsInCurrentCache.IsEmpty());
 
 	bool bBinaryFileIsValid = false;
-	bool bBinaryFileIsValidAndGuidMatch = false;
 
 	// Try to move the file to a temporary filename before the scan, so we won't try to read it again if it's corrupted
 	PlatformFile.DeleteFile(*ProgramCacheFilenameTemp);
 	PlatformFile.MoveFile(*ProgramCacheFilenameTemp, *ProgramCacheFilename);
 
 	TUniquePtr<FArchive> BinaryProgramReader = nullptr;
-	BinaryProgramReader = TUniquePtr<FArchive>( IFileManager::Get().CreateFileReader(*ProgramCacheFilenameTemp) );
+	BinaryProgramReader = TUniquePtr<FArchive>(IFileManager::Get().CreateFileReader(*ProgramCacheFilenameTemp));
 
-	if (BinaryProgramReader)
+	if (!BinaryProgramReader)
+	{
+		UE_LOG(LogRHI, Log, TEXT("OnShaderScanProgramCacheFile : %s was not found, recreating."), *ProgramCacheFilename);
+	}
+	else
 	{
-		FArchive& Ar = *BinaryProgramReader;
 		UE::OpenGL::FBinaryCacheFileHeader BinaryCacheHeader;
-		Ar << BinaryCacheHeader;
-		if (BinaryCacheHeader.IsValid())
+		FArchive& Ar = *BinaryProgramReader;
+		if (Ar.TotalSize() > 0)
 		{
-			const bool bUseCompressedProgramBinaries = UE::OpenGL::AreBinaryProgramsCompressed();
-			bBinaryFileIsValid = (bUseCompressedProgramBinaries == BinaryCacheHeader.bCacheUsesCompressedBinaries);
-			bBinaryFileIsValidAndGuidMatch = bBinaryFileIsValid && (!CurrentShaderPipelineCacheVersionGuid.IsValid() || CurrentShaderPipelineCacheVersionGuid == BinaryCacheHeader.BinaryCacheGuid);
-
-			if (bBinaryFileIsValidAndGuidMatch)
-			{
-				int32 ProgramIndex = 0;
-				FScopeLock Lock(&GPendingGLProgramCreateRequestsCS);
-				PendingGLContainerPrograms.Reserve(BinaryCacheHeader.ProgramCount);
-				if (BinaryCacheHeader.ProgramCount > 0)
-				{
-					TRefCountPtr<FOpenGLProgramBinaryMapping> CurrentMapping;
-					// success
-					{
-						TUniquePtr<IMappedFileHandle> MappedCacheFile;
-						TUniquePtr<IMappedFileRegion> MappedRegion;
-						if(UE::OpenGL::CanMemoryMapGLProgramCache())
-						{
-							MappedCacheFile = TUniquePtr<IMappedFileHandle>(FPlatformFileManager::Get().GetPlatformFile().OpenMapped(*ProgramCacheFilenameTemp));
-							if (ensure(MappedCacheFile.IsValid()))
-							{
-								MappedRegion = TUniquePtr<IMappedFileRegion>(MappedCacheFile->MapRegion());
-								check(MappedRegion.IsValid());
-							}
-						}
-						CurrentMapping = new FOpenGLProgramBinaryMapping(MoveTemp(MappedCacheFile), MoveTemp(MappedRegion), BinaryCacheHeader.ProgramCount);
-					}
-
-					UE_LOG(LogRHI, Log, TEXT("OnShaderScanProgramCacheFile : %s %s"), CurrentMapping->HasValidMapping() ? TEXT("mapped") : TEXT("opened"), *ProgramCacheFilenameTemp);
-
-					MappedCacheFiles.Add(CurrentShaderPipelineCacheVersionGuid, CurrentMapping);
-					ProgramsInCurrentCache.Reserve(BinaryCacheHeader.ProgramCount);
-					while (!Ar.AtEnd())
-					{
-						FOpenGLProgramKey ProgramKey;
-						uint32 ProgramBinarySize = 0;
-						Ar << ProgramKey;
-						Ar << ProgramBinarySize;
-						check(ProgramKey != FOpenGLProgramKey());
-						if (ensure(ProgramBinarySize > 0))
-						{
-							ProgramIndex++;
-							uint32 ProgramBinaryOffset = Ar.Tell();
-							CurrentMapping->AddProgramKey(ProgramKey);
-
-							UE_LOG(LogRHI, VeryVerbose, TEXT(" scan found PSO %s - %d"), *ProgramKey.ToString(), ProgramBinarySize);
-
-							ProgramsInCurrentCache.Add(ProgramKey);
-
-							if(CurrentMapping->HasValidMapping())
-							{
-								PendingGLContainerPrograms.Emplace(ProgramKey, TUniqueObj<FOpenGLProgramBinary>(CurrentMapping->GetView(ProgramBinaryOffset, ProgramBinarySize)));
-								Ar.Seek(ProgramBinaryOffset + ProgramBinarySize);
-							}
-							else
-							{
-								check(!UE::OpenGL::CanMemoryMapGLProgramCache());
-								TArray<uint8> ProgramBytes;
-								ProgramBytes.SetNumUninitialized(ProgramBinarySize);
-								Ar.Serialize(ProgramBytes.GetData(), ProgramBinarySize);
-								PendingGLContainerPrograms.Emplace(ProgramKey, TUniqueObj<FOpenGLProgramBinary>(MoveTemp(ProgramBytes)));
-							}
-						}
-					}
-
-					UE_LOG(LogRHI, VeryVerbose, TEXT("Program Binary cache: Found %d cached programs"), ProgramIndex);
-					UE_CLOG(ProgramIndex != BinaryCacheHeader.ProgramCount, LogRHI, Error, TEXT("Program Binary cache: Mismatched program count! expected: %d"), BinaryCacheHeader.ProgramCount);
-
-					BinaryProgramReader->Close();
-					// Rename the file back after a successful scan.
-					PlatformFile.MoveFile(*ProgramCacheFilename, *ProgramCacheFilenameTemp);
-					CurrentBinaryFileState = EBinaryFileState::ValidCacheFile;
-				}
-				else
-				{
-					// failed to find sentinel record, the file was not finalized.
-					UE_LOG(LogRHI, Warning, TEXT("ScanProgramCacheFile - incomplete or empty binary cache file encountered. Rebuilding binary program cache."));
-					BinaryProgramReader->Close();
-					bBinaryFileIsValid = false;
-					bBinaryFileIsValidAndGuidMatch = false;
-					PlatformFile.DeleteFile(*ProgramCacheFilenameTemp);
-				}
-			}
-			else
-			{
-				UE_LOG(LogRHI, Log, TEXT("OnShaderScanProgramCacheFile : binary file found but is invalid (%s, %s), recreating."), *CurrentShaderPipelineCacheVersionGuid.ToString(), *BinaryCacheHeader.BinaryCacheGuid.ToString());
-			}
+			Ar << BinaryCacheHeader;
+		}
+		bBinaryFileIsValid = BinaryCacheHeader.IsValid(&CurrentShaderPipelineProperties.CacheVersionGuid);
+		if (!bBinaryFileIsValid)
+		{
+			UE_LOG(LogRHI, Warning, TEXT("ScanProgramCacheFile - invalid binary cache file encountered (%s, %d). Rebuilding binary program cache."), *BinaryCacheHeader.ToString(), Ar.TotalSize());
+			BinaryProgramReader->Close();
+			PlatformFile.DeleteFile(*ProgramCacheFilenameTemp);
 		}
 		else
 		{
-			UE_LOG(LogRHI, Log, TEXT("OnShaderScanProgramCacheFile : binary file version invalid, recreating."));
+			UE_LOG(LogRHI, Warning, TEXT("ScanProgramCacheFile - reading binary cache, Programs %d, ValidSize: %d, TotalSize %d"), BinaryCacheHeader.ProgramCount, BinaryCacheHeader.ValidSize, Ar.TotalSize());
+			ProgramsInCurrentCache.Reserve(BinaryCacheHeader.ProgramCount);
+			bool bSuccess = ReadProgramFile_Internal(BinaryCacheHeader.ProgramCount, BinaryCacheHeader.ValidSize, ProgramCacheFilenameTemp, Ar);
+			check(bSuccess);
 		}
 	}
-	else
-	{
-		UE_LOG(LogRHI, Log, TEXT("OnShaderScanProgramCacheFile : %s was not found, recreating."), *ProgramCacheFilename);
-	}
 
-	if (!bBinaryFileIsValidAndGuidMatch)
+	if (!bBinaryFileIsValid)
 	{
-		if (OpenCacheWriteHandle())
+		if (OpenCacheWriteHandle(GetProgramBinaryCacheFilePath(), false))
 		{
 			CurrentBinaryFileState = EBinaryFileState::BuildingCacheFile;
 			// save header, 0 program count indicates an unfinished file.
 			// The header is overwritten at the end of the process.
-			UE::OpenGL::FBinaryCacheFileHeader OutHeader = UE::OpenGL::FBinaryCacheFileHeader::CreateHeader(CurrentShaderPipelineCacheVersionGuid, 0);
+			UE::OpenGL::FBinaryCacheFileHeader OutHeader = UE::OpenGL::FBinaryCacheFileHeader::CreateHeader(CurrentShaderPipelineProperties.CacheVersionGuid, 0, 0);
 			FArchive& Ar = *BinaryCacheWriteFileHandle;
 			Ar << OutHeader;
+			CurrentShaderPipelineProperties.LastMappedPosition = Ar.Tell();
 		}
 		else
 		{
@@ -555,26 +735,130 @@ void FOpenGLProgramBinaryCache::ScanProgramCacheFile()
 			UE_LOG(LogRHI, Fatal, TEXT("ScanProgramCacheFile - Failed to open binary cache."));
 		}
 	}
+	else
+	{
+		uint32 LastReadPos = BinaryProgramReader->Tell();
+		BinaryProgramReader->Close();
+
+		// Rename the file back after a successful scan.
+		PlatformFile.MoveFile(*ProgramCacheFilename, *ProgramCacheFilenameTemp);
+		CurrentBinaryFileState = EBinaryFileState::ValidCacheFile;
+
+		if (IsPrecachingEnabled())
+		{
+			if (OpenCacheWriteHandle(GetProgramBinaryCacheFilePath(), true))
+			{
+				check(LastReadPos == CurrentShaderPipelineProperties.LastMappedPosition);
+				CurrentBinaryFileState = EBinaryFileState::BuildingCacheFile;
+				BinaryCacheWriteFileHandle->Seek(CurrentShaderPipelineProperties.LastMappedPosition);
+			}
+		}
+	}
 }
 
-bool FOpenGLProgramBinaryCache::OpenCacheWriteHandle()
+bool FOpenGLProgramBinaryCache::ReadProgramFile_Internal(uint32 ProgramsToRead, int64 EndOffset, const FString& ProgramCacheFilename, FArchive& Ar)
 {
-	check(BinaryCacheWriteFileHandle == nullptr);
+	QUICK_SCOPE_CYCLE_COUNTER(STAT_OpenGLReadBinaryProgramCache);
+	int32 ProgramIndex = 0;
+	FScopeLock Lock(&GPendingGLProgramCreateRequestsCS);
+	PendingGLContainerPrograms.Reserve(ProgramsToRead);
+	int64 MappingSize = EndOffset - CurrentShaderPipelineProperties.LastMappedPosition;
+	TRefCountPtr<FOpenGLProgramBinaryMapping> CurrentMapping = GetOrAddFileMapping(ProgramCacheFilename, ProgramsToRead, CurrentShaderPipelineProperties.LastMappedPosition, MappingSize);
+
+	UE_LOG(LogRHI, Log, TEXT("OnShaderScanProgramCacheFile : %s %s"), CurrentMapping->HasValidMapping() ? TEXT("mapped") : TEXT("opened"), *ProgramCacheFilename);
+
+	while (Ar.Tell() < EndOffset)
+	{
+		FOpenGLProgramKey ProgramKey;
+		uint32 ProgramBinarySize = 0;
+		Ar << ProgramKey;
+		Ar << ProgramBinarySize;
+		check(ProgramKey != FOpenGLProgramKey());
+		if (ensure(ProgramBinarySize > 0))
+		{
+			ProgramIndex++;
+			uint32 ProgramBinaryOffset = Ar.Tell();
+			CurrentMapping->AddProgramKey(ProgramKey);
+
+			UE_LOG(LogRHI, VeryVerbose, TEXT(" scan found PSO %s - %d"), *ProgramKey.ToString(), ProgramBinarySize);
+
+			ProgramsInCurrentCache.Add(ProgramKey);
+
+			if (CurrentMapping->HasValidMapping())
+			{
+				PendingGLContainerPrograms.Emplace(ProgramKey, TUniqueObj<FOpenGLProgramBinary>(CurrentMapping->GetView(ProgramBinaryOffset, ProgramBinarySize)));
+				Ar.Seek(ProgramBinaryOffset + ProgramBinarySize);
+			}
+			else
+			{
+				check(!UE::OpenGL::CanMemoryMapGLProgramCache());
+				TArray<uint8> ProgramBytes;
+				ProgramBytes.SetNumUninitialized(ProgramBinarySize);
+				Ar.Serialize(ProgramBytes.GetData(), ProgramBinarySize);
+				PendingGLContainerPrograms.Emplace(ProgramKey, TUniqueObj<FOpenGLProgramBinary>(MoveTemp(ProgramBytes)));
+			}
+		}
+		else
+		{
+			return false;
+		}
+	}
 
-	// perform file writing to a file temporary filename so we don't attempt to use the file later if the write session is interrupted
-	const FString ProgramCacheFilenameWrite = GetProgramBinaryCacheFilePath() + TEXT(".write");
+	CurrentShaderPipelineProperties.LastMappedPosition = Ar.Tell();
 
-	BinaryCacheWriteFileHandle = IFileManager::Get().CreateFileWriter(*ProgramCacheFilenameWrite, EFileWrite::FILEWRITE_None);
+	UE_LOG(LogRHI, VeryVerbose, TEXT("Program Binary cache: Found %d cached programs"), ProgramIndex);
+	UE_CLOG(ProgramIndex != ProgramsToRead, LogRHI, Error, TEXT("Program Binary cache: Mismatched program count! expected: %d"), ProgramsToRead);
+	return true;
+}
+
+bool FOpenGLProgramBinaryCache::OpenCacheWriteHandle(const FString& ProgramCacheFilenameToWrite, bool bAppendToExisting)
+{
+	check(BinaryCacheWriteFileHandle == nullptr);
 
-	UE_CLOG(BinaryCacheWriteFileHandle == nullptr, LogRHI, Warning, TEXT("Failed to open OGL binary cache output file. (%s)"), *ProgramCacheFilenameWrite);
+	BinaryCacheWriteFileHandle = IFileManager::Get().CreateFileWriter(*ProgramCacheFilenameToWrite, EFileWrite::FILEWRITE_AllowRead | (bAppendToExisting ? EFileWrite::FILEWRITE_Append : EFileWrite::FILEWRITE_None));
+	UE_CLOG(BinaryCacheWriteFileHandle, LogRHI, Log, TEXT("Opened binary cache for write (%s)"), *ProgramCacheFilenameToWrite);
+	UE_CLOG(BinaryCacheWriteFileHandle == nullptr, LogRHI, Warning, TEXT("Failed to open OGL binary cache output file. (%s)"), *ProgramCacheFilenameToWrite);
 
 	return BinaryCacheWriteFileHandle != nullptr;
 }
 
-bool FOpenGLProgramBinaryCache::CloseCacheWriteHandle()
+// Update the file's header to indicate the last successful write location. This allows the app to ignore a truncated file such as after an abnormal exit.
+bool FOpenGLProgramBinaryCache::MarkValidContent(int32 NumPrograms)
 {
+	QUICK_SCOPE_CYCLE_COUNTER(STAT_OpenGLMarkValidBinaryProgramCache);
 	check(BinaryCacheWriteFileHandle != nullptr);
-	const bool bCacheFileIsEmpty = ProgramsInCurrentCache.IsEmpty();
+	const bool bCacheFileIsEmpty = NumPrograms == 0;
+
+	bool bArchiveFailed = BinaryCacheWriteFileHandle->IsError() || BinaryCacheWriteFileHandle->IsCriticalError();
+
+	// Overwrite the header with the final program count. This indicates a successful write.
+	if (!bArchiveFailed)
+	{
+		FArchive& Ar = *BinaryCacheWriteFileHandle;
+		int64 CurrPos = Ar.Tell();
+		Ar.Seek(0);
+		UE::OpenGL::FBinaryCacheFileHeader OutHeader = UE::OpenGL::FBinaryCacheFileHeader::CreateHeader(CurrentShaderPipelineProperties.CacheVersionGuid, NumPrograms, (int32)CurrPos);
+		UE_LOG(LogRHI, Verbose, TEXT("MarkValidContent, file valid at %d Programs, %d bytes."), NumPrograms, CurrPos);
+		Ar << OutHeader;
+		Ar.Seek(CurrPos);
+		bArchiveFailed = BinaryCacheWriteFileHandle->IsError() || BinaryCacheWriteFileHandle->IsCriticalError();
+	}
+
+	if (bArchiveFailed)
+	{
+		RHIGetPanicDelegate().ExecuteIfBound(FName("FailedBinaryProgramArchiveWrite"));
+		UE_LOG(LogRHI, Fatal, TEXT("MarkValidContent - FArchive error bit set, failed to write binary cache."));
+	}
+	CurrentShaderPipelineProperties.NumProgramsFlushed = NumPrograms;
+	BinaryCacheWriteFileHandle->Flush();
+
+	return !bCacheFileIsEmpty;
+}
+
+bool FOpenGLProgramBinaryCache::CloseCacheWriteHandle(int32 NumProgramsAdded)
+{
+	check(BinaryCacheWriteFileHandle != nullptr);
+	const bool bCacheFileIsEmpty = NumProgramsAdded == 0;
 
 	bool bArchiveFailed = BinaryCacheWriteFileHandle->IsError() || BinaryCacheWriteFileHandle->IsCriticalError();
 
@@ -583,7 +867,7 @@ bool FOpenGLProgramBinaryCache::CloseCacheWriteHandle()
 	{
 		FArchive& Ar = *BinaryCacheWriteFileHandle;
 		Ar.Seek(0);
-		UE::OpenGL::FBinaryCacheFileHeader OutHeader = UE::OpenGL::FBinaryCacheFileHeader::CreateHeader(CurrentShaderPipelineCacheVersionGuid, ProgramsInCurrentCache.Num());
+		UE::OpenGL::FBinaryCacheFileHeader OutHeader = UE::OpenGL::FBinaryCacheFileHeader::CreateHeader(CurrentShaderPipelineProperties.CacheVersionGuid, NumProgramsAdded, (int32)Ar.TotalSize());
 		Ar << OutHeader;
 		bArchiveFailed = BinaryCacheWriteFileHandle->IsError() || BinaryCacheWriteFileHandle->IsCriticalError();
 	}
@@ -593,7 +877,6 @@ bool FOpenGLProgramBinaryCache::CloseCacheWriteHandle()
 	BinaryCacheWriteFileHandle = nullptr;
 
 	const FString ProgramCacheFilename = GetProgramBinaryCacheFilePath();
-	const FString ProgramCacheFilenameWrite = ProgramCacheFilename + TEXT(".write");
 	IPlatformFile& PlatformFile = FPlatformFileManager::Get().GetPlatformFile();
 
 	if (bArchiveFailed)
@@ -605,18 +888,13 @@ bool FOpenGLProgramBinaryCache::CloseCacheWriteHandle()
 	if (bCacheFileIsEmpty)
 	{
 		// we dont want empty files left on disk.
-		PlatformFile.DeleteFile(*ProgramCacheFilenameWrite);
-	}
-	else
-	{
-		// rename the temp filename back to the final filename
-		PlatformFile.DeleteFile(*ProgramCacheFilename); // file should never exist, but for safety
-		PlatformFile.MoveFile(*ProgramCacheFilename, *ProgramCacheFilenameWrite);
+		PlatformFile.DeleteFile(*ProgramCacheFilename);
 	}
+
 	return !bCacheFileIsEmpty;
 }
 
-void FOpenGLProgramBinaryCache::CacheProgramBinary(const FOpenGLProgramKey& ProgramKey, const FOpenGLProgramBinary& BinaryProgramData)
+void FOpenGLProgramBinaryCache::CacheProgramBinary(const FOpenGLProgramKey& ProgramKey, TUniqueObj<FOpenGLProgramBinary>&& ProgramBinary)
 {
 	if (CachePtr)
 	{
@@ -624,7 +902,13 @@ void FOpenGLProgramBinaryCache::CacheProgramBinary(const FOpenGLProgramKey& Prog
 
 		if (!CachePtr->ProgramsInCurrentCache.Contains(ProgramKey))
 		{
-			CachePtr->AddProgramBinaryDataToBinaryCache(ProgramKey, BinaryProgramData);
+			CachePtr->AddProgramBinaryDataToBinaryCache(ProgramKey, ProgramBinary.Get());
+
+			if (IsPrecachingEnabled())
+			{
+				// If we're precaching then we need to send it to RHIT immediately. without this we'll wait for this portion of the in progress file to be mmapped.
+				EnqueueBinaryForGLProgramContainer(ProgramKey, MoveTemp(ProgramBinary));
+			}
 		}
 	}
 }
@@ -640,10 +924,17 @@ void FOpenGLProgramBinaryCache::AddProgramBinaryDataToBinaryCache(const FOpenGLP
 	const TArrayView<const uint8> BinaryProgramDataView = BinaryProgramData.GetDataView();
 	uint32 ProgramBinarySize = (uint32)BinaryProgramDataView.Num();
 	const uint8* ProgramBinaryBytes = BinaryProgramDataView.GetData();
+	uint32 Start = Ar.Tell();
+	check(SerializedProgramKey != FOpenGLProgramKey());
 	Ar << SerializedProgramKey;
 	uint32 ProgramBinaryOffset = Ar.Tell();
+	check(ProgramBinarySize > 0);
 	Ar << ProgramBinarySize;
 	Ar.Serialize(const_cast<uint8*>(ProgramBinaryBytes), ProgramBinarySize);
+	uint32 End = Ar.Tell();
+
+	UE_LOG(LogRHI, VeryVerbose, TEXT("AddProgramBinaryDataToBinaryCache : added %d bytes to cache"), End-Start);
+
 	if (UE::OpenGL::AreBinaryProgramsCompressed())
 	{
 		static uint32 TotalUncompressed = 0;
@@ -656,7 +947,17 @@ void FOpenGLProgramBinaryCache::AddProgramBinaryDataToBinaryCache(const FOpenGLP
 		UE_LOG(LogRHI, Verbose, TEXT("AppendProgramBinaryFile: total Uncompressed: %d, total Compressed %d, Total saved so far: %d"), TotalUncompressed, TotalCompressed, TotalUncompressed - TotalCompressed);
 	}
 	UE_LOG(LogRHI, VeryVerbose, TEXT("AddProgramBinaryDataToBinaryCache: written Program %s to cache (%d bytes)"), *ProgramKey.ToString(), BinaryProgramDataView.Num());
+
 	ProgramsInCurrentCache.Add(ProgramKey);
+
+ 	if (IsPrecachingEnabled())
+ 	{
+		int32 NumProgramsStored = ProgramsInCurrentCache.Num();
+		if (NumProgramsStored >= CurrentShaderPipelineProperties.NumProgramsFlushed + GBinaryCachePeriodicFlushProgramCount)
+		{
+			MarkValidContent(NumProgramsStored);
+		}
+ 	}
 }
 
 void FOpenGLProgramBinaryCache::EnqueueBinaryForGLProgramContainer(const FOpenGLProgramKey& ProgramKey, TUniqueObj<FOpenGLProgramBinary>&& ProgramBinary)
@@ -677,36 +978,43 @@ void FOpenGLProgramBinaryCache::Shutdown()
 	}
 }
 
-bool FOpenGLProgramBinaryCache::DoesCurrentCacheContain(const FOpenGLProgramKey& ProgramKey)
+bool FOpenGLProgramBinaryCache::RequiresCaching(const FOpenGLProgramKey& ProgramKey)
 {
 	if (CachePtr)
 	{
 		FScopeLock Lock(&GProgramBinaryFileCacheCS);
-		return CachePtr->DoesCurrentCacheContain_Internal(ProgramKey);
+		return CachePtr->RequiresCaching_Internal(ProgramKey);
 	}
 	return false;
 }
 
-bool FOpenGLProgramBinaryCache::DoesCurrentCacheContain_Internal(const FOpenGLProgramKey& ProgramKey) const
+bool FOpenGLProgramBinaryCache::RequiresCaching_Internal(const FOpenGLProgramKey& ProgramKey)
 {
-	return ProgramsInCurrentCache.Contains(ProgramKey);
+	return !ProgramsInCurrentCache.Contains(ProgramKey);
 }
 
 FString FOpenGLProgramBinaryCache::GetProgramBinaryCacheFilePath() const
 {
-	check(CurrentShaderPipelineCacheVersionGuid != FGuid());
-	FString ProgramFilename = CachePathRoot / CacheSubDir / CurrentShaderPipelineCacheName;
+	check(CurrentShaderPipelineProperties.CacheVersionGuid != FGuid());
+	FString ProgramFilename = CachePathRoot / CacheSubDir / CurrentShaderPipelineProperties.PipelineCacheName;
 	return ProgramFilename;
 }
 
-void FOpenGLProgramBinaryCache::CheckPendingGLProgramCreateRequests()
+void FOpenGLProgramBinaryCache::TickBinaryCache()
 {
 	if (CachePtr)
 	{
-		QUICK_SCOPE_CYCLE_COUNTER(STAT_OpenGLCheckPendingGLProgramCreateRequests);
-		check(IsInRenderingThread() || IsInRHIThread());
-		FScopeLock Lock(&GPendingGLProgramCreateRequestsCS);
-		CachePtr->CheckPendingGLProgramCreateRequests_internal();
+		if(IsPrecachingEnabled())
+		{
+			CachePtr->UpdatePrecacheMapping();
+		}
+
+		{
+			QUICK_SCOPE_CYCLE_COUNTER(STAT_OpenGLCheckPendingGLProgramCreateRequests);
+			check(IsInRenderingThread() || IsInRHIThread());
+			FScopeLock Lock(&GPendingGLProgramCreateRequestsCS);
+			CachePtr->CheckPendingGLProgramCreateRequests_internal();
+		}
 	}
 }
 
diff --git a/Engine/Source/Runtime/OpenGLDrv/Private/OpenGLProgramBinaryFileCache.h b/Engine/Source/Runtime/OpenGLDrv/Private/OpenGLProgramBinaryFileCache.h
--- a/Engine/Source/Runtime/OpenGLDrv/Private/OpenGLProgramBinaryFileCache.h
+++ b/Engine/Source/Runtime/OpenGLDrv/Private/OpenGLProgramBinaryFileCache.h
@@ -33,15 +33,15 @@ public:
 
 	static bool IsEnabled();
 
-	/** Is the program already in the currently processing cache */
-	static bool DoesCurrentCacheContain(const FOpenGLProgramKey& ProgramKey);
+	/** has the program already been encountered */
+	static bool RequiresCaching(const FOpenGLProgramKey& ProgramKey);
 
 	/** Take an existing program binary and store on disk. Only when ProgramBinaryCache is enabled
 	*/
-	static void CacheProgramBinary(const FOpenGLProgramKey& ProgramKey, const FOpenGLProgramBinary& CachedProgramBinary);
+	static void CacheProgramBinary(const FOpenGLProgramKey& ProgramKey, TUniqueObj<FOpenGLProgramBinary>&& ProgramBinary);
 
 	/** Create any pending GL programs that have come from shader library requests */
-	static void CheckPendingGLProgramCreateRequests();
+	static void TickBinaryCache();
 
 	/** Create any single GL program that have come from shader library requests */
 	static bool CheckSinglePendingGLProgramCreateRequest(const FOpenGLProgramKey& ProgramKey);
@@ -77,11 +77,21 @@ private:
 	/* Subdir is a hash of the device's version info, anything within CachePathRoot that does not == CacheSubDir is deleted. */
 	FString CacheSubDir;
 
-	// Guid of the PSO cache being procesed.
-	FGuid CurrentShaderPipelineCacheVersionGuid;
-
-	// String name of the current cache.
-	FString CurrentShaderPipelineCacheName;
+	struct FCurrentShaderPipelineProperties
+	{
+		// Guid of the PSO cache being procesed.
+		FGuid CacheVersionGuid;
+		// String name of the current cache.
+		FString PipelineCacheName;
+
+		// contains the count of programs when the write cache was last flushed.
+		int32 NumProgramsFlushed = 0;
+		// Last mmapped position
+		int64 LastMappedPosition = 0;
+		// Last number of programs mapped
+		int32 MappedPrograms = 0;
+	};
+	FCurrentShaderPipelineProperties CurrentShaderPipelineProperties;
 
 	/**
 	* Shaders that were requested for compilation
@@ -96,12 +106,15 @@ private:
 
 	void ScanProgramCacheFile();
 
-	bool OpenCacheWriteHandle();
+	bool OpenCacheWriteHandle(const FString& ProgramCacheFilenameToWrite, bool bAppendToExisting);
+
+	// set the file header to the current write position and program count, then flush the cache to storage.
+	bool MarkValidContent(int32 NumPrograms);
 
 	// finalize the current cache file, returns true if a valid file was created.
-	bool CloseCacheWriteHandle(); 
+	bool CloseCacheWriteHandle(int32 NumProgramsAdded); 
 
-	bool DoesCurrentCacheContain_Internal(const FOpenGLProgramKey& ProgramKey) const;
+	bool RequiresCaching_Internal(const FOpenGLProgramKey& ProgramKey);
 
 	void AddProgramBinaryDataToBinaryCache(const FOpenGLProgramKey& ProgramKey, const FOpenGLProgramBinary& BinaryProgramData);
 
@@ -114,8 +127,14 @@ private:
 	void OnShaderPipelineCacheOpened(FString const& Name, EShaderPlatform Platform, uint32 Count, const FGuid& VersionGuid, FShaderPipelineCache::FShaderCachePrecompileContext& ShaderCachePrecompileContext);
 	void OnShaderPipelineCachePrecompilationComplete(uint32 Count, double Seconds, const FShaderPipelineCache::FShaderCachePrecompileContext& ShaderCachePrecompileContext);
 
+	void InitPrecaching();
+	void UpdatePrecacheMapping();
+
 private:
+	TRefCountPtr<FOpenGLProgramBinaryMapping> GetOrAddFileMapping(const FString& ProgramCacheFilename, int32 ProgramCount, int64 Offset, int64 Size);
+
 	FString GetProgramBinaryCacheDir() const { return (CachePathRoot / CacheSubDir);	}
+	bool ReadProgramFile_Internal(uint32 ProgramsToRead, int64 EndOffset, const FString& ProgramCacheFilename, FArchive& Ar);
 
 	FDelegateHandle OnShaderPipelineCacheOpenedDelegate;
 	FDelegateHandle OnShaderPipelineCachePrecompilationCompleteDelegate;
diff --git a/Engine/Source/Runtime/OpenGLDrv/Private/OpenGLShaders.cpp b/Engine/Source/Runtime/OpenGLDrv/Private/OpenGLShaders.cpp
--- a/Engine/Source/Runtime/OpenGLDrv/Private/OpenGLShaders.cpp
+++ b/Engine/Source/Runtime/OpenGLDrv/Private/OpenGLShaders.cpp
@@ -100,6 +100,16 @@ static FAutoConsoleVariableRef CVarProgramLRUResidentCountBeforeEviction(
 	ECVF_RenderThreadSafe
 );
 
+static bool GCacheAllProgramBinaries = true;
+static FAutoConsoleVariableRef CVarGCacheAllProgramBinaries(
+	TEXT("r.OpenGL.CacheAllProgramBinaries"),
+	GCacheAllProgramBinaries,
+	TEXT("Place all encountered program in the binary cache.\n")
+	TEXT("requires r.PSOPrecaching.")
+	,
+	ECVF_RenderThreadSafe
+);
+
 #if PLATFORM_ANDROID
 bool GOpenGLShaderHackLastCompileSuccess = false;
 #endif
@@ -138,6 +148,151 @@ static bool ReportProgramLinkFailures()
 }
 
 
+#define OGL_BINARYCACHE_STATS !UE_BUILD_SHIPPING
+
+#if OGL_BINARYCACHE_STATS
+
+#define OGL_BINARYCACHE_STATS_MARKBEGINCOMPILE(x)		if(FOpenGLBinaryCacheStats::IsEnabled())	{ FOpenGLBinaryCacheStats::Get().MarkStartTime(x);}
+#define OGL_BINARYCACHE_STATS_MARKCOMPILED(x)			if(FOpenGLBinaryCacheStats::IsEnabled())	{ FOpenGLBinaryCacheStats::Get().MarkCompileFinishTime(x);}
+#define OGL_BINARYCACHE_STATS_MARKBINARYCACHEMISS(x,y)	if(FOpenGLBinaryCacheStats::IsEnabled())	{ FOpenGLBinaryCacheStats::Get().MarkCacheMissedTime(x,y);}
+#define OGL_BINARYCACHE_STATS_MARKBINARYCACHEUSE(x)		if(FOpenGLBinaryCacheStats::IsEnabled())	{ FOpenGLBinaryCacheStats::Get().MarkCacheUse(x);}
+#define OGL_BINARYCACHE_STATS_LOG()						if(FOpenGLBinaryCacheStats::IsEnabled())	{ FOpenGLBinaryCacheStats::Get().LogStats();}
+
+class FOpenGLBinaryCacheStats
+{
+public:
+	inline static bool IsEnabled()
+	{
+		static bool bEnabled = FParse::Param(FCommandLine::Get(), TEXT("openglprecachestats"));
+		return bEnabled;
+	}
+
+	inline static FOpenGLBinaryCacheStats& Get() { static FOpenGLBinaryCacheStats CacheStats;  return CacheStats; }
+
+	void MarkStartTime(const FOpenGLProgramKey& ProgramKey)
+	{
+		FScopeLock Lock(&CacheStatsCS);
+		KeyToTimes.FindOrAdd(ProgramKey).StartTime = FPlatformTime::Seconds();
+	}
+
+	void MarkCompileFinishTime(const FOpenGLProgramKey& ProgramKey)
+	{
+		FScopeLock Lock(&CacheStatsCS);
+		FProgramUseTimes& ProgramTime = KeyToTimes.FindOrAdd(ProgramKey);
+		if (!ProgramTime.UsedTime && !ProgramTime.CompileTime)
+		{
+			double CurrentTime = FPlatformTime::Seconds();
+			if (!ProgramTime.StartTime)
+			{
+				ProgramTime.StartTime = CurrentTime;
+				PreloadedBinaries++;
+			}
+
+			ProgramTime.CompileTime = CurrentTime;
+
+			CombinedCompileTime += ProgramTime.CompileTime - ProgramTime.StartTime;
+		}
+	}
+
+	void MarkCacheMissedTime(const FOpenGLProgramKey& ProgramKey, bool bLogOnFirstUse)
+	{
+		FScopeLock Lock(&CacheStatsCS);
+		FProgramUseTimes& ProgramTimes = KeyToTimes.FindOrAdd(ProgramKey);
+		double CurrentTime = FPlatformTime::Seconds();
+		if (!ProgramTimes.UsedTime)
+		{
+			bLogMe = true;
+			ProgramTimes.UsedTime = CurrentTime;
+
+			if (!ProgramTimes.StartTime)
+			{
+				TotalMisses++;
+				UE_CLOG(bLogOnFirstUse, LogRHI, Log, TEXT("BinaryCacheUsage: Program %s was not in the binary cache when first used."), *ProgramKey.ToString());
+			}
+			else if (!ProgramTimes.CompileTime)
+			{
+				TotalEarlyUses++;
+				double TimeToUse = CurrentTime - ProgramTimes.StartTime;
+				CombinedEarlyTimeToUse += TimeToUse;
+				UE_CLOG(bLogOnFirstUse, LogRHI, Log, TEXT("BinaryCacheUsage: Program %s was used too early, binary compile was not ready when first used. Span between compile and use: %f"), *ProgramKey.ToString(), (float)TimeToUse);
+			}
+		}
+	}
+
+	void MarkCacheUse(const FOpenGLProgramKey& ProgramKey)
+	{
+		FScopeLock Lock(&CacheStatsCS);
+
+		FProgramUseTimes& ProgramTime = KeyToTimes.FindChecked(ProgramKey);
+
+		if (!ProgramTime.UsedTime)
+		{
+			check(ProgramTime.StartTime && ProgramTime.CompileTime);
+
+			double CurrentTime = FPlatformTime::Seconds();
+			ProgramTime.UsedTime = CurrentTime;
+			TotalHits++;
+		}
+	}
+
+	void LogStats()
+	{
+		FScopeLock Lock(&CacheStatsCS);
+		const float AvgEarlyTimeToUse = TotalEarlyUses ? (float)(CombinedEarlyTimeToUse / (double)TotalEarlyUses) : 0.0f;
+		const uint32 CompiledBinaries = KeyToTimes.Num() - (TotalEarlyUses + TotalMisses + PreloadedBinaries);
+		const float AvgCompileTime = CompiledBinaries ? (float)(CombinedCompileTime / (double)(CompiledBinaries)) : 0.0f;
+
+		UE_CLOG(bLogMe, LogRHI, Log, TEXT("BinaryCacheUsage: %d programs seen, %d preloaded, %d used in time, %d used before compile finished (avg early miss time span %f), %d programs used were not in the cache. %f avg compile time"),
+			KeyToTimes.Num(),
+			PreloadedBinaries,
+			TotalHits,
+			TotalEarlyUses,
+			AvgEarlyTimeToUse,
+			TotalMisses,
+			AvgCompileTime
+		);
+		bLogMe = false;
+	}
+
+private:
+	FCriticalSection CacheStatsCS;
+
+	struct FProgramUseTimes
+	{
+		double StartTime = 0;
+		double CompileTime = 0;
+		double UsedTime = 0;
+	};
+	TMap< FOpenGLProgramKey, FProgramUseTimes> KeyToTimes;
+
+	uint32 TotalMisses = 0; // Num programs marked as used but were not in the cache.
+	uint32 TotalEarlyUses = 0; // Num program marked as used before their compile had finished.
+	uint32 TotalHits = 0; // Num programs that had compiled in time for their first used.
+	uint32 PreloadedBinaries = 0; // Num programs that came pre-loaded from the binary cache.
+	double CombinedEarlyTimeToUse = 0;
+	double CombinedCompileTime = 0;
+	mutable bool bLogMe = false;
+};
+#else
+#define OGL_BINARYCACHE_STATS_MARKBEGINCOMPILE(x) (void)
+#define OGL_BINARYCACHE_STATS_MARKCOMPILED(x)  (void)
+#define OGL_BINARYCACHE_STATS_MARKBINARYCACHEMISS(x,y) (void)
+#define OGL_BINARYCACHE_STATS_MARKBINARYCACHEUSE(x) (void)
+#define OGL_BINARYCACHE_STATS_LOG() (void)
+
+#endif
+
+bool IsPrecachingEnabled()
+{
+	static const auto CVarPSOPrecaching = IConsoleManager::Get().FindConsoleVariable(TEXT("r.PSOPrecaching"));
+	return CVarPSOPrecaching && (CVarPSOPrecaching->GetInt() != 0);
+}
+
+static bool ShouldCacheAllProgramBinaries()
+{
+	return IsPrecachingEnabled() && GCacheAllProgramBinaries;
+}
+
 static uint32 GCurrentDriverProgramBinaryAllocation = 0;
 static uint32 GNumPrograms = 0;
 
@@ -1375,6 +1530,8 @@ class FGLProgramCacheLRU
 
 			if(bSuccess)
 			{
+				OGL_BINARYCACHE_STATS_MARKBINARYCACHEUSE(ProgramKey);
+
 				// Always keep any mmapped data resident.
 				if(CVarLRUKeepProgramBinaryResident.GetValueOnAnyThread() || !CachedProgramBinary->IsOwned())
 				{
@@ -2631,7 +2788,8 @@ static FOpenGLProgramBinary ExternalProgramCompile(const FOpenGLProgramKey& Prog
 
 void FOpenGLDynamicRHI::PrepareGFXBoundShaderState(const FGraphicsPipelineStateInitializer& Initializer)
 {
-	if (!Initializer.bFromPSOFileCache || !FOpenGLProgramBinaryCache::IsEnabled())
+	const bool bIsPreCachePSO = Initializer.bPSOPrecache || Initializer.bFromPSOFileCache;
+	if (!bIsPreCachePSO || !FOpenGLProgramBinaryCache::IsEnabled())
 	{
 		return;
 	}
@@ -2652,12 +2810,13 @@ void FOpenGLDynamicRHI::PrepareGFXBoundShaderState(const FGraphicsPipelineStateI
 	bool bCreateProgram = false;
 	bool bCreateBinary = false;
 
-	if (!FOpenGLProgramBinaryCache::DoesCurrentCacheContain(ProgramKey))
+	if (FOpenGLProgramBinaryCache::RequiresCaching(ProgramKey))
 	{
 		if (FOpenGLProgramBinaryCache::IsBuildingCache())
 		{
+			OGL_BINARYCACHE_STATS_MARKBEGINCOMPILE(ProgramKey);
 			FOpenGLProgramBinary CompiledProgram;
-			if (CanCreateExternally(Initializer.bFromPSOFileCache))
+			if (CanCreateExternally(bIsPreCachePSO))
 			{
 				CompiledProgram = ExternalProgramCompile(ProgramKey, VertexShaderRHI, PixelShaderRHI);
 			}
@@ -2700,7 +2859,7 @@ void FOpenGLDynamicRHI::PrepareGFXBoundShaderState(const FGraphicsPipelineStateI
 
 			if (CompiledProgram.IsValid())
 			{
-				FOpenGLProgramBinaryCache::CacheProgramBinary(ProgramKey, CompiledProgram);
+				FOpenGLProgramBinaryCache::CacheProgramBinary(ProgramKey, TUniqueObj<FOpenGLProgramBinary>(MoveTemp(CompiledProgram)));
 			}
 			else
 			{
@@ -2808,11 +2967,11 @@ FBoundShaderStateRHIRef FOpenGLDynamicRHI::RHICreateBoundShaderState_OnThisThrea
 			}
 			else
 			{
+				OGL_BINARYCACHE_STATS_MARKBINARYCACHEMISS(Config.ProgramKey, true);
+
 				FOpenGLVertexShader* VertexShader = ResourceCast(VertexShaderRHI);
 				FOpenGLPixelShader* PixelShader = ResourceCast(PixelShaderRHI);
 				FOpenGLGeometryShader* GeometryShader = ResourceCast(GeometryShaderRHI);
-		
-
 
 				// Link program, using the data provided in config
 				LinkedProgram = LinkProgram(Config);
@@ -2839,6 +2998,13 @@ FBoundShaderStateRHIRef FOpenGLDynamicRHI::RHICreateBoundShaderState_OnThisThrea
 				}
 				else
 				{
+					if ( ShouldCacheAllProgramBinaries() && FOpenGLProgramBinaryCache::RequiresCaching(Config.ProgramKey))
+					{
+						// In precache mode we can put any newly compiled programs in the binary cache
+						FOpenGLProgramBinary CompiledProgram = UE::OpenGL::GetProgramBinaryFromGLProgram(LinkedProgram->Program);
+						FOpenGLProgramBinaryCache::CacheProgramBinary(Config.ProgramKey, TUniqueObj<FOpenGLProgramBinary>(MoveTemp(CompiledProgram)));
+					}
+
 					GetOpenGLProgramsCache().Add(Config.ProgramKey, LinkedProgram);
 				}
 			}
@@ -3441,6 +3607,8 @@ namespace UE
 		// Called from the binary file cache when the binary version of a program has been encountered.
 		void OnGLProgramLoadedFromBinaryCache(const FOpenGLProgramKey& ProgramKey, TUniqueObj<FOpenGLProgramBinary>&& ProgramBinaryData)
 		{
+			OGL_BINARYCACHE_STATS_MARKCOMPILED(ProgramKey);
+
 			QUICK_SCOPE_CYCLE_COUNTER(STAT_OpenGLOnGLProgramLoadedFromBinaryCache);
 
 // 			FScopeLock Lock(&GProgramBinaryCacheCS);
@@ -3513,6 +3681,7 @@ static void TickProgramLRU()
 void FOpenGLDynamicRHI::EndFrameTick()
 {
 	TickProgramLRU();
-	FOpenGLProgramBinaryCache::CheckPendingGLProgramCreateRequests();
+	FOpenGLProgramBinaryCache::TickBinaryCache();
 	FTextureEvictionLRU::Get().TickEviction();
+	OGL_BINARYCACHE_STATS_LOG();
 }
