CIS fix - fix shipping buiild.

diff --git a/Engine/Source/Runtime/OpenGLDrv/Private/OpenGLShaders.cpp b/Engine/Source/Runtime/OpenGLDrv/Private/OpenGLShaders.cpp
--- a/Engine/Source/Runtime/OpenGLDrv/Private/OpenGLShaders.cpp
+++ b/Engine/Source/Runtime/OpenGLDrv/Private/OpenGLShaders.cpp
@@ -274,11 +274,11 @@ private:
 	mutable bool bLogMe = false;
 };
 #else
-#define OGL_BINARYCACHE_STATS_MARKBEGINCOMPILE(x) (void)
-#define OGL_BINARYCACHE_STATS_MARKCOMPILED(x)  (void)
-#define OGL_BINARYCACHE_STATS_MARKBINARYCACHEMISS(x,y) (void)
-#define OGL_BINARYCACHE_STATS_MARKBINARYCACHEUSE(x) (void)
-#define OGL_BINARYCACHE_STATS_LOG() (void)
+#define OGL_BINARYCACHE_STATS_MARKBEGINCOMPILE(x)
+#define OGL_BINARYCACHE_STATS_MARKCOMPILED(x)  
+#define OGL_BINARYCACHE_STATS_MARKBINARYCACHEMISS(x,y) 
+#define OGL_BINARYCACHE_STATS_MARKBINARYCACHEUSE(x) 
+#define OGL_BINARYCACHE_STATS_LOG() 
 
 #endif
 
