Fixed assert on Boot Test Android related to creating a buffer with BUF_NullResource flag

diff --git a/Engine/Source/Runtime/OpenGLDrv/Public/OpenGLResources.h b/Engine/Source/Runtime/OpenGLDrv/Public/OpenGLResources.h
--- a/Engine/Source/Runtime/OpenGLDrv/Public/OpenGLResources.h
+++ b/Engine/Source/Runtime/OpenGLDrv/Public/OpenGLResources.h
@@ -176,13 +176,16 @@ public:
 			IncrementBufferMemory(Type, BaseType::GetSize());
 		};
 
-		if (RHICmdList)
+		if (!BufferDesc.IsNull())
 		{
-			RHICmdList->EnqueueLambda([Lambda = MoveTemp(InitLambda)](FRHICommandListBase&) { Lambda(); });
-		}
-		else
-		{
-			InitLambda();
+			if (RHICmdList)
+			{
+				RHICmdList->EnqueueLambda([Lambda = MoveTemp(InitLambda)](FRHICommandListBase&) { Lambda(); });
+			}
+			else
+			{
+				InitLambda();
+			}
 		}
 	}
 
