Allow GRHISupportsMultithreadedShaderCreation on GL RHI, as FOpenGLShader can be created on any thread (since CL# 25447818)

diff --git a/Engine/Source/Runtime/OpenGLDrv/Private/OpenGLDevice.cpp b/Engine/Source/Runtime/OpenGLDrv/Private/OpenGLDevice.cpp
--- a/Engine/Source/Runtime/OpenGLDrv/Private/OpenGLDevice.cpp
+++ b/Engine/Source/Runtime/OpenGLDrv/Private/OpenGLDevice.cpp
@@ -1365,9 +1365,6 @@ void FOpenGLDynamicRHI::Init()
 	check(!GIsRHIInitialized);
 	VERIFY_GL_SCOPE();
 
-	GRHISupportsMultithreadedShaderCreation = false;
-	GRequiredRecursiveShaders = ERecursiveShader::Resolve | ERecursiveShader::Clear | ERecursiveShader::Null;
-
 	FOpenGLProgramBinaryCache::Initialize();
 
 	InitializeStateResources();
