Android GLES binary program cache - Add some additional logging and testing for errors while producing the binary cache output archive.

diff --git a/Engine/Source/Runtime/OpenGLDrv/Private/OpenGLProgramBinaryFileCache.cpp b/Engine/Source/Runtime/OpenGLDrv/Private/OpenGLProgramBinaryFileCache.cpp
--- a/Engine/Source/Runtime/OpenGLDrv/Private/OpenGLProgramBinaryFileCache.cpp
+++ b/Engine/Source/Runtime/OpenGLDrv/Private/OpenGLProgramBinaryFileCache.cpp
@@ -818,7 +818,7 @@ bool FOpenGLProgramBinaryCache::OpenCacheWriteHandle(const FString& ProgramCache
 	BinaryCacheWriteFileHandle = IFileManager::Get().CreateFileWriter(*ProgramCacheFilenameToWrite, EFileWrite::FILEWRITE_AllowRead | (bAppendToExisting ? EFileWrite::FILEWRITE_Append : EFileWrite::FILEWRITE_None));
 	UE_CLOG(BinaryCacheWriteFileHandle, LogRHI, Log, TEXT("Opened binary cache for write (%s)"), *ProgramCacheFilenameToWrite);
 	UE_CLOG(BinaryCacheWriteFileHandle == nullptr, LogRHI, Warning, TEXT("Failed to open OGL binary cache output file. (%s)"), *ProgramCacheFilenameToWrite);
-
+	UE_CLOG(BinaryCacheWriteFileHandle->IsError() || BinaryCacheWriteFileHandle->IsCriticalError(), LogRHI, Error, TEXT("OGL binary cache output archive error (%s, %d,%d)"), *ProgramCacheFilenameToWrite, BinaryCacheWriteFileHandle->IsError(), BinaryCacheWriteFileHandle->IsCriticalError());
 	return BinaryCacheWriteFileHandle != nullptr;
 }
 
@@ -829,10 +829,10 @@ bool FOpenGLProgramBinaryCache::MarkValidContent(int32 NumPrograms)
 	check(BinaryCacheWriteFileHandle != nullptr);
 	const bool bCacheFileIsEmpty = NumPrograms == 0;
 
-	bool bArchiveFailed = BinaryCacheWriteFileHandle->IsError() || BinaryCacheWriteFileHandle->IsCriticalError();
+	uint32 FailCode = BinaryCacheWriteFileHandle->IsError() || BinaryCacheWriteFileHandle->IsCriticalError() ? 1 : 0;
 
 	// Overwrite the header with the final program count. This indicates a successful write.
-	if (!bArchiveFailed)
+	if (FailCode == 0)
 	{
 		FArchive& Ar = *BinaryCacheWriteFileHandle;
 		int64 CurrPos = Ar.Tell();
@@ -841,16 +841,17 @@ bool FOpenGLProgramBinaryCache::MarkValidContent(int32 NumPrograms)
 		UE_LOG(LogRHI, Verbose, TEXT("MarkValidContent, file valid at %d Programs, %d bytes."), NumPrograms, CurrPos);
 		Ar << OutHeader;
 		Ar.Seek(CurrPos);
-		bArchiveFailed = BinaryCacheWriteFileHandle->IsError() || BinaryCacheWriteFileHandle->IsCriticalError();
+		Ar.Flush();
+
+		FailCode = BinaryCacheWriteFileHandle->IsError() || BinaryCacheWriteFileHandle->IsCriticalError() ? 2 : 0;
 	}
 
-	if (bArchiveFailed)
+	if (FailCode != 0)
 	{
 		RHIGetPanicDelegate().ExecuteIfBound(FName("FailedBinaryProgramArchiveWrite"));
-		UE_LOG(LogRHI, Fatal, TEXT("MarkValidContent - FArchive error bit set, failed to write binary cache."));
+		UE_LOG(LogRHI, Fatal, TEXT("MarkValidContent - FArchive error bit set, failed to write binary cache. %d, %d"), NumPrograms, FailCode);
 	}
 	CurrentShaderPipelineProperties.NumProgramsFlushed = NumPrograms;
-	BinaryCacheWriteFileHandle->Flush();
 
 	return !bCacheFileIsEmpty;
 }
@@ -860,16 +861,16 @@ bool FOpenGLProgramBinaryCache::CloseCacheWriteHandle(int32 NumProgramsAdded)
 	check(BinaryCacheWriteFileHandle != nullptr);
 	const bool bCacheFileIsEmpty = NumProgramsAdded == 0;
 
-	bool bArchiveFailed = BinaryCacheWriteFileHandle->IsError() || BinaryCacheWriteFileHandle->IsCriticalError();
+	uint32 FailCode = BinaryCacheWriteFileHandle->IsError() || BinaryCacheWriteFileHandle->IsCriticalError() ? 1 : 0;
 
 	// Overwrite the header with the final program count. This indicates a successful write.
-	if(!bArchiveFailed)
+	if(FailCode == 0)
 	{
 		FArchive& Ar = *BinaryCacheWriteFileHandle;
 		Ar.Seek(0);
 		UE::OpenGL::FBinaryCacheFileHeader OutHeader = UE::OpenGL::FBinaryCacheFileHeader::CreateHeader(CurrentShaderPipelineProperties.CacheVersionGuid, NumProgramsAdded, (int32)Ar.TotalSize());
 		Ar << OutHeader;
-		bArchiveFailed = BinaryCacheWriteFileHandle->IsError() || BinaryCacheWriteFileHandle->IsCriticalError();
+		FailCode = BinaryCacheWriteFileHandle->IsError() || BinaryCacheWriteFileHandle->IsCriticalError() ? 2 : 0;
 	}
 
 	BinaryCacheWriteFileHandle->Close();
@@ -879,10 +880,10 @@ bool FOpenGLProgramBinaryCache::CloseCacheWriteHandle(int32 NumProgramsAdded)
 	const FString ProgramCacheFilename = GetProgramBinaryCacheFilePath();
 	IPlatformFile& PlatformFile = FPlatformFileManager::Get().GetPlatformFile();
 
-	if (bArchiveFailed)
+	if (FailCode != 0)
 	{
 		RHIGetPanicDelegate().ExecuteIfBound(FName("FailedBinaryProgramArchiveWrite"));
-		UE_LOG(LogRHI, Fatal, TEXT("CloseCacheWriteHandle - FArchive error bit set, failed to write binary cache."));
+		UE_LOG(LogRHI, Fatal, TEXT("CloseCacheWriteHandle - FArchive error bit set, failed to write binary cache. %d"), FailCode);
 	}
 
 	if (bCacheFileIsEmpty)
@@ -919,6 +920,7 @@ void FOpenGLProgramBinaryCache::AddProgramBinaryDataToBinaryCache(const FOpenGLP
 	check(IsBuildingCache_internal());
 	check(BinaryProgramData.IsValid());
 	FArchive& Ar = *BinaryCacheWriteFileHandle;
+	bool bInitiallyValid = !(Ar.IsError() || Ar.IsCriticalError());
 	// Serialize to output file:
 	FOpenGLProgramKey SerializedProgramKey = ProgramKey;
 	const TArrayView<const uint8> BinaryProgramDataView = BinaryProgramData.GetDataView();
@@ -933,6 +935,7 @@ void FOpenGLProgramBinaryCache::AddProgramBinaryDataToBinaryCache(const FOpenGLP
 	Ar.Serialize(const_cast<uint8*>(ProgramBinaryBytes), ProgramBinarySize);
 	uint32 End = Ar.Tell();
 
+	UE_CLOG(Ar.IsError() || Ar.IsCriticalError(), LogRHI, Error, TEXT("AddProgramBinaryDataToBinaryCache : archive failed (%d, %d, %d, %d, %d, %d)"), Ar.IsError(), Ar.IsCriticalError(), Start, End, ProgramBinarySize, bInitiallyValid);
 	UE_LOG(LogRHI, VeryVerbose, TEXT("AddProgramBinaryDataToBinaryCache : added %d bytes to cache"), End-Start);
 
 	if (UE::OpenGL::AreBinaryProgramsCompressed())
