Fix Static analysis issue.

diff --git a/Engine/Source/Runtime/OpenGLDrv/Private/OpenGLProgramBinaryFileCache.cpp b/Engine/Source/Runtime/OpenGLDrv/Private/OpenGLProgramBinaryFileCache.cpp
--- a/Engine/Source/Runtime/OpenGLDrv/Private/OpenGLProgramBinaryFileCache.cpp
+++ b/Engine/Source/Runtime/OpenGLDrv/Private/OpenGLProgramBinaryFileCache.cpp
@@ -818,7 +818,7 @@ bool FOpenGLProgramBinaryCache::OpenCacheWriteHandle(const FString& ProgramCache
 	BinaryCacheWriteFileHandle = IFileManager::Get().CreateFileWriter(*ProgramCacheFilenameToWrite, EFileWrite::FILEWRITE_AllowRead | (bAppendToExisting ? EFileWrite::FILEWRITE_Append : EFileWrite::FILEWRITE_None));
 	UE_CLOG(BinaryCacheWriteFileHandle, LogRHI, Log, TEXT("Opened binary cache for write (%s)"), *ProgramCacheFilenameToWrite);
 	UE_CLOG(BinaryCacheWriteFileHandle == nullptr, LogRHI, Warning, TEXT("Failed to open OGL binary cache output file. (%s)"), *ProgramCacheFilenameToWrite);
-	UE_CLOG(BinaryCacheWriteFileHandle->IsError() || BinaryCacheWriteFileHandle->IsCriticalError(), LogRHI, Error, TEXT("OGL binary cache output archive error (%s, %d,%d)"), *ProgramCacheFilenameToWrite, BinaryCacheWriteFileHandle->IsError(), BinaryCacheWriteFileHandle->IsCriticalError());
+	UE_CLOG(BinaryCacheWriteFileHandle && (BinaryCacheWriteFileHandle->IsError() || BinaryCacheWriteFileHandle->IsCriticalError()), LogRHI, Error, TEXT("OGL binary cache output archive error (%s, %d,%d)"), *ProgramCacheFilenameToWrite, BinaryCacheWriteFileHandle->IsError(), BinaryCacheWriteFileHandle->IsCriticalError());
 	return BinaryCacheWriteFileHandle != nullptr;
 }
 
