Moving GPUFamily calculation outside of general FAndroidGPUInfo to delay it later in the init after Zen client had a chance to establish the connection.

diff --git a/Engine/Source/Runtime/OpenGLDrv/Private/Android/AndroidOpenGL.cpp b/Engine/Source/Runtime/OpenGLDrv/Private/Android/AndroidOpenGL.cpp
--- a/Engine/Source/Runtime/OpenGLDrv/Private/Android/AndroidOpenGL.cpp
+++ b/Engine/Source/Runtime/OpenGLDrv/Private/Android/AndroidOpenGL.cpp
@@ -1064,7 +1064,7 @@ void FAndroidOpenGL::ProcessExtensions(const FString& ExtensionsString)
 
 FString FAndroidMisc::GetGPUFamily()
 {
-	return FAndroidGPUInfo::Get().GPUFamily;
+	return FAndroidGPUInfo::Get().GetGPUFamily();
 }
 
 FString FAndroidMisc::GetGLVersion()
diff --git a/Engine/Source/Runtime/OpenGLDrv/Private/Android/AndroidOpenGLPrivate.h b/Engine/Source/Runtime/OpenGLDrv/Private/Android/AndroidOpenGLPrivate.h
--- a/Engine/Source/Runtime/OpenGLDrv/Private/Android/AndroidOpenGLPrivate.h
+++ b/Engine/Source/Runtime/OpenGLDrv/Private/Android/AndroidOpenGLPrivate.h
@@ -26,7 +26,6 @@ public:
 		return This;
 	}
 
-	FString GPUFamily;
 	FString GLVersion;
 	FString VendorName;
 	bool bSupportsFloatingPointRenderTargets;
@@ -38,7 +37,17 @@ public:
 		TargetPlatformNames.Remove(PlatformName);
 	}
 
+	// computing GPU family needs regex access, which might not be available early in init
+	FString& GetGPUFamily()
+	{
+		if (GPUFamily.IsEmpty())
+			ReadGPUFamily();
+		return GPUFamily;
+	}
+
 private:
+	FString GPUFamily;
+
 	FAndroidGPUInfo()
 	{
 		// this is only valid in the game thread, make sure we are initialized there before being called on other threads!
@@ -64,6 +73,38 @@ private:
 		const ANSICHAR* GlGetStringOutput = (const ANSICHAR*) glGetString(GL_EXTENSIONS);
 		FString ExtensionsString = GlGetStringOutput;
 
+		GLVersion = (const ANSICHAR*)glGetString(GL_VERSION);
+
+		// highest priority is the per-texture version
+		if (ExtensionsString.Contains(TEXT("GL_KHR_texture_compression_astc_ldr")))
+		{
+			TargetPlatformNames.Add(TEXT("Android_ASTC"));
+		}
+		if (ExtensionsString.Contains(TEXT("GL_NV_texture_compression_s3tc")) || ExtensionsString.Contains(TEXT("GL_EXT_texture_compression_s3tc")))
+		{
+			TargetPlatformNames.Add(TEXT("Android_DXT"));
+		}
+		
+		TargetPlatformNames.Add(TEXT("Android_ETC2"));
+
+		// finally, generic Android
+		TargetPlatformNames.Add(TEXT("Android"));
+
+		bSupportsFloatingPointRenderTargets = 
+			ExtensionsString.Contains(TEXT("GL_EXT_color_buffer_half_float")) 
+			// According to https://www.khronos.org/registry/gles/extensions/EXT/EXT_color_buffer_float.txt
+			|| (ExtensionsString.Contains(TEXT("GL_EXT_color_buffer_float")));
+
+		bSupportsFrameBufferFetch = ExtensionsString.Contains(TEXT("GL_EXT_shader_framebuffer_fetch")) || ExtensionsString.Contains(TEXT("GL_NV_shader_framebuffer_fetch")) 
+			|| ExtensionsString.Contains(TEXT("GL_ARM_shader_framebuffer_fetch ")); // has space at the end to exclude GL_ARM_shader_framebuffer_fetch_depth_stencil match
+
+		GAndroidGPUInfoReady = true;
+
+		VendorName = FString(ANSI_TO_TCHAR((const ANSICHAR*)glGetString(GL_VENDOR)));
+	}
+
+	void ReadGPUFamily()
+	{
 		GPUFamily = (const ANSICHAR*)glGetString(GL_RENDERER);
 		check(!GPUFamily.IsEmpty());
 
@@ -93,34 +134,5 @@ private:
 				}
 			}
 		}
-
-		GLVersion = (const ANSICHAR*)glGetString(GL_VERSION);
-
-		// highest priority is the per-texture version
-		if (ExtensionsString.Contains(TEXT("GL_KHR_texture_compression_astc_ldr")))
-		{
-			TargetPlatformNames.Add(TEXT("Android_ASTC"));
-		}
-		if (ExtensionsString.Contains(TEXT("GL_NV_texture_compression_s3tc")) || ExtensionsString.Contains(TEXT("GL_EXT_texture_compression_s3tc")))
-		{
-			TargetPlatformNames.Add(TEXT("Android_DXT"));
-		}
-		
-		TargetPlatformNames.Add(TEXT("Android_ETC2"));
-
-		// finally, generic Android
-		TargetPlatformNames.Add(TEXT("Android"));
-
-		bSupportsFloatingPointRenderTargets = 
-			ExtensionsString.Contains(TEXT("GL_EXT_color_buffer_half_float")) 
-			// According to https://www.khronos.org/registry/gles/extensions/EXT/EXT_color_buffer_float.txt
-			|| (ExtensionsString.Contains(TEXT("GL_EXT_color_buffer_float")));
-
-		bSupportsFrameBufferFetch = ExtensionsString.Contains(TEXT("GL_EXT_shader_framebuffer_fetch")) || ExtensionsString.Contains(TEXT("GL_NV_shader_framebuffer_fetch")) 
-			|| ExtensionsString.Contains(TEXT("GL_ARM_shader_framebuffer_fetch ")); // has space at the end to exclude GL_ARM_shader_framebuffer_fetch_depth_stencil match
-
-		GAndroidGPUInfoReady = true;
-
-		VendorName = FString(ANSI_TO_TCHAR((const ANSICHAR*)glGetString(GL_VENDOR)));
 	}
 };
