Added Shader hash logging to OpenGL UAV vertex shader usage check.

diff --git a/Engine/Source/Runtime/OpenGLDrv/Private/OpenGLShaders.cpp b/Engine/Source/Runtime/OpenGLDrv/Private/OpenGLShaders.cpp
--- a/Engine/Source/Runtime/OpenGLDrv/Private/OpenGLShaders.cpp
+++ b/Engine/Source/Runtime/OpenGLDrv/Private/OpenGLShaders.cpp
@@ -2184,7 +2184,7 @@ void FOpenGLLinkedProgram::ConfigureShaderStage( int Stage, uint32 FirstUniformB
 			// glUniform1i(Location, FirstUAVUnit[Stage] + UAVIndex);
 
 			// verify that only CS and PS uses UAVs (limitation on MALI GPUs)
-			check(Stage == CrossCompiler::SHADER_STAGE_COMPUTE || Stage == CrossCompiler::SHADER_STAGE_PIXEL);
+			checkf(Stage == CrossCompiler::SHADER_STAGE_COMPUTE || Stage == CrossCompiler::SHADER_STAGE_PIXEL, TEXT("%s uses UAV in vertex shader"), *Config.ProgramKey.ToString());
 			
 			UAVStageNeeds[ FirstUAVUnit[Stage] + UAVIndex ] = true;
 			MaxUAVUnitUsed = FMath::Max(MaxUAVUnitUsed, FirstUAVUnit[Stage] + UAVIndex);
