Fix GLES not deleting old binary program cache files when precaching is enabled.

diff --git a/Engine/Source/Runtime/OpenGLDrv/Private/OpenGLProgramBinaryFileCache.cpp b/Engine/Source/Runtime/OpenGLDrv/Private/OpenGLProgramBinaryFileCache.cpp
--- a/Engine/Source/Runtime/OpenGLDrv/Private/OpenGLProgramBinaryFileCache.cpp
+++ b/Engine/Source/Runtime/OpenGLDrv/Private/OpenGLProgramBinaryFileCache.cpp
@@ -471,7 +471,7 @@ void FOpenGLProgramBinaryCache::InitPrecaching()
 		if (bDeleteMe)
 		{
 			UE_LOG(LogRHI, Log, TEXT("Deleting binary cache file %s (%d bytes)"), *CacheFileName, CacheStatData.FileSize);
-			PlatformFile.DeleteFile(*CacheName);
+			PlatformFile.DeleteFile(*CacheFileName);
 		}
 	}
 
