Fixed issues with GL uniform buffer views and RHIT

diff --git a/Engine/Source/Runtime/OpenGLDrv/Private/OpenGLUniformBuffer.cpp b/Engine/Source/Runtime/OpenGLDrv/Private/OpenGLUniformBuffer.cpp
--- a/Engine/Source/Runtime/OpenGLDrv/Private/OpenGLUniformBuffer.cpp
+++ b/Engine/Source/Runtime/OpenGLDrv/Private/OpenGLUniformBuffer.cpp
@@ -417,7 +417,7 @@ FOpenGLUniformBuffer::FOpenGLUniformBuffer(const FRHIUniformBufferLayout* InLayo
 	, UniqueID(UniqueUniformBufferID())
 	, AllocatedSize(0)
 	, bStreamDraw(false)
-	, bOwnsResource(false)
+	, bOwnsResource(true)
 {
 	bIsEmulatedUniformBuffer = GUseEmulatedUniformBuffers && !InLayout->bNoEmulatedUniformBuffer;
 	RangeSize = InLayout->ConstantBufferSize;
@@ -625,7 +625,7 @@ static FUniformBufferRHIRef CreateUniformBuffer(const void* Contents, const FRHI
 	return NewUniformBuffer;
 }
 
-static FOpenGLUniformBuffer* CreateUniformBufferView(const FRHIUniformBufferLayout* Layout, const void* Contents)
+static FOpenGLUniformBuffer* CreateUniformBufferView(FRHICommandListImmediate& RHICmdList, const FRHIUniformBufferLayout* Layout, const void* Contents)
 {
 	FOpenGLUniformBuffer* UniformBufferView = nullptr;
 	
@@ -633,24 +633,29 @@ static FOpenGLUniformBuffer* CreateUniformBufferView(const FRHIUniformBufferLayo
 		Layout->Resources[0].MemberType == UBMT_RDG_UNIFORM_BLOCK_SRV)
 	{
 		UniformBufferView = new FOpenGLUniformBuffer(Layout);
+		UniformBufferView->SetLayoutTable(Contents, EUniformBufferValidation::None);
 		
 		FRHIShaderResourceView* SRV = (FRHIShaderResourceView*)GetShaderParameterResourceRHI(Contents, Layout->Resources[0].MemberOffset, UBMT_RDG_UNIFORM_BLOCK_SRV);
-		FOpenGLBuffer* UBO = FOpenGLDynamicRHI::ResourceCast(SRV->GetBuffer());
-		const FRHIViewDesc::FBufferSRV& SRVInfo = SRV->GetDesc().Buffer.SRV;
-		
-		check(UBO->GetSize() >= PLATFORM_MAX_UNIFORM_BUFFER_RANGE);
 
-		UniformBufferView->Resource = UBO->Resource;
-		UniformBufferView->AllocatedSize = UBO->GetSize();
-		UniformBufferView->bOwnsResource = false;
-		UniformBufferView->Offset = SRVInfo.OffsetInBytes;
-		UniformBufferView->RangeSize = PLATFORM_MAX_UNIFORM_BUFFER_RANGE;
-		
-		UniformBufferView->PersistentlyMappedBuffer = nullptr;
-		UniformBufferView->EmulatedBufferData = nullptr;
-		UniformBufferView->bStreamDraw = false;
-
-		UniformBufferView->SetLayoutTable(Contents, EUniformBufferValidation::None);
+		RHICmdList.EnqueueLambda([UniformBufferView, SRV](FRHICommandListImmediate&)
+		{
+			VERIFY_GL_SCOPE();
+			
+			FOpenGLBuffer* UBO = FOpenGLDynamicRHI::ResourceCast(SRV->GetBuffer());
+			const FRHIViewDesc::FBufferSRV& SRVInfo = SRV->GetDesc().Buffer.SRV;
+			
+			check(UBO->Resource);
+			check(UBO->GetSize() >= PLATFORM_MAX_UNIFORM_BUFFER_RANGE);
+
+			UniformBufferView->Resource = UBO->Resource;
+			UniformBufferView->AllocatedSize = UBO->GetSize();
+			UniformBufferView->bOwnsResource = false;
+			UniformBufferView->Offset = SRVInfo.OffsetInBytes;
+			UniformBufferView->RangeSize = PLATFORM_MAX_UNIFORM_BUFFER_RANGE;
+			UniformBufferView->PersistentlyMappedBuffer = nullptr;
+			UniformBufferView->EmulatedBufferData = nullptr;
+			UniformBufferView->bStreamDraw = false;
+		});
 	}
 	
 	return UniformBufferView;
@@ -673,7 +678,7 @@ FUniformBufferRHIRef FOpenGLDynamicRHI::RHICreateUniformBuffer(const void* Conte
 	}
 
 	// 
-	FOpenGLUniformBuffer* UniformBufferView = CreateUniformBufferView(Layout, Contents);
+	FOpenGLUniformBuffer* UniformBufferView = CreateUniformBufferView(RHICmdList, Layout, Contents);
 	if (UniformBufferView)
 	{
 		return UniformBufferView;
