Mali read-only driver bug workaround

diff --git a/Engine/Source/Runtime/OpenGLDrv/Private/Android/AndroidOpenGL.cpp b/Engine/Source/Runtime/OpenGLDrv/Private/Android/AndroidOpenGL.cpp
--- a/Engine/Source/Runtime/OpenGLDrv/Private/Android/AndroidOpenGL.cpp
+++ b/Engine/Source/Runtime/OpenGLDrv/Private/Android/AndroidOpenGL.cpp
@@ -901,6 +901,11 @@ void FAndroidOpenGL::ProcessExtensions(const FString& ExtensionsString)
 		UE_CLOG(bRequiresAdrenoTilingHint, LogRHI, Log, TEXT("Enabling Adreno tiling hint."));
 	}
 
+	if (bIsMaliBased)
+	{
+		bRequiresReadOnlyBuffersWorkaround = VersionString.Contains("r9p");
+	}
+
 	// Disable ASTC if requested by device profile
 	static const auto CVarDisableASTC = IConsoleManager::Get().FindTConsoleVariableDataInt(TEXT("r.Android.DisableASTCSupport"));
 	if (bSupportsASTC && CVarDisableASTC->GetValueOnAnyThread())
diff --git a/Engine/Source/Runtime/OpenGLDrv/Private/OpenGLES.cpp b/Engine/Source/Runtime/OpenGLDrv/Private/OpenGLES.cpp
--- a/Engine/Source/Runtime/OpenGLDrv/Private/OpenGLES.cpp
+++ b/Engine/Source/Runtime/OpenGLDrv/Private/OpenGLES.cpp
@@ -122,6 +122,9 @@ bool FOpenGLES::bBinaryProgramRetrievalFailed = false;
 /* Some Mali devices do not work correctly with early_fragment_test enabled */
 bool FOpenGLES::bRequiresDisabledEarlyFragmentTests = false;
 
+/* This is a workaround for a Mali bug where read-only buffers do not work when passed to functions*/
+bool FOpenGLES::bRequiresReadOnlyBuffersWorkaround = false;
+
 /* This is to avoid a bug in Adreno drivers that define GL_ARM_shader_framebuffer_fetch_depth_stencil even when device does not support this extension  */
 bool FOpenGLES::bRequiresARMShaderFramebufferFetchDepthStencilUndef = false;
 
diff --git a/Engine/Source/Runtime/OpenGLDrv/Private/OpenGLShaders.cpp b/Engine/Source/Runtime/OpenGLDrv/Private/OpenGLShaders.cpp
--- a/Engine/Source/Runtime/OpenGLDrv/Private/OpenGLShaders.cpp
+++ b/Engine/Source/Runtime/OpenGLDrv/Private/OpenGLShaders.cpp
@@ -774,12 +774,14 @@ void OPENGLDRV_API GetCurrentOpenGLShaderDeviceCapabilities(FOpenGLShaderDeviceC
 		Capabilities.TargetPlatform = EOpenGLShaderTargetPlatform::OGLSTP_Android;
 		Capabilities.bSupportsShaderFramebufferFetch = FOpenGL::SupportsShaderFramebufferFetch();
 		Capabilities.bRequiresARMShaderFramebufferFetchDepthStencilUndef = false;
+		Capabilities.bRequiresReadOnlyBuffersWorkaround = false;
 		Capabilities.MaxVaryingVectors = FOpenGL::GetMaxVaryingVectors();
 	}
 
 #elif PLATFORM_ANDROID
 		Capabilities.TargetPlatform = EOpenGLShaderTargetPlatform::OGLSTP_Android;
 		Capabilities.bSupportsShaderFramebufferFetch = FOpenGL::SupportsShaderFramebufferFetch();
+		Capabilities.bRequiresReadOnlyBuffersWorkaround = FOpenGL::RequiresReadOnlyBuffersWorkaround();
 		Capabilities.bRequiresARMShaderFramebufferFetchDepthStencilUndef = FOpenGL::RequiresARMShaderFramebufferFetchDepthStencilUndef();
 		Capabilities.MaxVaryingVectors = FOpenGL::GetMaxVaryingVectors();
 		Capabilities.bRequiresDisabledEarlyFragmentTests = FOpenGL::RequiresDisabledEarlyFragmentTests();
@@ -828,6 +830,11 @@ void OPENGLDRV_API GLSLToDeviceCompatibleGLSL(FAnsiCharArray& GlslCodeOriginal,
 		AppendCString(GlslCode, "#define fma(A, B, C) ((A) * (B) + (C))\n");
 	}
 
+	if (Capabilities.bRequiresReadOnlyBuffersWorkaround)
+	{
+		ReplaceCString(GlslCodeOriginal, "readonly buffer", "buffer");
+	}
+
 	if (TypeEnum == GL_FRAGMENT_SHADER && Capabilities.bRequiresDisabledEarlyFragmentTests)
 	{
 		ReplaceCString(GlslCodeOriginal, "layout(early_fragment_tests) in;", "");
diff --git a/Engine/Source/Runtime/OpenGLDrv/Public/OpenGLES.h b/Engine/Source/Runtime/OpenGLDrv/Public/OpenGLES.h
--- a/Engine/Source/Runtime/OpenGLDrv/Public/OpenGLES.h
+++ b/Engine/Source/Runtime/OpenGLDrv/Public/OpenGLES.h
@@ -231,6 +231,7 @@ struct FOpenGLES : public FOpenGLBase
 
 	static FORCEINLINE bool HasBinaryProgramRetrievalFailed() { return bBinaryProgramRetrievalFailed; }
 	static FORCEINLINE bool RequiresDisabledEarlyFragmentTests() { return bRequiresDisabledEarlyFragmentTests; }
+	static FORCEINLINE bool RequiresReadOnlyBuffersWorkaround() { return bRequiresReadOnlyBuffersWorkaround; }
 	static FORCEINLINE bool RequiresARMShaderFramebufferFetchDepthStencilUndef() { return bRequiresARMShaderFramebufferFetchDepthStencilUndef; }
 
 	// Adreno doesn't support HALF_FLOAT
@@ -858,6 +859,9 @@ public:
 	/* Some Mali devices do not work correctly with early_fragment_test enabled */
 	static bool bRequiresDisabledEarlyFragmentTests;
 		
+	/* This is a workaround for a Mali bug where read-only buffers do not work when passed to functions*/
+	static bool bRequiresReadOnlyBuffersWorkaround;
+
 	/* This is to avoid a bug in Adreno drivers that define GL_ARM_shader_framebuffer_fetch_depth_stencil even when device does not support this extension  */
 	static bool bRequiresARMShaderFramebufferFetchDepthStencilUndef;
 
diff --git a/Engine/Source/Runtime/OpenGLDrv/Public/OpenGLShaders.h b/Engine/Source/Runtime/OpenGLDrv/Public/OpenGLShaders.h
--- a/Engine/Source/Runtime/OpenGLDrv/Public/OpenGLShaders.h
+++ b/Engine/Source/Runtime/OpenGLDrv/Public/OpenGLShaders.h
@@ -28,6 +28,7 @@ struct FOpenGLShaderDeviceCapabilities
 	bool bRequiresARMShaderFramebufferFetchDepthStencilUndef;
 	GLint MaxVaryingVectors;
 	bool bRequiresDisabledEarlyFragmentTests;
+	bool bRequiresReadOnlyBuffersWorkaround;
 };
 
 /**
