Enable pso precaching with GLES.
GLES RHI disables the PSO filecache if precaching is active. (provided r.OpenGL.EnablePSOFileCacheWhenPrecachingActive cvar should you want the 2 systems active)

diff --git a/Engine/Source/Runtime/OpenGLDrv/Private/OpenGLDevice.cpp b/Engine/Source/Runtime/OpenGLDrv/Private/OpenGLDevice.cpp
--- a/Engine/Source/Runtime/OpenGLDrv/Private/OpenGLDevice.cpp
+++ b/Engine/Source/Runtime/OpenGLDrv/Private/OpenGLDevice.cpp
@@ -71,6 +71,15 @@ static TAutoConsoleVariable<int32> CVarGLDepth24Bit(
 	TEXT("0: Use 32-bit float depth buffer \n1: Use 24-bit fixed point depth buffer (default)\n"),
 	ECVF_RenderThreadSafe | ECVF_ReadOnly);
 
+// If precaching is active we should not need the file cache.
+// however, precaching and filecache are compatible with each other, there maybe some scenarios in which both could be used.
+static TAutoConsoleVariable<bool> CVarEnablePSOFileCacheWhenPrecachingActive(
+	TEXT("r.OpenGL.EnablePSOFileCacheWhenPrecachingActive"),
+	false,
+	TEXT("false: If precaching is active (r.PSOPrecaching=1) then disable the PSO filecache. (default)\n")
+	TEXT("true: GL RHI Allows both PSO file cache and precaching."),
+	ECVF_RenderThreadSafe | ECVF_ReadOnly);
+
 void OnQueryCreation( FOpenGLRenderQuery* Query )
 {
 	check(PrivateOpenGLDevicePtr);
@@ -1009,7 +1018,7 @@ static void InitRHICapabilitiesForGL()
 
 	// OpenGL ES does not support glTextureView
 	GRHISupportsTextureViews = false;
-	
+
 	// By default use emulated UBs on mobile
 	GUseEmulatedUniformBuffers = IsUsingEmulatedUniformBuffers(GMaxRHIShaderPlatform);
 
@@ -1236,7 +1245,13 @@ static void InitRHICapabilitiesForGL()
 	// @TODO revisit this with newer drivers
 	GRHINeedsUnatlasedCSMDepthsWorkaround = true;
 
-	GRHISupportsPipelineFileCache = true;
+	static const auto CVarPSOPrecaching = IConsoleManager::Get().FindConsoleVariable(TEXT("r.PSOPrecaching"));
+	if (CVarPSOPrecaching && CVarPSOPrecaching->GetInt() != 0)
+	{
+		GRHISupportsPSOPrecaching = true;
+	}
+	
+	GRHISupportsPipelineFileCache = !GRHISupportsPSOPrecaching || CVarEnablePSOFileCacheWhenPrecachingActive.GetValueOnAnyThread();
 }
 
 FDynamicRHI* FOpenGLDynamicRHIModule::CreateRHI(ERHIFeatureLevel::Type InRequestedFeatureLevel)
diff --git a/Engine/Source/Runtime/VulkanRHI/Private/VulkanRHI.cpp b/Engine/Source/Runtime/VulkanRHI/Private/VulkanRHI.cpp
--- a/Engine/Source/Runtime/VulkanRHI/Private/VulkanRHI.cpp
+++ b/Engine/Source/Runtime/VulkanRHI/Private/VulkanRHI.cpp
@@ -487,13 +487,8 @@ FVulkanDynamicRHI::FVulkanDynamicRHI()
 	GRHITransitionPrivateData_AlignInBytes = alignof(FVulkanPipelineBarrier);
 	GConfig->GetInt(TEXT("TextureStreaming"), TEXT("PoolSizeVRAMPercentage"), GPoolSizeVRAMPercentage, GEngineIni);
 
-	// VulkanRHI does not yet support PSO precaching, we disable the cvar here.
-	static const auto CVarPSOPrecaching = IConsoleManager::Get().FindConsoleVariable(TEXT("r.PSOPrecaching"));
-	if (CVarPSOPrecaching)
-	{
-		UE_CLOG(CVarPSOPrecaching->GetInt()!=0, LogVulkanRHI, Display, TEXT("r.PSOPrecaching was enabled but is not supported by VulkanRHI. setting r.PSOPrecaching=0"));
-		CVarPSOPrecaching->Set(false);
-	}
+	// VulkanRHI does not yet support PSO precaching, we ensure it is disabled.
+	GRHISupportsPSOPrecaching = false;
 
 	// Copy source requires its own image layout.
 	EnumRemoveFlags(GRHIMergeableAccessMask, ERHIAccess::CopySrc);
