RHIUpdateTexture2D Fixes
- Fixing implementations to handle source offsets.
- Fixing implementations to handle source data smaller than the image's pitch (very small regions).
- Fixing implementations to handle RHI bypass correctly.

diff --git a/Engine/Plugins/Tests/RHITests/Source/RHITests/Public/RHITextureTests.h b/Engine/Plugins/Tests/RHITests/Source/RHITests/Public/RHITextureTests.h
--- a/Engine/Plugins/Tests/RHITests/Source/RHITests/Public/RHITextureTests.h
+++ b/Engine/Plugins/Tests/RHITests/Source/RHITests/Public/RHITextureTests.h
@@ -647,18 +647,22 @@ public:
 				}
 			}
 
-#if 0 // @todo enable when we want to support source offsets
-			// Test source offsets
-			for (uint32 SrcRow = 0; SrcRow < 3; SrcRow++)
+			// @todo enable for all types when they support source offsets
+			if (UpdateType == ETextureUpdateType::Texture2D)
 			{
-				Region.SrcY = AxisSlotToCoord(SrcRow, SrcDataHeight, UpdateHeight);
-				for (uint32 SrcColumn = 0; SrcColumn < 3; SrcColumn++)
+				Region = FUpdateTextureRegion2D(0, 0, 0, 0, UpdateWidth, UpdateHeight);
+
+				// Test source offsets
+				for (uint32 SrcRow = 0; SrcRow < 3; SrcRow++)
 				{
-					Region.SrcX = AxisSlotToCoord(SrcColumn, SrcDataWidth, UpdateWidth);
-					bResult &= Test_UpdateTexture_Impl<SourceType>(RHICmdList, TestName, Texture, Region, SrcDataSize, UpdateDataPitch, UpdateData, ZeroData, UpdateType);
+					Region.SrcY = AxisSlotToCoord(SrcRow, SrcDataHeight, UpdateHeight);
+					for (uint32 SrcColumn = 0; SrcColumn < 3; SrcColumn++)
+					{
+						Region.SrcX = AxisSlotToCoord(SrcColumn, SrcDataWidth, UpdateWidth);
+						bResult &= Test_UpdateTexture_Impl<SourceType>(RHICmdList, TestName, Texture, Region, SrcDataSize, UpdateDataPitch, UpdateData, ZeroData, UpdateType);
+					}
 				}
 			}
-#endif
 		}
 
 		return bResult;
diff --git a/Engine/Source/Runtime/Apple/MetalRHI/Private/MetalTexture.cpp b/Engine/Source/Runtime/Apple/MetalRHI/Private/MetalTexture.cpp
--- a/Engine/Source/Runtime/Apple/MetalRHI/Private/MetalTexture.cpp
+++ b/Engine/Source/Runtime/Apple/MetalRHI/Private/MetalTexture.cpp
@@ -1566,42 +1566,59 @@ static void InternalExpandR8ToStandardRGBA(uint32* pDest, const struct FUpdateTe
 
 static FMetalBuffer Internal_CreateBufferAndCopyTexture2DUpdateRegionData(FRHITexture2D* TextureRHI, const struct FUpdateTextureRegion2D& UpdateRegion, uint32& InOutSourcePitch, const uint8* SourceData)
 {
+	const EPixelFormat PixelFormat = TextureRHI->GetFormat();
+	const FPixelFormatInfo& FormatInfo = GPixelFormats[PixelFormat];
+
+	check(UpdateRegion.Width  % FormatInfo.BlockSizeX == 0);
+	check(UpdateRegion.Height % FormatInfo.BlockSizeY == 0);
+	check(UpdateRegion.DestX  % FormatInfo.BlockSizeX == 0);
+	check(UpdateRegion.DestY  % FormatInfo.BlockSizeY == 0);
+	check(UpdateRegion.SrcX   % FormatInfo.BlockSizeX == 0);
+	check(UpdateRegion.SrcY   % FormatInfo.BlockSizeY == 0);
+
+	const uint32 SrcXInBlocks   = FMath::DivideAndRoundUp<uint32>(UpdateRegion.SrcX,   FormatInfo.BlockSizeX);
+	const uint32 SrcYInBlocks   = FMath::DivideAndRoundUp<uint32>(UpdateRegion.SrcY,   FormatInfo.BlockSizeY);
+	const uint32 WidthInBlocks  = FMath::DivideAndRoundUp<uint32>(UpdateRegion.Width,  FormatInfo.BlockSizeX);
+	const uint32 HeightInBlocks = FMath::DivideAndRoundUp<uint32>(UpdateRegion.Height, FormatInfo.BlockSizeY);
+
+	const uint8* OffsetSourceData = SourceData + FormatInfo.BlockBytes * SrcXInBlocks + InOutSourcePitch * SrcYInBlocks * FormatInfo.BlockSizeY;
+	uint32 UpdatePitch = InOutSourcePitch;
+
 	FMetalBuffer OutBuffer;
 
 	FMetalSurface* Texture = ResourceCast(TextureRHI);
 
 #if PLATFORM_MAC
 	// Expand R8_sRGB into RGBA8_sRGB for non Apple Silicon Mac.
-	if (   Texture->GetFormat() == PF_G8
+	if (   PixelFormat == PF_G8
 		&& EnumHasAnyFlags(Texture->GetFlags(), TexCreate_SRGB)
 		&& Texture->Texture.GetPixelFormat() == mtlpp::PixelFormat::RGBA8Unorm_sRGB)
 	{
 		const uint32 ExpandedBufferSize = UpdateRegion.Height * UpdateRegion.Width * sizeof(uint32);
 		OutBuffer = GetMetalDeviceContext().CreatePooledBuffer(FMetalPooledBufferArgs(GetMetalDeviceContext().GetDevice(), ExpandedBufferSize, BUF_Static, mtlpp::StorageMode::Shared));
-		InternalExpandR8ToStandardRGBA((uint32*)OutBuffer.GetContents(), UpdateRegion, InOutSourcePitch, SourceData);
+		InternalExpandR8ToStandardRGBA((uint32*)OutBuffer.GetContents(), UpdateRegion, InOutSourcePitch, OffsetSourceData);
 	}
 	else
 #endif
 	{
-		const FPixelFormatInfo& FormatInfo = GPixelFormats[TextureRHI->GetFormat()];
-		
+		const uint32 SourcePitch = InOutSourcePitch;
+		const uint32 StagingPitch = static_cast<size_t>(WidthInBlocks) * FormatInfo.BlockBytes;
+
 		const uint32 BufferSize = UpdateRegion.Height * InOutSourcePitch;
 		OutBuffer = GetMetalDeviceContext().CreatePooledBuffer(FMetalPooledBufferArgs(GetMetalDeviceContext().GetDevice(), BufferSize, BUF_Static, mtlpp::StorageMode::Shared));
 
-		uint32 CopyPitch = FMath::DivideAndRoundUp(UpdateRegion.Width, (uint32)FormatInfo.BlockSizeX) * FormatInfo.BlockBytes;
-		check(CopyPitch <= InOutSourcePitch);
-		
 		uint8* pDestRow = (uint8*)OutBuffer.GetContents();
-		uint8* pSourceRow = (uint8*)SourceData;
-		const uint32 NumRows = FMath::DivideAndRoundUp(UpdateRegion.Height, (uint32)FormatInfo.BlockSizeY);
+		const uint8* pSourceRow = OffsetSourceData;
 		
 		// Limit copy to line by line by update region pitch otherwise we can go off the end of source data on the last row
-		for (uint32 i = 0;i < NumRows;++i)
+		for (uint32 BlockRow = 0; BlockRow < HeightInBlocks; BlockRow++)
 		{
-			FMemory::Memcpy(pDestRow, pSourceRow, CopyPitch);
-			pSourceRow += InOutSourcePitch;
-			pDestRow += InOutSourcePitch;
+			FMemory::Memcpy(pDestRow, pSourceRow, StagingPitch);
+			pSourceRow += SourcePitch;
+			pDestRow += StagingPitch;
 		}
+
+		InOutSourcePitch = StagingPitch;
 	}
 
 	return OutBuffer;
diff --git a/Engine/Source/Runtime/D3D12RHI/Private/D3D12Texture.cpp b/Engine/Source/Runtime/D3D12RHI/Private/D3D12Texture.cpp
--- a/Engine/Source/Runtime/D3D12RHI/Private/D3D12Texture.cpp
+++ b/Engine/Source/Runtime/D3D12RHI/Private/D3D12Texture.cpp
@@ -2091,46 +2091,50 @@ void FD3D12Texture::UnlockInternal(class FRHICommandListImmediate* RHICmdList, F
 
 void FD3D12Texture::UpdateTexture2D(FRHICommandListBase& RHICmdList, uint32 MipIndex, const FUpdateTextureRegion2D& UpdateRegion, uint32 SourcePitch, const uint8* SourceData)
 {
-	const FPixelFormatInfo& FormatInfo = GPixelFormats[this->GetFormat()];
-	check(UpdateRegion.Width  %	FormatInfo.BlockSizeX == 0);
+	const FPixelFormatInfo& FormatInfo = GPixelFormats[GetFormat()];
+
+	check(UpdateRegion.Width  % FormatInfo.BlockSizeX == 0);
 	check(UpdateRegion.Height % FormatInfo.BlockSizeY == 0);
-	check(UpdateRegion.DestX  %	FormatInfo.BlockSizeX == 0);
-	check(UpdateRegion.DestY  %	FormatInfo.BlockSizeY == 0);
-	check(UpdateRegion.SrcX   %	FormatInfo.BlockSizeX == 0);
-	check(UpdateRegion.SrcY   %	FormatInfo.BlockSizeY == 0);
+	check(UpdateRegion.DestX  % FormatInfo.BlockSizeX == 0);
+	check(UpdateRegion.DestY  % FormatInfo.BlockSizeY == 0);
+	check(UpdateRegion.SrcX   % FormatInfo.BlockSizeX == 0);
+	check(UpdateRegion.SrcY   % FormatInfo.BlockSizeY == 0);
 
-	const uint32 WidthInBlocks = UpdateRegion.Width / FormatInfo.BlockSizeX;
-	const uint32 HeightInBlocks = UpdateRegion.Height / FormatInfo.BlockSizeY;
+	const uint32 SrcXInBlocks   = FMath::DivideAndRoundUp<uint32>(UpdateRegion.SrcX,   FormatInfo.BlockSizeX);
+	const uint32 SrcYInBlocks   = FMath::DivideAndRoundUp<uint32>(UpdateRegion.SrcY,   FormatInfo.BlockSizeY);
+	const uint32 WidthInBlocks  = FMath::DivideAndRoundUp<uint32>(UpdateRegion.Width,  FormatInfo.BlockSizeX);
+	const uint32 HeightInBlocks = FMath::DivideAndRoundUp<uint32>(UpdateRegion.Height, FormatInfo.BlockSizeY);
 
-	const uint32 AlignedSourcePitch = Align(SourcePitch, FD3D12_TEXTURE_DATA_PITCH_ALIGNMENT);
-	const uint32 bufferSize = Align(HeightInBlocks*AlignedSourcePitch, D3D12_TEXTURE_DATA_PLACEMENT_ALIGNMENT);
+	// D3D12 requires specific alignments for pitch and size since we have to do the updates via buffers
+	const size_t StagingPitch = Align(static_cast<size_t>(WidthInBlocks) * FormatInfo.BlockBytes, FD3D12_TEXTURE_DATA_PITCH_ALIGNMENT);
+	const size_t StagingBufferSize = Align(StagingPitch * HeightInBlocks, D3D12_TEXTURE_DATA_PLACEMENT_ALIGNMENT);
 
 	for (FD3D12Texture& Texture : *this)
 	{
-		FD3D12ResourceLocation UploadHeapResourceLocation(GetParentDevice());
-		void* pData = GetParentDevice()->GetDefaultFastAllocator().Allocate(bufferSize, D3D12_TEXTURE_DATA_PLACEMENT_ALIGNMENT, &UploadHeapResourceLocation);
-		check(nullptr != pData);
+		FD3D12Device* Device = Texture.GetParentDevice();
+
+		FD3D12ResourceLocation UploadHeapResourceLocation(Device);
+		void* const StagingMemory = Device->GetDefaultFastAllocator().Allocate(StagingBufferSize, D3D12_TEXTURE_DATA_PLACEMENT_ALIGNMENT, &UploadHeapResourceLocation);
+		check(StagingMemory);
 
-		byte* pRowData = (byte*)pData;
-		const byte* pSourceRowData = (byte*)SourceData;
-		const uint32 CopyPitch = WidthInBlocks * FormatInfo.BlockBytes;
-		check(CopyPitch <= SourcePitch);
-		for (uint32 i = 0; i < HeightInBlocks; i++)
+		const uint8* CopySrc = SourceData + FormatInfo.BlockBytes * SrcXInBlocks + SourcePitch * SrcYInBlocks * FormatInfo.BlockSizeY;
+		uint8* CopyDst = (uint8*)StagingMemory;
+		for (uint32 BlockRow = 0; BlockRow < HeightInBlocks; BlockRow++)
 		{
-			FMemory::Memcpy(pRowData, pSourceRowData, CopyPitch);
-			pSourceRowData += SourcePitch;
-			pRowData += AlignedSourcePitch;
+			FMemory::Memcpy(CopyDst, CopySrc, WidthInBlocks * FormatInfo.BlockBytes);
+			CopySrc += SourcePitch;
+			CopyDst += StagingPitch;
 		}
 
-		D3D12_SUBRESOURCE_FOOTPRINT SourceSubresource;
+		D3D12_SUBRESOURCE_FOOTPRINT SourceSubresource{};
 		SourceSubresource.Depth = 1;
 		SourceSubresource.Height = UpdateRegion.Height;
 		SourceSubresource.Width = UpdateRegion.Width;
 		SourceSubresource.Format = (DXGI_FORMAT)FormatInfo.PlatformFormat;
-		SourceSubresource.RowPitch = AlignedSourcePitch;
+		SourceSubresource.RowPitch = StagingPitch;
 		check(SourceSubresource.RowPitch % FD3D12_TEXTURE_DATA_PITCH_ALIGNMENT == 0);
 
-		D3D12_PLACED_SUBRESOURCE_FOOTPRINT PlacedTexture2D = { 0 };
+		D3D12_PLACED_SUBRESOURCE_FOOTPRINT PlacedTexture2D{};
 		PlacedTexture2D.Offset = UploadHeapResourceLocation.GetOffsetFromBaseOfResource();
 		PlacedTexture2D.Footprint = SourceSubresource;
 
diff --git a/Engine/Source/Runtime/OpenGLDrv/Private/OpenGLTexture.cpp b/Engine/Source/Runtime/OpenGLDrv/Private/OpenGLTexture.cpp
--- a/Engine/Source/Runtime/OpenGLDrv/Private/OpenGLTexture.cpp
+++ b/Engine/Source/Runtime/OpenGLDrv/Private/OpenGLTexture.cpp
@@ -1581,45 +1581,89 @@ void FOpenGLDynamicRHI::RHIUnlockTexture2DArray(FRHITexture2DArray* TextureRHI,
 	ResourceCast(TextureRHI)->Unlock(MipIndex, TextureIndex);
 }
 
-void FOpenGLDynamicRHI::RHIUpdateTexture2D(FRHICommandListBase& RHICmdList, FRHITexture2D* TextureRHI, uint32 MipIndex, const FUpdateTextureRegion2D& UpdateRegionIn, uint32 SourcePitch, const uint8* SourceDataIn)
+void FOpenGLDynamicRHI::RHIUpdateTexture2D(FRHICommandListBase& RHICmdList, FRHITexture2D* TextureRHI, uint32 MipIndex, const FUpdateTextureRegion2D& UpdateRegion, uint32 SourcePitch, const uint8* SourceData)
 {
-	const FUpdateTextureRegion2D UpdateRegion = UpdateRegionIn;
+	const FPixelFormatInfo& FormatInfo = GPixelFormats[TextureRHI->GetFormat()];
 
-	uint8* RHITSourceData = nullptr;
-	if (!ShouldRunGLRenderContextOpOnThisThread(RHICmdList))
+	check(UpdateRegion.Width  % FormatInfo.BlockSizeX == 0);
+	check(UpdateRegion.Height % FormatInfo.BlockSizeY == 0);
+	check(UpdateRegion.DestX  % FormatInfo.BlockSizeX == 0);
+	check(UpdateRegion.DestY  % FormatInfo.BlockSizeY == 0);
+	check(UpdateRegion.SrcX   % FormatInfo.BlockSizeX == 0);
+	check(UpdateRegion.SrcY   % FormatInfo.BlockSizeY == 0);
+
+	const uint32 SrcXInBlocks   = FMath::DivideAndRoundUp<uint32>(UpdateRegion.SrcX,   FormatInfo.BlockSizeX);
+	const uint32 SrcYInBlocks   = FMath::DivideAndRoundUp<uint32>(UpdateRegion.SrcY,   FormatInfo.BlockSizeY);
+	const uint32 WidthInBlocks  = FMath::DivideAndRoundUp<uint32>(UpdateRegion.Width,  FormatInfo.BlockSizeX);
+	const uint32 HeightInBlocks = FMath::DivideAndRoundUp<uint32>(UpdateRegion.Height, FormatInfo.BlockSizeY);
+
+	const void* UpdateMemory = SourceData + FormatInfo.BlockBytes * SrcXInBlocks + SourcePitch * SrcYInBlocks * FormatInfo.BlockSizeY;
+	uint32 UpdatePitch = SourcePitch;
+
+	const bool bNeedStagingMemory = !ShouldRunGLRenderContextOpOnThisThread(RHICmdList);
+	if (bNeedStagingMemory)
 	{
-		const FPixelFormatInfo& FormatInfo = GPixelFormats[TextureRHI->GetFormat()];
-		const size_t UpdateHeightInTiles = FMath::DivideAndRoundUp(UpdateRegion.Height, (uint32)FormatInfo.BlockSizeY);
-		const size_t DataSize = static_cast<size_t>(SourcePitch) * UpdateHeightInTiles;
-		RHITSourceData = (uint8*)FMemory::Malloc(DataSize, 16);
-		FMemory::Memcpy(RHITSourceData, SourceDataIn, DataSize);
+		const size_t SourceDataSizeInBlocks = static_cast<size_t>(WidthInBlocks) * static_cast<size_t>(HeightInBlocks);
+		const size_t SourceDataSize = SourceDataSizeInBlocks * FormatInfo.BlockBytes;
+
+		uint8* const StagingMemory = (uint8*)FMemory::Malloc(SourceDataSize);
+		const size_t StagingPitch = static_cast<size_t>(WidthInBlocks) * FormatInfo.BlockBytes;
+
+		const uint8* CopySrc = (const uint8*)UpdateMemory;
+		uint8* CopyDst = (uint8*)StagingMemory;
+		for (uint32 BlockRow = 0; BlockRow < HeightInBlocks; BlockRow++)
+		{
+			FMemory::Memcpy(CopyDst, CopySrc, WidthInBlocks * FormatInfo.BlockBytes);
+			CopySrc += SourcePitch;
+			CopyDst += StagingPitch;
+		}
+
+		UpdateMemory = StagingMemory;
+		UpdatePitch = StagingPitch;
 	}
-	const uint8* SourceData = RHITSourceData ? RHITSourceData : SourceDataIn;
-	RHICmdList.EnqueueLambda([this, TextureRHI, MipIndex, UpdateRegion, SourcePitch, SourceData, RHITSourceData](FRHICommandListBase&)
+
+	RHICmdList.EnqueueLambda([this, TextureRHI, MipIndex, UpdateRegion, UpdatePitch, UpdateMemory, bNeedStagingMemory](FRHICommandListBase&)
 	{
 		VERIFY_GL_SCOPE();
 
 		FOpenGLTexture* Texture = ResourceCast(TextureRHI);
+		const EPixelFormat PixelFormat = TextureRHI->GetFormat();
+
+		const FPixelFormatInfo& FormatInfo = GPixelFormats[PixelFormat];
+		const FOpenGLTextureFormat& GLFormat = GOpenGLTextureFormats[PixelFormat];
 
 		// Use a texture stage that's not likely to be used for draws, to avoid waiting
 		FOpenGLContextState& ContextState = GetContextStateForCurrentContext();
 		CachedSetupTextureStage(ContextState, FOpenGL::GetMaxCombinedTextureImageUnits() - 1, Texture->Target, Texture->GetResource(), 0, Texture->GetNumMips());
 		CachedBindPixelUnpackBuffer(ContextState, 0);
 
-		EPixelFormat PixelFormat = Texture->GetFormat();
-		check(GPixelFormats[PixelFormat].BlockSizeX == 1);
-		check(GPixelFormats[PixelFormat].BlockSizeY == 1);
-		const FOpenGLTextureFormat& GLFormat = GOpenGLTextureFormats[PixelFormat];
-		const uint32 FormatBPP = GPixelFormats[PixelFormat].BlockBytes;
-		checkf(!GLFormat.bCompressed, TEXT("RHIUpdateTexture2D not currently supported for compressed (%s) textures by the OpenGL RHI"), GPixelFormats[PixelFormat].Name);
+		glPixelStorei(GL_UNPACK_ROW_LENGTH, UpdatePitch / FormatInfo.BlockBytes);
+		glPixelStorei(GL_UNPACK_ALIGNMENT, 1);
 
-		glPixelStorei(GL_UNPACK_ROW_LENGTH, SourcePitch / FormatBPP);
+		if (GLFormat.bCompressed)
+		{
+			glCompressedTexSubImage2D(
+				Texture->Target,
+				MipIndex,
+				UpdateRegion.DestX, UpdateRegion.DestY,
+				UpdateRegion.Width, UpdateRegion.Height,
+				GLFormat.Format,
+				UpdatePitch * FMath::DivideAndRoundUp<uint32>(UpdateRegion.Height, FormatInfo.BlockSizeY),
+				UpdateMemory);
+		}
+		else
+		{
+			glTexSubImage2D(
+				Texture->Target,
+				MipIndex,
+				UpdateRegion.DestX, UpdateRegion.DestY,
+				UpdateRegion.Width, UpdateRegion.Height,
+				GLFormat.Format,
+				GLFormat.Type,
+				UpdateMemory);
+		}
 
-		glPixelStorei(GL_UNPACK_ALIGNMENT, 1);
-		glTexSubImage2D(Texture->Target, MipIndex, UpdateRegion.DestX, UpdateRegion.DestY, UpdateRegion.Width, UpdateRegion.Height,
-			GLFormat.Format, GLFormat.Type, SourceData);
 		glPixelStorei(GL_UNPACK_ALIGNMENT, 4);
-
 		glPixelStorei(GL_UNPACK_ROW_LENGTH, 0);
 
 		// No need to restore texture stage; leave it like this,
@@ -1627,9 +1671,9 @@ void FOpenGLDynamicRHI::RHIUpdateTexture2D(FRHICommandListBase& RHICmdList, FRHI
 		// next operation that needs the stage will switch something else in on it.
 
 		// free source data if we're on RHIT
-		if (RHITSourceData)
+		if (bNeedStagingMemory)
 		{
-			FMemory::Free(RHITSourceData);
+			FMemory::Free(const_cast<void*>(UpdateMemory));
 		}
 	});
 }
diff --git a/Engine/Source/Runtime/VulkanRHI/Private/VulkanTexture.cpp b/Engine/Source/Runtime/VulkanRHI/Private/VulkanTexture.cpp
--- a/Engine/Source/Runtime/VulkanRHI/Private/VulkanTexture.cpp
+++ b/Engine/Source/Runtime/VulkanRHI/Private/VulkanTexture.cpp
@@ -1270,60 +1270,53 @@ void FVulkanDynamicRHI::RHIUnlockTexture2DArray(FRHITexture2DArray* TextureRHI,
 	}
 }
 
-void FVulkanDynamicRHI::InternalUpdateTexture2D(FRHICommandListBase& RHICmdList, FRHITexture2D* TextureRHI, uint32 MipIndex, const struct FUpdateTextureRegion2D& UpdateRegion, uint32 SourceRowPitch, const uint8* SourceData)
+void FVulkanDynamicRHI::InternalUpdateTexture2D(FRHICommandListBase& RHICmdList, FRHITexture2D* TextureRHI, uint32 MipIndex, const FUpdateTextureRegion2D& UpdateRegion, uint32 SourcePitch, const uint8* SourceData)
 {
 	LLM_SCOPE_VULKAN(ELLMTagVulkan::VulkanTextures);
-	FVulkanTexture* Texture = ResourceCast(TextureRHI);
 
-	const EPixelFormat PixelFormat = Texture->GetDesc().Format;
-	const int32 BlockSizeX = GPixelFormats[PixelFormat].BlockSizeX;
-	const int32 BlockSizeY = GPixelFormats[PixelFormat].BlockSizeY;
-	const int32 BlockSizeZ = GPixelFormats[PixelFormat].BlockSizeZ;
-	const int32 BlockBytes = GPixelFormats[PixelFormat].BlockBytes;
-	VkFormat Format = UEToVkTextureFormat(PixelFormat, false);
+	const FPixelFormatInfo& FormatInfo = GPixelFormats[TextureRHI->GetFormat()];
 
-	ensure(BlockSizeZ == 1);
+	check(UpdateRegion.Width  % FormatInfo.BlockSizeX == 0);
+	check(UpdateRegion.Height % FormatInfo.BlockSizeY == 0);
+	check(UpdateRegion.DestX  % FormatInfo.BlockSizeX == 0);
+	check(UpdateRegion.DestY  % FormatInfo.BlockSizeY == 0);
+	check(UpdateRegion.SrcX   % FormatInfo.BlockSizeX == 0);
+	check(UpdateRegion.SrcY   % FormatInfo.BlockSizeY == 0);
 
-	const VkPhysicalDeviceLimits& Limits = Device->GetLimits();
+	const uint32 SrcXInBlocks   = FMath::DivideAndRoundUp<uint32>(UpdateRegion.SrcX,   FormatInfo.BlockSizeX);
+	const uint32 SrcYInBlocks   = FMath::DivideAndRoundUp<uint32>(UpdateRegion.SrcY,   FormatInfo.BlockSizeY);
+	const uint32 WidthInBlocks  = FMath::DivideAndRoundUp<uint32>(UpdateRegion.Width,  FormatInfo.BlockSizeX);
+	const uint32 HeightInBlocks = FMath::DivideAndRoundUp<uint32>(UpdateRegion.Height, FormatInfo.BlockSizeY);
 
-	VkBufferImageCopy Region;
-	FMemory::Memzero(Region);
-	VulkanRHI::FStagingBuffer* StagingBuffer = nullptr;
-	const uint32 NumBlocksX = (uint32)FMath::DivideAndRoundUp<int32>(UpdateRegion.Width, (uint32)BlockSizeX);
-	const uint32 NumBlocksY = (uint32)FMath::DivideAndRoundUp<int32>(UpdateRegion.Height, (uint32)BlockSizeY);
-	ensure(NumBlocksX * BlockBytes <= SourceRowPitch);
+	const VkPhysicalDeviceLimits& Limits = Device->GetLimits();
 
-	const uint32 DestRowPitch = NumBlocksX * BlockBytes;
-	const uint32 DestSlicePitch = DestRowPitch * NumBlocksY;
+	const size_t StagingPitch = static_cast<size_t>(WidthInBlocks) * FormatInfo.BlockBytes;
+	const size_t StagingBufferSize = Align(StagingPitch * HeightInBlocks, Limits.minMemoryMapAlignment);
 
-	const uint32 BufferSize = Align(DestSlicePitch, Limits.minMemoryMapAlignment);
-	StagingBuffer = Device->GetStagingManager().AcquireBuffer(BufferSize);
-	void* RESTRICT Memory = StagingBuffer->GetMappedPointer();
+	VulkanRHI::FStagingBuffer* StagingBuffer = Device->GetStagingManager().AcquireBuffer(StagingBufferSize);
+	void* RESTRICT StagingMemory = StagingBuffer->GetMappedPointer();
 
-	uint8* RESTRICT DestData = (uint8*)Memory;
-	uint8* RESTRICT SourceRowData = (uint8*)SourceData;
-	for (uint32 Height = 0; Height < NumBlocksY; ++Height)
+	const uint8* CopySrc = SourceData + FormatInfo.BlockBytes * SrcXInBlocks + SourcePitch * SrcYInBlocks * FormatInfo.BlockSizeY;
+	uint8* CopyDst = (uint8*)StagingMemory;
+	for (uint32 BlockRow = 0; BlockRow < HeightInBlocks; BlockRow++)
 	{
-		FMemory::Memcpy(DestData, SourceRowData, NumBlocksX * BlockBytes);
-		DestData += DestRowPitch;
-		SourceRowData += SourceRowPitch;
+		FMemory::Memcpy(CopyDst, CopySrc, WidthInBlocks * FormatInfo.BlockBytes);
+		CopySrc += SourcePitch;
+		CopyDst += StagingPitch;
 	}
 
-	//Region.bufferOffset = 0;
-	// Set these to zero to assume tightly packed buffer
-	//Region.bufferRowLength = 0;
-	//Region.bufferImageHeight = 0;
+	VkBufferImageCopy Region{};
 	Region.imageSubresource.aspectMask = VK_IMAGE_ASPECT_COLOR_BIT;
 	Region.imageSubresource.mipLevel = MipIndex;
-	//Region.imageSubresource.baseArrayLayer = 0;
 	Region.imageSubresource.layerCount = 1;
 	Region.imageOffset.x = UpdateRegion.DestX;
 	Region.imageOffset.y = UpdateRegion.DestY;
-	//Region.imageOffset.z = 0;
 	Region.imageExtent.width = UpdateRegion.Width;
 	Region.imageExtent.height = UpdateRegion.Height;
 	Region.imageExtent.depth = 1;
 
+	FVulkanTexture* Texture = ResourceCast(TextureRHI);
+
 	if (RHICmdList.IsBottomOfPipe())
 	{
 		FVulkanTexture::InternalLockWrite(FVulkanCommandListContext::GetVulkanContext(RHICmdList.GetContext()), Texture, Region, StagingBuffer);
diff --git a/Engine/Source/Runtime/Windows/D3D11RHI/Private/D3D11Texture.cpp b/Engine/Source/Runtime/Windows/D3D11RHI/Private/D3D11Texture.cpp
--- a/Engine/Source/Runtime/Windows/D3D11RHI/Private/D3D11Texture.cpp
+++ b/Engine/Source/Runtime/Windows/D3D11RHI/Private/D3D11Texture.cpp
@@ -1273,13 +1273,45 @@ void FD3D11DynamicRHI::RHIUnlockTexture2DArray(FRHITexture2DArray* TextureRHI, u
 void FD3D11DynamicRHI::RHIUpdateTexture2D(FRHICommandListBase& RHICmdList, FRHITexture2D* TextureRHI, uint32 MipIndex, const FUpdateTextureRegion2D& UpdateRegion, uint32 SourcePitch, const uint8* SourceData)
 {
 	const FPixelFormatInfo& FormatInfo = GPixelFormats[TextureRHI->GetFormat()];
-	const size_t UpdateHeightInTiles = FMath::DivideAndRoundUp(UpdateRegion.Height, (uint32)FormatInfo.BlockSizeY);
-	const size_t SourceDataSize = static_cast<size_t>(SourcePitch) * UpdateHeightInTiles;
-	uint8* SourceDataCopy = (uint8*)FMemory::Malloc(SourceDataSize);
-	FMemory::Memcpy(SourceDataCopy, SourceData, SourceDataSize);
-	SourceData = SourceDataCopy;
 
-	RHICmdList.EnqueueLambda([this, TextureRHI, MipIndex, UpdateRegion, SourcePitch, SourceData] (FRHICommandListBase&)
+	check(UpdateRegion.Width  % FormatInfo.BlockSizeX == 0);
+	check(UpdateRegion.Height % FormatInfo.BlockSizeY == 0);
+	check(UpdateRegion.DestX  % FormatInfo.BlockSizeX == 0);
+	check(UpdateRegion.DestY  % FormatInfo.BlockSizeY == 0);
+	check(UpdateRegion.SrcX   % FormatInfo.BlockSizeX == 0);
+	check(UpdateRegion.SrcY   % FormatInfo.BlockSizeY == 0);
+
+	const uint32 SrcXInBlocks   = FMath::DivideAndRoundUp<uint32>(UpdateRegion.SrcX,   FormatInfo.BlockSizeX);
+	const uint32 SrcYInBlocks   = FMath::DivideAndRoundUp<uint32>(UpdateRegion.SrcY,   FormatInfo.BlockSizeY);
+	const uint32 WidthInBlocks  = FMath::DivideAndRoundUp<uint32>(UpdateRegion.Width,  FormatInfo.BlockSizeX);
+	const uint32 HeightInBlocks = FMath::DivideAndRoundUp<uint32>(UpdateRegion.Height, FormatInfo.BlockSizeY);
+
+	const void* UpdateMemory = SourceData + FormatInfo.BlockBytes * SrcXInBlocks + SourcePitch * SrcYInBlocks * FormatInfo.BlockSizeY;
+	uint32 UpdatePitch = SourcePitch;
+
+	const bool bNeedStagingMemory = RHICmdList.IsTopOfPipe();
+	if (bNeedStagingMemory)
+	{
+		const size_t SourceDataSizeInBlocks = static_cast<size_t>(WidthInBlocks) * static_cast<size_t>(HeightInBlocks);
+		const size_t SourceDataSize = SourceDataSizeInBlocks * FormatInfo.BlockBytes;
+
+		uint8* const StagingMemory = (uint8*)FMemory::Malloc(SourceDataSize);
+		const size_t StagingPitch = static_cast<size_t>(WidthInBlocks) * FormatInfo.BlockBytes;
+
+		const uint8* CopySrc = (const uint8*)UpdateMemory;
+		uint8* CopyDst = (uint8*)StagingMemory;
+		for (uint32 BlockRow = 0; BlockRow < HeightInBlocks; BlockRow++)
+		{
+			FMemory::Memcpy(CopyDst, CopySrc, WidthInBlocks * FormatInfo.BlockBytes);
+			CopySrc += SourcePitch;
+			CopyDst += StagingPitch;
+		}
+
+		UpdateMemory = StagingMemory;
+		UpdatePitch = StagingPitch;
+	}
+
+	RHICmdList.EnqueueLambda([this, TextureRHI, MipIndex, UpdateRegion, UpdatePitch, UpdateMemory, bNeedStagingMemory] (FRHICommandListBase&)
 	{
 		FD3D11Texture* Texture = ResourceCast(TextureRHI);
 
@@ -1289,16 +1321,12 @@ void FD3D11DynamicRHI::RHIUpdateTexture2D(FRHICommandListBase& RHICmdList, FRHIT
 			UpdateRegion.DestX + UpdateRegion.Width, UpdateRegion.DestY + UpdateRegion.Height, 1
 		};
 
-		check(UpdateRegion.Width % GPixelFormats[Texture->GetFormat()].BlockSizeX == 0);
-		check(UpdateRegion.Height % GPixelFormats[Texture->GetFormat()].BlockSizeY == 0);
-		check(UpdateRegion.DestX % GPixelFormats[Texture->GetFormat()].BlockSizeX == 0);
-		check(UpdateRegion.DestY % GPixelFormats[Texture->GetFormat()].BlockSizeY == 0);
-		check(UpdateRegion.SrcX % GPixelFormats[Texture->GetFormat()].BlockSizeX == 0);
-		check(UpdateRegion.SrcY % GPixelFormats[Texture->GetFormat()].BlockSizeY == 0);
+		Direct3DDeviceIMContext->UpdateSubresource(Texture->GetResource(), MipIndex, &DestBox, UpdateMemory, UpdatePitch, 0);
 
-		Direct3DDeviceIMContext->UpdateSubresource(Texture->GetResource(), MipIndex, &DestBox, SourceData, SourcePitch, 0);
-
-		FMemory::Free((void*)SourceData);
+		if (bNeedStagingMemory)
+		{
+			FMemory::Free(const_cast<void*>(UpdateMemory));
+		}
 	});
 }
 
