Disable Niagara GPU particles on Adreno devcies with an old driver to fix a GPU crash.

diff --git a/Engine/Source/Runtime/OpenGLDrv/Private/OpenGLDevice.cpp b/Engine/Source/Runtime/OpenGLDrv/Private/OpenGLDevice.cpp
--- a/Engine/Source/Runtime/OpenGLDrv/Private/OpenGLDevice.cpp
+++ b/Engine/Source/Runtime/OpenGLDrv/Private/OpenGLDevice.cpp
@@ -1327,6 +1327,31 @@ FOpenGLDynamicRHI::FOpenGLDynamicRHI()
 			}
 		}
 	}
+
+#if PLATFORM_ANDROID
+	// Temp disable gpu particles for OpenGL because of issues on Adreno devices with old driver version
+	if (GRHIVendorId == 0x5143)
+	{
+		auto* CVar = IConsoleManager::Get().FindConsoleVariable(TEXT("fx.NiagaraAllowGPUParticles"));
+		if (CVar)
+		{
+			int32 DriverVersion = INT32_MAX;
+			FString SubVersionString, DriverVersionString;
+			FString VersionString = FString(ANSI_TO_TCHAR((const ANSICHAR*)glGetString(GL_VERSION)));
+			
+			if (VersionString.Split(TEXT("V@"), nullptr, &SubVersionString) && SubVersionString.Split(TEXT(" "), &DriverVersionString, nullptr))
+			{
+				DriverVersion = FCString::Atoi(*DriverVersionString);
+			}
+
+			if (DriverVersion < 415)
+			{
+				CVar->Set(false);
+				UE_LOG(LogRHI, Log, TEXT("GPU particles are disabled on this device because the driver version %d is less than 415"), DriverVersion);
+			}
+		}
+	}
+#endif
 	
 	// Disable SingleRHIThreadStall for GL occlusion queiresn, which should be set for D3D11 only. Enabling it causes RT->RHIT deadlock
 	{
