Set bRequiresReadOnlyBuffersWorkaround  to true for a subset of devices

diff --git a/Engine/Source/Runtime/OpenGLDrv/Private/Android/AndroidOpenGL.cpp b/Engine/Source/Runtime/OpenGLDrv/Private/Android/AndroidOpenGL.cpp
--- a/Engine/Source/Runtime/OpenGLDrv/Private/Android/AndroidOpenGL.cpp
+++ b/Engine/Source/Runtime/OpenGLDrv/Private/Android/AndroidOpenGL.cpp
@@ -903,7 +903,8 @@ void FAndroidOpenGL::ProcessExtensions(const FString& ExtensionsString)
 
 	if (bIsMaliBased)
 	{
-		bRequiresReadOnlyBuffersWorkaround = VersionString.Contains("r9p");
+		//TODO restrict this to problematic drivers only
+		bRequiresReadOnlyBuffersWorkaround = true;
 	}
 
 	// Disable ASTC if requested by device profile
