Fix for OpenGL buffer reporting to LLM that it didn't release any memory

diff --git a/Engine/Source/Runtime/OpenGLDrv/Public/OpenGLResources.h b/Engine/Source/Runtime/OpenGLDrv/Public/OpenGLResources.h
--- a/Engine/Source/Runtime/OpenGLDrv/Public/OpenGLResources.h
+++ b/Engine/Source/Runtime/OpenGLDrv/Public/OpenGLResources.h
@@ -574,15 +574,33 @@ public:
 	FOpenGLBaseBuffer(FRHIBufferDesc const& BufferDesc)
 		: FRHIBuffer(BufferDesc)
 	{
-		LLM_SCOPED_PAUSE_TRACKING_WITH_ENUM_AND_AMOUNT(ELLMTag::GraphicsPlatform, BufferDesc.Size, ELLMTracker::Platform, ELLMAllocType::None);
-		LLM_SCOPED_PAUSE_TRACKING_WITH_ENUM_AND_AMOUNT(ELLMTag::Meshes, BufferDesc.Size, ELLMTracker::Default, ELLMAllocType::None);
+		LLM_SCOPED_PAUSE_TRACKING_WITH_ENUM_AND_AMOUNT(ELLMTag::GraphicsPlatform, (int64)GetSize(), ELLMTracker::Platform, ELLMAllocType::None);
+		LLM_SCOPED_PAUSE_TRACKING_WITH_ENUM_AND_AMOUNT(ELLMTag::Meshes, (int64)GetSize(), ELLMTracker::Default, ELLMAllocType::None);
 	}
 
 	~FOpenGLBaseBuffer()
+	{
+		// If ReleaseOwnership was called, then Size is already 0
+		if (GetSize() != 0)
+		{
+			LLM_SCOPED_PAUSE_TRACKING_WITH_ENUM_AND_AMOUNT(ELLMTag::GraphicsPlatform, -(int64)GetSize(), ELLMTracker::Platform, ELLMAllocType::None);
+			LLM_SCOPED_PAUSE_TRACKING_WITH_ENUM_AND_AMOUNT(ELLMTag::Meshes, -(int64)GetSize(), ELLMTracker::Default, ELLMAllocType::None);
+		}
+    }
+
+	void TakeOwnership(FOpenGLBaseBuffer& Other)
+	{
+		FRHIBuffer::TakeOwnership(Other);
+		LLM_SCOPED_PAUSE_TRACKING_WITH_ENUM_AND_AMOUNT(ELLMTag::GraphicsPlatform, (int64)GetSize(), ELLMTracker::Platform, ELLMAllocType::None);
+		LLM_SCOPED_PAUSE_TRACKING_WITH_ENUM_AND_AMOUNT(ELLMTag::Meshes, (int64)GetSize(), ELLMTracker::Default, ELLMAllocType::None);
+	}
+
+	void ReleaseOwnership()
 	{
 		LLM_SCOPED_PAUSE_TRACKING_WITH_ENUM_AND_AMOUNT(ELLMTag::GraphicsPlatform, -(int64)GetSize(), ELLMTracker::Platform, ELLMAllocType::None);
 		LLM_SCOPED_PAUSE_TRACKING_WITH_ENUM_AND_AMOUNT(ELLMTag::Meshes, -(int64)GetSize(), ELLMTracker::Default, ELLMAllocType::None);
-    }
+		FRHIBuffer::ReleaseOwnership();
+	}
 
 	static bool OnDelete(GLuint Resource,uint32 Size,bool bStreamDraw,uint32 Offset)
 	{
