Android GLES: Do not attempt to precompile programs if the external services are not available.
 - attempting to precompile on the RHIT will cause unacceptable hitching.

diff --git a/Engine/Source/Runtime/OpenGLDrv/Private/OpenGLShaders.cpp b/Engine/Source/Runtime/OpenGLDrv/Private/OpenGLShaders.cpp
--- a/Engine/Source/Runtime/OpenGLDrv/Private/OpenGLShaders.cpp
+++ b/Engine/Source/Runtime/OpenGLDrv/Private/OpenGLShaders.cpp
@@ -2797,7 +2797,17 @@ static FOpenGLProgramBinary ExternalProgramCompile(const FOpenGLProgramKey& Prog
 void FOpenGLDynamicRHI::PrepareGFXBoundShaderState(const FGraphicsPipelineStateInitializer& Initializer)
 {
 	const bool bIsPreCachePSO = Initializer.bPSOPrecache || Initializer.bFromPSOFileCache;
-	if (!bIsPreCachePSO || !FOpenGLProgramBinaryCache::IsEnabled())
+	// if external creation is not available then ignore precache PSOs
+	// precaching on the RHIT will cause severe hitching.
+	const bool bCanCreateExternally = CanCreateExternally(bIsPreCachePSO);
+
+	{
+		static bool bOneTime = true;
+		UE_CLOG(bOneTime && !bCanCreateExternally, LogRHI, Warning, TEXT("Ignoring precache PSO, external compiler not active."));
+		bOneTime = bOneTime && bCanCreateExternally;
+	}
+
+	if (!bIsPreCachePSO || !FOpenGLProgramBinaryCache::IsEnabled() || !bCanCreateExternally)
 	{
 		return;
 	}
@@ -2823,47 +2833,7 @@ void FOpenGLDynamicRHI::PrepareGFXBoundShaderState(const FGraphicsPipelineStateI
 		if (FOpenGLProgramBinaryCache::IsBuildingCache())
 		{
 			OGL_BINARYCACHE_STATS_MARKBEGINCOMPILE(ProgramKey);
-			FOpenGLProgramBinary CompiledProgram;
-			if (CanCreateExternally(bIsPreCachePSO))
-			{
-				CompiledProgram = ExternalProgramCompile(ProgramKey, VertexShaderRHI, PixelShaderRHI);
-			}
-			else
-			{
-				RunOnGLRenderContextThread([&CompiledProgram, &VertexShaderRHI, &PixelShaderRHI, &GeometryShaderRHI,this]()
-				{
-					check(IsInRenderingThread() || IsInRHIThread());
-					VERIFY_GL_SCOPE();
-					FOpenGLLinkedProgramConfiguration Config = CreateConfig(VertexShaderRHI,
-						PixelShaderRHI,
-						GeometryShaderRHI);
-
-					FOpenGLLinkedProgram* LinkedProgram = LinkProgram(Config);
-					if(LinkedProgram == nullptr)
-					{
-						RHIGetPanicDelegate().ExecuteIfBound(FName("FailedProgramLinkDuringPrecompile"));
-						UE_LOG(LogRHI, Fatal, TEXT("Failed to link program [%s]. Current total programs: %d"), *Config.ProgramKey.ToString(), GNumPrograms);
-						return;
-					}
-
-					CompiledProgram = UE::OpenGL::GetProgramBinaryFromGLProgram(LinkedProgram->Program);
-					// optional: enqueue this binary, it will be picked up during createboundshaderstate.
-					// If this is enabled then the binary program will be sent to the GL container.
-					// not doing this for now as the completed binary cache will be loaded as normal at the end of the pre-caching process.
-					// FOpenGLProgramBinaryCache::EnqueueBinaryForGLProgramContainer(ProgramKey, TUniqueObj<FOpenGLProgramBinary>(MoveTemp(CompiledProgram)));
-
-					if (!FGLProgramCache::IsUsingLRU())
-					{
-						GetOpenGLProgramsCache().Add(Config.ProgramKey, LinkedProgram);	// if we're not using the LRU then we add the program as normal.
-					}
-					else
-					{
-						// with LRU mode the program is made available as an evicted program after the PSO cache completes.
-						delete LinkedProgram;
-						LinkedProgram = nullptr;
-					}
-				}, true); // TODO: this wait can be optimized away
-			}
+			FOpenGLProgramBinary CompiledProgram = ExternalProgramCompile(ProgramKey, VertexShaderRHI, PixelShaderRHI);
 
 			if (CompiledProgram.IsValid())
 			{
