Android GLES PSO precaching
Add gl rhi specific cvar toggle rhi precache support.
Add GL RHI specific precache PSO hash, only shader hashes are relevant for GL binary program caching.
move gl program binary dest directory

diff --git a/Engine/Source/Runtime/OpenGLDrv/Private/OpenGLDevice.cpp b/Engine/Source/Runtime/OpenGLDrv/Private/OpenGLDevice.cpp
--- a/Engine/Source/Runtime/OpenGLDrv/Private/OpenGLDevice.cpp
+++ b/Engine/Source/Runtime/OpenGLDrv/Private/OpenGLDevice.cpp
@@ -80,6 +80,13 @@ static TAutoConsoleVariable<bool> CVarEnablePSOFileCacheWhenPrecachingActive(
 	TEXT("true: GL RHI Allows both PSO file cache and precaching."),
 	ECVF_RenderThreadSafe | ECVF_ReadOnly);
 
+static TAutoConsoleVariable<bool> CVarAllowPSOPrecaching(
+	TEXT("r.OpenGL.AllowPSOPrecaching"),
+	true,
+	TEXT("true: if r.PSOPrecaching=1 GL RHI will use precaching. (default)\n")
+	TEXT("false: GL RHI will disable precaching (even if r.PSOPrecaching=1). "),
+	ECVF_RenderThreadSafe | ECVF_ReadOnly);
+
 void OnQueryCreation( FOpenGLRenderQuery* Query )
 {
 	check(PrivateOpenGLDevicePtr);
@@ -1246,7 +1253,7 @@ static void InitRHICapabilitiesForGL()
 	GRHINeedsUnatlasedCSMDepthsWorkaround = true;
 
 	static const auto CVarPSOPrecaching = IConsoleManager::Get().FindConsoleVariable(TEXT("r.PSOPrecaching"));
-	if (CVarPSOPrecaching && CVarPSOPrecaching->GetInt() != 0)
+	if (CVarPSOPrecaching && CVarPSOPrecaching->GetInt() != 0 && CVarAllowPSOPrecaching.GetValueOnAnyThread())
 	{
 		GRHISupportsPSOPrecaching = true;
 	}
diff --git a/Engine/Source/Runtime/OpenGLDrv/Private/OpenGLDrv.cpp b/Engine/Source/Runtime/OpenGLDrv/Private/OpenGLDrv.cpp
--- a/Engine/Source/Runtime/OpenGLDrv/Private/OpenGLDrv.cpp
+++ b/Engine/Source/Runtime/OpenGLDrv/Private/OpenGLDrv.cpp
@@ -150,6 +150,62 @@ bool FOpenGLDynamicRHI::RHIRequiresComputeGenerateMips() const
 	return !FOpenGL::SupportsGenerateMipmap();
 };
 
+// only use shader hashes to determine GL PSO hash;
+uint64 FOpenGLDynamicRHI::RHIComputeStatePrecachePSOHash(const FGraphicsPipelineStateInitializer& Initializer)
+{
+	struct FHashKey
+	{
+		FSHAHash VertexShader;
+		FSHAHash PixelShader;
+#if PLATFORM_SUPPORTS_GEOMETRY_SHADERS
+		FSHAHash GeometryShader;
+#endif // PLATFORM_SUPPORTS_GEOMETRY_SHADERS
+#if PLATFORM_SUPPORTS_MESH_SHADERS
+		FSHAHash MeshShader;
+#endif // PLATFORM_SUPPORTS_MESH_SHADERS
+	} HashKey;
+
+	FMemory::Memzero(&HashKey, sizeof(FHashKey));
+
+	HashKey.VertexShader = Initializer.BoundShaderState.GetVertexShader() ? Initializer.BoundShaderState.GetVertexShader()->GetHash() : FSHAHash();
+	HashKey.PixelShader = Initializer.BoundShaderState.GetPixelShader() ? Initializer.BoundShaderState.GetPixelShader()->GetHash() : FSHAHash();
+#if PLATFORM_SUPPORTS_GEOMETRY_SHADERS
+	HashKey.GeometryShader = Initializer.BoundShaderState.GetGeometryShader() ? Initializer.BoundShaderState.GetGeometryShader()->GetHash() : FSHAHash();
+#endif
+#if PLATFORM_SUPPORTS_MESH_SHADERS
+	HashKey.MeshShader = Initializer.BoundShaderState.GetMeshShader() ? Initializer.BoundShaderState.GetMeshShader()->GetHash() : FSHAHash();
+#endif
+
+	uint64 PrecachePSOHash = CityHash64((const char*)&HashKey, sizeof(FHashKey));
+	return PrecachePSOHash;
+}
+
+uint64 FOpenGLDynamicRHI::RHIComputePrecachePSOHash(const FGraphicsPipelineStateInitializer& Initializer)
+{
+	uint64 StatePrecachePSOHash = Initializer.StatePrecachePSOHash;
+	if (StatePrecachePSOHash == 0)
+	{
+		StatePrecachePSOHash = RHIComputeStatePrecachePSOHash(Initializer);
+	}
+
+	return StatePrecachePSOHash;
+}
+
+bool FOpenGLDynamicRHI::RHIMatchPrecachePSOInitializers(const FGraphicsPipelineStateInitializer& LHS, const FGraphicsPipelineStateInitializer& RHS)
+{
+	// check the RHI shaders (pointer check for shaders should be fine)
+	if (LHS.BoundShaderState.VertexShaderRHI != RHS.BoundShaderState.VertexShaderRHI ||
+		LHS.BoundShaderState.PixelShaderRHI != RHS.BoundShaderState.PixelShaderRHI ||
+		LHS.BoundShaderState.GetMeshShader() != RHS.BoundShaderState.GetMeshShader() ||
+		LHS.BoundShaderState.GetAmplificationShader() != RHS.BoundShaderState.GetAmplificationShader() ||
+		LHS.BoundShaderState.GetGeometryShader() != RHS.BoundShaderState.GetGeometryShader())
+	{
+		return false;
+	}
+
+	return true;
+}
+
 void FOpenGLGPUProfiler::BeginFrame(FOpenGLDynamicRHI* InRHI)
 {
 	if (NestedFrameCount++>0)
diff --git a/Engine/Source/Runtime/OpenGLDrv/Private/OpenGLProgramBinaryFileCache.cpp b/Engine/Source/Runtime/OpenGLDrv/Private/OpenGLProgramBinaryFileCache.cpp
--- a/Engine/Source/Runtime/OpenGLDrv/Private/OpenGLProgramBinaryFileCache.cpp
+++ b/Engine/Source/Runtime/OpenGLDrv/Private/OpenGLProgramBinaryFileCache.cpp
@@ -343,13 +343,15 @@ void FOpenGLProgramBinaryCache::Initialize()
 
 
 	FString CacheFolderPathRoot;
+	FString OldCacheFolderPathRoot;
 #if PLATFORM_ANDROID && USE_ANDROID_FILE
 	// @todo Lumin: Use that GetPathForExternalWrite or something?
 	extern FString GExternalFilePath;
-	CacheFolderPathRoot = GExternalFilePath / TEXT("ProgramBinaryCache");
-
+	OldCacheFolderPathRoot = GExternalFilePath / TEXT("ProgramBinaryCache");
+	CacheFolderPathRoot = GExternalFilePath / TEXT("RHICache") / TEXT("ProgramBinaryCache");
 #else
-	CacheFolderPathRoot = FPaths::ProjectSavedDir() / TEXT("ProgramBinaryCache");
+	OldCacheFolderPathRoot = FPaths::ProjectSavedDir() / TEXT("ProgramBinaryCache");
+	CacheFolderPathRoot = FPaths::ProjectSavedDir() / TEXT("RHICache") / TEXT("ProgramBinaryCache");
 #endif
 
 	// Remove entire ProgramBinaryCache folder if -ClearOpenGLBinaryProgramCache is specified on command line
@@ -359,6 +361,17 @@ void FOpenGLProgramBinaryCache::Initialize()
 		FPlatformFileManager::Get().GetPlatformFile().DeleteDirectoryRecursively(*CacheFolderPathRoot);
 	}
 
+	if (FPlatformFileManager::Get().GetPlatformFile().DirectoryExists(*OldCacheFolderPathRoot))
+	{
+		UE_LOG(LogRHI, Log, TEXT("Moving program binary cache: %s -> %s"), *OldCacheFolderPathRoot, *CacheFolderPathRoot);
+		
+		// Note: have to copy and delete as TManagedStoragePlatformFile prevents moving of directories.
+		// FPlatformFileManager::Get().GetPlatformFile().MoveFile(*CacheFolderPathRoot, *OldCacheFolderPathRoot);
+
+		FPlatformFileManager::Get().GetPlatformFile().CopyDirectoryTree(*CacheFolderPathRoot, *OldCacheFolderPathRoot, false);
+		FPlatformFileManager::Get().GetPlatformFile().DeleteDirectoryRecursively(*OldCacheFolderPathRoot);
+	}
+
 	CachePtr = new FOpenGLProgramBinaryCache(CacheFolderPathRoot);
 	UE_LOG(LogRHI, Log, TEXT("Enabling program binary cache dir at %s"), *CachePtr->GetProgramBinaryCacheDir());
 
diff --git a/Engine/Source/Runtime/OpenGLDrv/Private/OpenGLShaders.cpp b/Engine/Source/Runtime/OpenGLDrv/Private/OpenGLShaders.cpp
--- a/Engine/Source/Runtime/OpenGLDrv/Private/OpenGLShaders.cpp
+++ b/Engine/Source/Runtime/OpenGLDrv/Private/OpenGLShaders.cpp
@@ -2801,14 +2801,14 @@ void FOpenGLDynamicRHI::PrepareGFXBoundShaderState(const FGraphicsPipelineStateI
 	// precaching on the RHIT will cause severe hitching.
 	const bool bCanCreateExternally = CanCreateExternally(bIsPreCachePSO);
 
-	{
-		static bool bOneTime = true;
-		UE_CLOG(bOneTime && !bCanCreateExternally, LogRHI, Warning, TEXT("Ignoring precache PSO, external compiler not active."));
-		bOneTime = bOneTime && bCanCreateExternally;
-	}
-
 	if (!bIsPreCachePSO || !FOpenGLProgramBinaryCache::IsEnabled() || !bCanCreateExternally)
 	{
+		static bool bOneTime = true;
+		if(bOneTime && bIsPreCachePSO && FOpenGLProgramBinaryCache::IsEnabled())
+		{
+			UE_LOG(LogRHI, Warning, TEXT("Ignoring precache PSO, external compiler not active."));
+			bOneTime = false;
+		}
 		return;
 	}
 
diff --git a/Engine/Source/Runtime/OpenGLDrv/Public/OpenGLDrv.h b/Engine/Source/Runtime/OpenGLDrv/Public/OpenGLDrv.h
--- a/Engine/Source/Runtime/OpenGLDrv/Public/OpenGLDrv.h
+++ b/Engine/Source/Runtime/OpenGLDrv/Public/OpenGLDrv.h
@@ -507,6 +507,15 @@ public:
 
 	virtual uint64 RHIGetMinimumAlignmentForBufferBackedSRV(EPixelFormat Format) final override;
 
+	// Compute the hash of the state components of the PSO initializer for PSO Precaching (only hash data relevant for the RHI specific PSO)
+	virtual uint64 RHIComputeStatePrecachePSOHash(const FGraphicsPipelineStateInitializer& Initializer) final override;
+
+	// Compute the hash of the PSO initializer for PSO Precaching (only hash data relevant for the RHI specific PSO)
+	virtual uint64 RHIComputePrecachePSOHash(const FGraphicsPipelineStateInitializer& Initializer) final override;
+
+	// Check if PSO Initializers are the same used during PSO Precaching (only compare data relevant for the RHI specific PSO)
+	virtual bool RHIMatchPrecachePSOInitializers(const FGraphicsPipelineStateInitializer& LHS, const FGraphicsPipelineStateInitializer& RHS) final override;
+
 	void Cleanup();
 
 	void PurgeFramebufferFromCaches(GLuint Framebuffer);
