Android GLES: Remove some logspam

diff --git a/Engine/Source/Runtime/OpenGLDrv/Private/OpenGLProgramBinaryFileCache.cpp b/Engine/Source/Runtime/OpenGLDrv/Private/OpenGLProgramBinaryFileCache.cpp
--- a/Engine/Source/Runtime/OpenGLDrv/Private/OpenGLProgramBinaryFileCache.cpp
+++ b/Engine/Source/Runtime/OpenGLDrv/Private/OpenGLProgramBinaryFileCache.cpp
@@ -1054,14 +1054,7 @@ void FOpenGLProgramBinaryCache::CheckPendingGLProgramCreateRequests_internal()
 			It.RemoveCurrent();
 		}
 		float TimeTaken = (float)GMaxBinaryProgramLoadTimeMS - (TimeRemainingS * 1000.0f);
-		if (TimeRemainingS < 0.005f)
-		{
-			UE_LOG(LogRHI, Warning, TEXT("CheckPendingGLProgramCreateRequests : iter count = %d, time taken = %f ms (remaining %d)"), Count, TimeTaken, PendingGLContainerPrograms.Num());
-		}
-		else
-		{
-			UE_LOG(LogRHI, Verbose, TEXT("CheckPendingGLProgramCreateRequests : iter count = %d, time taken = %f ms (remaining %d)"), Count, TimeTaken, PendingGLContainerPrograms.Num());
-		}
+		UE_LOG(LogRHI, Verbose, TEXT("CheckPendingGLProgramCreateRequests : iter count = %d, time taken = %f ms (remaining %d)"), Count, TimeTaken, PendingGLContainerPrograms.Num());
 	}
 }
 
