IWYU fixup

diff --git a/Engine/Source/Runtime/OpenGLDrv/Private/Linux/OpenGLLinux.cpp b/Engine/Source/Runtime/OpenGLDrv/Private/Linux/OpenGLLinux.cpp
--- a/Engine/Source/Runtime/OpenGLDrv/Private/Linux/OpenGLLinux.cpp
+++ b/Engine/Source/Runtime/OpenGLDrv/Private/Linux/OpenGLLinux.cpp
@@ -12,6 +12,7 @@
 #include "ComponentReregisterContext.h"
 #include "Linux/LinuxPlatformApplicationMisc.h"
 #include "GenericPlatform/GenericPlatformFramePacer.h"
+#include "RHIUtilities.h"
 
 /*------------------------------------------------------------------------------
 	OpenGL function pointers.
