Fixed various RHI's to only clear targets with clear actions rather than all targets when any clear action is specified.

diff --git a/Engine/Source/Runtime/OpenGLDrv/Private/OpenGLCommands.cpp b/Engine/Source/Runtime/OpenGLDrv/Private/OpenGLCommands.cpp
--- a/Engine/Source/Runtime/OpenGLDrv/Private/OpenGLCommands.cpp
+++ b/Engine/Source/Runtime/OpenGLDrv/Private/OpenGLCommands.cpp
@@ -1762,22 +1762,28 @@ void FOpenGLDynamicRHI::SetRenderTargetsAndClear(const FRHISetRenderTargetsInfo&
 	bool bClearDepth = RenderTargetsInfo.bClearDepth;
 
 	FLinearColor ClearColors[MaxSimultaneousRenderTargets];
+	bool bClearColorArray[MaxSimultaneousRenderTargets];
 	float DepthClear = 0.0;
 	uint32 StencilClear = 0;
 
 	for (int32 i = 0; i < RenderTargetsInfo.NumColorRenderTargets; ++i)
 	{
+		bClearColorArray[i] = RenderTargetsInfo.ColorRenderTarget[i].LoadAction == ERenderTargetLoadAction::EClear;
+
 		if (RenderTargetsInfo.ColorRenderTarget[i].Texture != nullptr)
 		{
 			const FClearValueBinding& ClearValue = RenderTargetsInfo.ColorRenderTarget[i].Texture->GetClearBinding();
 
 			if (bIsTiledGPU)
 			{
-				bClearColor |= RenderTargetsInfo.ColorRenderTarget[i].LoadAction == ERenderTargetLoadAction::ENoAction;
+				const bool bNoAction = RenderTargetsInfo.ColorRenderTarget[i].LoadAction == ERenderTargetLoadAction::ENoAction;
+
+				bClearColor |= bNoAction;
+				bClearColorArray[i] |= bNoAction;
 
 				ClearColors[i] = ClearValue.ColorBinding == EClearBinding::EColorBound ? ClearValue.GetClearColor() : FLinearColor::Black;
 			}
-			else if(bClearColor)
+			else if(bClearColorArray[i])
 			{
 				checkf(ClearValue.ColorBinding == EClearBinding::EColorBound, TEXT("Texture: %s does not have a color bound for fast clears"), *RenderTargetsInfo.ColorRenderTarget[i].Texture->GetName().GetPlainNameString());
 
@@ -1811,7 +1817,7 @@ void FOpenGLDynamicRHI::SetRenderTargetsAndClear(const FRHISetRenderTargetsInfo&
 
 	if (bClearColor || bClearStencil || bClearDepth)
 	{
-		this->RHIClearMRT(bClearColor, RenderTargetsInfo.NumColorRenderTargets, ClearColors, bClearDepth, DepthClear, bClearStencil, StencilClear);
+		this->RHIClearMRT(bClearColor ? bClearColorArray : nullptr, RenderTargetsInfo.NumColorRenderTargets, ClearColors, bClearDepth, DepthClear, bClearStencil, StencilClear);
 	}
 }
 
@@ -2718,16 +2724,19 @@ static inline void ClearCurrentDepthStencilWithCurrentScissor( int8 ClearType, f
 	}
 }
 
-void FOpenGLDynamicRHI::ClearCurrentFramebufferWithCurrentScissor(FOpenGLContextState& ContextState, int8 ClearType, int32 NumClearColors, const FLinearColor* ClearColorArray, float Depth, uint32 Stencil)
+void FOpenGLDynamicRHI::ClearCurrentFramebufferWithCurrentScissor(FOpenGLContextState& ContextState, int8 ClearType, int32 NumClearColors, const bool* bClearColorArray, const FLinearColor* ClearColorArray, float Depth, uint32 Stencil)
 {
 	VERIFY_GL_SCOPE();
 		
 	// Clear color buffers
-	if (ClearType & CT_Color)
+	if (ClearType & CT_Color && bClearColorArray)
 	{
 		for(int32 ColorIndex = 0; ColorIndex < NumClearColors; ++ColorIndex)
 		{
-			FOpenGL::ClearBufferfv( GL_COLOR, ColorIndex, (const GLfloat*)&ClearColorArray[ColorIndex] );
+			if (bClearColorArray[ColorIndex])
+			{
+				FOpenGL::ClearBufferfv( GL_COLOR, ColorIndex, (const GLfloat*)&ClearColorArray[ColorIndex] );
+			}
 		}
 	}
 
@@ -2737,14 +2746,14 @@ void FOpenGLDynamicRHI::ClearCurrentFramebufferWithCurrentScissor(FOpenGLContext
 	}
 }
 
-void FOpenGLDynamicRHI::RHIClearMRT(bool bClearColor,int32 NumClearColors,const FLinearColor* ClearColorArray,bool bClearDepth,float Depth,bool bClearStencil,uint32 Stencil)
+void FOpenGLDynamicRHI::RHIClearMRT(const bool* bClearColorArray,int32 NumClearColors,const FLinearColor* ClearColorArray,bool bClearDepth,float Depth,bool bClearStencil,uint32 Stencil)
 {
 	FIntRect ExcludeRect;
 	VERIFY_GL_SCOPE();
 
 	check((GMaxRHIFeatureLevel >= ERHIFeatureLevel::SM5) || !PendingState.bFramebufferSetupInvalid);
 
-	if (bClearColor)
+	if (bClearColorArray)
 	{
 		// This is copied from DirectX11 code - apparently there's a silent assumption that there can be no valid render target set at index higher than an invalid one.
 		int32 NumActiveRenderTargets = 0;
@@ -2785,22 +2794,25 @@ void FOpenGLDynamicRHI::RHIClearMRT(bool bClearColor,int32 NumClearColors,const
 	int8 ClearType = CT_None;
 
 	// Prepare color buffer masks, if applicable
-	if (bClearColor)
+	if (bClearColorArray)
 	{
 		ClearType |= CT_Color;
 
 		for(int32 ColorIndex = 0; ColorIndex < NumClearColors; ++ColorIndex)
 		{
-			if( !ContextState.BlendState.RenderTargets[ColorIndex].ColorWriteMaskR ||
-				!ContextState.BlendState.RenderTargets[ColorIndex].ColorWriteMaskG ||
-				!ContextState.BlendState.RenderTargets[ColorIndex].ColorWriteMaskB ||
-				!ContextState.BlendState.RenderTargets[ColorIndex].ColorWriteMaskA)
+			if (bClearColorArray[ColorIndex])
 			{
-				FOpenGL::ColorMaskIndexed(ColorIndex, GL_TRUE, GL_TRUE, GL_TRUE, GL_TRUE);
-				ContextState.BlendState.RenderTargets[ColorIndex].ColorWriteMaskR = 1;
-				ContextState.BlendState.RenderTargets[ColorIndex].ColorWriteMaskG = 1;
-				ContextState.BlendState.RenderTargets[ColorIndex].ColorWriteMaskB = 1;
-				ContextState.BlendState.RenderTargets[ColorIndex].ColorWriteMaskA = 1;
+				if (!ContextState.BlendState.RenderTargets[ColorIndex].ColorWriteMaskR ||
+					!ContextState.BlendState.RenderTargets[ColorIndex].ColorWriteMaskG ||
+					!ContextState.BlendState.RenderTargets[ColorIndex].ColorWriteMaskB ||
+					!ContextState.BlendState.RenderTargets[ColorIndex].ColorWriteMaskA)
+				{
+					FOpenGL::ColorMaskIndexed(ColorIndex, GL_TRUE, GL_TRUE, GL_TRUE, GL_TRUE);
+					ContextState.BlendState.RenderTargets[ColorIndex].ColorWriteMaskR = 1;
+					ContextState.BlendState.RenderTargets[ColorIndex].ColorWriteMaskG = 1;
+					ContextState.BlendState.RenderTargets[ColorIndex].ColorWriteMaskB = 1;
+					ContextState.BlendState.RenderTargets[ColorIndex].ColorWriteMaskA = 1;
+				}
 			}
 		}
 	}
@@ -2830,7 +2842,7 @@ void FOpenGLDynamicRHI::RHIClearMRT(bool bClearColor,int32 NumClearColors,const
 	}
 
 	// Just one clear
-	ClearCurrentFramebufferWithCurrentScissor(ContextState, ClearType, NumClearColors, ClearColorArray, Depth, Stencil);
+	ClearCurrentFramebufferWithCurrentScissor(ContextState, ClearType, NumClearColors, bClearColorArray, ClearColorArray, Depth, Stencil);
 
 	if (bScissorChanged)
 	{
diff --git a/Engine/Source/Runtime/OpenGLDrv/Public/OpenGLDrv.h b/Engine/Source/Runtime/OpenGLDrv/Public/OpenGLDrv.h
--- a/Engine/Source/Runtime/OpenGLDrv/Public/OpenGLDrv.h
+++ b/Engine/Source/Runtime/OpenGLDrv/Public/OpenGLDrv.h
@@ -923,7 +923,7 @@ private:
 	void SetupUAVsForCompute(FOpenGLContextState& ContextState, const FOpenGLComputeShader* ComputeShader);
 	void SetupUAVsForProgram(FOpenGLContextState& ContextState, const TBitArray<>& NeededBits, int32 MaxUAVUnitUsed);
 
-	void RHIClearMRT(bool bClearColor, int32 NumClearColors, const FLinearColor* ColorArray, bool bClearDepth, float Depth, bool bClearStencil, uint32 Stencil);
+	void RHIClearMRT(const bool* bClearColorArray, int32 NumClearColors, const FLinearColor* ColorArray, bool bClearDepth, float Depth, bool bClearStencil, uint32 Stencil);
 
 public:
 	/** Remember what RHI user wants set on a specific OpenGL texture stage, translating from Stage and TextureIndex for stage pair. */
@@ -940,7 +940,7 @@ private:
 
 	void BindUniformBufferBase(FOpenGLContextState& ContextState, int32 NumUniformBuffers, FRHIUniformBuffer** BoundUniformBuffers, uint32* DynamicOffsets, uint32 FirstUniformBuffer, bool ForceUpdate);
 
-	void ClearCurrentFramebufferWithCurrentScissor(FOpenGLContextState& ContextState, int8 ClearType, int32 NumClearColors, const FLinearColor* ClearColorArray, float Depth, uint32 Stencil);
+	void ClearCurrentFramebufferWithCurrentScissor(FOpenGLContextState& ContextState, int8 ClearType, int32 NumClearColors, const bool* bClearColorArray, const FLinearColor* ClearColorArray, float Depth, uint32 Stencil);
 
 	FTextureLockTracker GLLockTracker;
 
diff --git a/Engine/Source/Runtime/RHI/Private/RHI.cpp b/Engine/Source/Runtime/RHI/Private/RHI.cpp
--- a/Engine/Source/Runtime/RHI/Private/RHI.cpp
+++ b/Engine/Source/Runtime/RHI/Private/RHI.cpp
@@ -1175,8 +1175,8 @@ static FAutoConsoleVariableRef CVarEnableVariableRateShading(
 static FAutoConsoleVariableRef CVarEnableAttachmentVariableRateShading(
 	TEXT("r.VRS.EnableImage"),
 	GRHIAttachmentVariableRateShadingEnabled,
-	TEXT("Enable image-based Variable Rate Shading."),
-	ECVF_ReadOnly);
+	TEXT("Toggle to enable image-based Variable Rate Shading."),
+	ECVF_RenderThreadSafe);
 
 
 FString GRHIBindlessResourceConfiguration = TEXT("Disabled");
diff --git a/Engine/Source/Runtime/Windows/D3D11RHI/Private/D3D11Commands.cpp b/Engine/Source/Runtime/Windows/D3D11RHI/Private/D3D11Commands.cpp
--- a/Engine/Source/Runtime/Windows/D3D11RHI/Private/D3D11Commands.cpp
+++ b/Engine/Source/Runtime/Windows/D3D11RHI/Private/D3D11Commands.cpp
@@ -1013,6 +1013,7 @@ void FD3D11DynamicRHI::SetRenderTargetsAndClear(const FRHISetRenderTargetsInfo&
 	if (RenderTargetsInfo.bClearColor || RenderTargetsInfo.bClearStencil || RenderTargetsInfo.bClearDepth)
 	{
 		FLinearColor ClearColors[MaxSimultaneousRenderTargets];
+		bool bClearColorArray[MaxSimultaneousRenderTargets];
 		float DepthClear = 0.0;
 		uint32 StencilClear = 0;
 
@@ -1020,7 +1021,9 @@ void FD3D11DynamicRHI::SetRenderTargetsAndClear(const FRHISetRenderTargetsInfo&
 		{
 			for (int32 i = 0; i < RenderTargetsInfo.NumColorRenderTargets; ++i)
 			{
-				if (RenderTargetsInfo.ColorRenderTarget[i].Texture != nullptr)
+				bClearColorArray[i] = RenderTargetsInfo.ColorRenderTarget[i].LoadAction == ERenderTargetLoadAction::EClear;
+
+				if (bClearColorArray[i] && RenderTargetsInfo.ColorRenderTarget[i].Texture != nullptr)
 				{
 					const FClearValueBinding& ClearValue = RenderTargetsInfo.ColorRenderTarget[i].Texture->GetClearBinding();
 					checkf(ClearValue.ColorBinding == EClearBinding::EColorBound, TEXT("Texture: %s does not have a color bound for fast clears"), *RenderTargetsInfo.ColorRenderTarget[i].Texture->GetName().GetPlainNameString());
@@ -1035,7 +1038,7 @@ void FD3D11DynamicRHI::SetRenderTargetsAndClear(const FRHISetRenderTargetsInfo&
 			ClearValue.GetDepthStencil(DepthClear, StencilClear);
 	}
 
-		this->RHIClearMRTImpl(RenderTargetsInfo.bClearColor, RenderTargetsInfo.NumColorRenderTargets, ClearColors, RenderTargetsInfo.bClearDepth, DepthClear, RenderTargetsInfo.bClearStencil, StencilClear);
+		this->RHIClearMRTImpl(RenderTargetsInfo.bClearColor ? bClearColorArray : nullptr, RenderTargetsInfo.NumColorRenderTargets, ClearColors, RenderTargetsInfo.bClearDepth, DepthClear, RenderTargetsInfo.bClearStencil, StencilClear);
 	}
 }
 
@@ -1353,18 +1356,12 @@ void FD3D11DynamicRHI::RHIDrawIndexedPrimitiveIndirect(FRHIBuffer* IndexBufferRH
 	EnableUAVOverlap();
 }
 
-// Raster operations.
-void FD3D11DynamicRHI::RHIClearMRT(bool bClearColor, int32 NumClearColors, const FLinearColor* ClearColorArray, bool bClearDepth, float Depth, bool bClearStencil, uint32 Stencil)
-{
-	RHIClearMRTImpl(bClearColor, NumClearColors, ClearColorArray, bClearDepth, Depth, bClearStencil, Stencil);
-}
-
-void FD3D11DynamicRHI::RHIClearMRTImpl(bool bClearColor, int32 NumClearColors, const FLinearColor* ClearColorArray, bool bClearDepth, float Depth, bool bClearStencil, uint32 Stencil)
+void FD3D11DynamicRHI::RHIClearMRTImpl(const bool* bClearColorArray, int32 NumClearColors, const FLinearColor* ClearColorArray, bool bClearDepth, float Depth, bool bClearStencil, uint32 Stencil)
 {
 	FD3D11BoundRenderTargets BoundRenderTargets(Direct3DDeviceIMContext);
 
 	// Must specify enough clear colors for all active RTs
-	check(!bClearColor || NumClearColors >= BoundRenderTargets.GetNumActiveTargets());
+	check(!bClearColorArray || NumClearColors >= BoundRenderTargets.GetNumActiveTargets());
 
 	// If we're clearing depth or stencil and we have a readonly depth/stencil view bound, we need to use a writable depth/stencil view
 	if (CurrentDepthTexture)
@@ -1378,14 +1375,17 @@ void FD3D11DynamicRHI::RHIClearMRTImpl(bool bClearColor, int32 NumClearColors, c
 
 	ID3D11DepthStencilView* DepthStencilView = BoundRenderTargets.GetDepthStencilView();
 
-	if (bClearColor && BoundRenderTargets.GetNumActiveTargets() > 0)
+	if (bClearColorArray && BoundRenderTargets.GetNumActiveTargets() > 0)
 	{
 		for (int32 TargetIndex = 0; TargetIndex < BoundRenderTargets.GetNumActiveTargets(); TargetIndex++)
-		{				
-			ID3D11RenderTargetView* RenderTargetView = BoundRenderTargets.GetRenderTargetView(TargetIndex);
-			if (RenderTargetView != nullptr)
+		{
+			if (bClearColorArray[TargetIndex])
 			{
-				Direct3DDeviceIMContext->ClearRenderTargetView(RenderTargetView, (float*)&ClearColorArray[TargetIndex]);
+				ID3D11RenderTargetView* RenderTargetView = BoundRenderTargets.GetRenderTargetView(TargetIndex);
+				if (RenderTargetView != nullptr)
+				{
+					Direct3DDeviceIMContext->ClearRenderTargetView(RenderTargetView, (float*)&ClearColorArray[TargetIndex]);
+				}
 			}
 		}
 	}
diff --git a/Engine/Source/Runtime/Windows/D3D11RHI/Private/D3D11RHIPrivate.h b/Engine/Source/Runtime/Windows/D3D11RHI/Private/D3D11RHIPrivate.h
--- a/Engine/Source/Runtime/Windows/D3D11RHI/Private/D3D11RHIPrivate.h
+++ b/Engine/Source/Runtime/Windows/D3D11RHI/Private/D3D11RHIPrivate.h
@@ -724,8 +724,6 @@ public:
 private:
 	void EnableDepthBoundsTest(bool bEnable, float MinDepth, float MaxDepth);
 
-	void RHIClearMRT(bool bClearColor, int32 NumClearColors, const FLinearColor* ColorArray, bool bClearDepth, float Depth, bool bClearStencil, uint32 Stencil);
-
 	static void ClearUAV(TRHICommandList_RecursiveHazardous<FD3D11DynamicRHI>& RHICmdList, FD3D11UnorderedAccessView* UAV, const void* ClearValues, bool bFloat);
 
 	enum class EForceFullScreenClear
@@ -734,7 +732,7 @@ private:
 		EForce
 	};
 
-	virtual void RHIClearMRTImpl(bool bClearColor, int32 NumClearColors, const FLinearColor* ColorArray, bool bClearDepth, float Depth, bool bClearStencil, uint32 Stencil);
+	virtual void RHIClearMRTImpl(const bool* bClearColorArray, int32 NumClearColors, const FLinearColor* ColorArray, bool bClearDepth, float Depth, bool bClearStencil, uint32 Stencil);
 
 	template <EShaderFrequency ShaderFrequency>
 	void ClearShaderResourceViews(FD3D11ViewableResource* Resource);
