Android GL: Set GPU time if swappy is in use.

diff --git a/Engine/Source/Runtime/OpenGLDrv/Private/Android/AndroidOpenGLFramePacer.cpp b/Engine/Source/Runtime/OpenGLDrv/Private/Android/AndroidOpenGLFramePacer.cpp
--- a/Engine/Source/Runtime/OpenGLDrv/Private/Android/AndroidOpenGLFramePacer.cpp
+++ b/Engine/Source/Runtime/OpenGLDrv/Private/Android/AndroidOpenGLFramePacer.cpp
@@ -89,6 +89,26 @@ void FAndroidOpenGLFramePacer::Init()
 }
 
 #if USE_ANDROID_OPENGL_SWAPPY
+
+namespace AndroidGL
+{
+	void SwappyPostWaitCallback(void*, int64_t cpu_time_ns, int64_t gpu_time_ns)
+	{
+		const double Frequency = 1.0;// FGPUTiming::GetTimingFrequency();
+		const double CyclesPerSecond = 1.0 / (Frequency * FPlatformTime::GetSecondsPerCycle64());
+		const double GPUTimeInSeconds = (double)gpu_time_ns / 1000000000.0;
+
+		GetDynamicRHI<FOpenGLDynamicRHI>()->RHISetExternalGPUTime(CyclesPerSecond * GPUTimeInSeconds);
+	}
+
+	void SetSwappyPostWaitCallback()
+	{
+		SwappyTracer Tracer = { 0 };
+		Tracer.postWait = AndroidGL::SwappyPostWaitCallback;
+		SwappyGL_injectTracer(&Tracer);
+	}
+};
+
 void FAndroidOpenGLFramePacer::InitSwappy()
 {
 	if (!bSwappyInit)
@@ -103,6 +123,8 @@ void FAndroidOpenGLFramePacer::InitSwappy()
 		{
 			FPlatformMisc::LowLevelOutputDebugStringf(TEXT("Init Swappy: version %d"), Swappy_version());
 			SwappyGL_init(Env, FJavaWrapper::GameActivityThis);
+
+			AndroidGL::SetSwappyPostWaitCallback();
 		}
 		bSwappyInit = true;
 	}
