Fixed: occasional freeze when loading into the world
OpenGLRHI specific issue where RHIMapStagingSurface could cause deadlock betwen RT and RHIT depending on supplied Fence state

diff --git a/Engine/Source/Runtime/OpenGLDrv/Private/OpenGLQuery.cpp b/Engine/Source/Runtime/OpenGLDrv/Private/OpenGLQuery.cpp
--- a/Engine/Source/Runtime/OpenGLDrv/Private/OpenGLQuery.cpp
+++ b/Engine/Source/Runtime/OpenGLDrv/Private/OpenGLQuery.cpp
@@ -1229,12 +1229,23 @@ bool FOpenGLGPUFence::Poll() const
 		return true;
 	}
 
-	RunOnGLRenderContextThread([Proxy = Proxy]()
+	if (!(IsInRenderingThread() || IsInRHIThread()))
 	{
-		VERIFY_GL_SCOPE();
-		check(Proxy != nullptr);
-		Proxy->Poll();
-	});
+		ENQUEUE_RENDER_COMMAND(FOpenGLGPUFence_Poll)(
+		[this](FRHICommandListImmediate& RHICmdList)
+		{
+			Poll();
+		});
+	}
+	else
+	{
+		RunOnGLRenderContextThread([Proxy = Proxy]()
+		{
+			VERIFY_GL_SCOPE();
+			check(Proxy != nullptr);
+			Proxy->Poll();
+		});
+	}
 
 	return Proxy->bIsSignaled;
 }
diff --git a/Engine/Source/Runtime/OpenGLDrv/Private/OpenGLRenderTarget.cpp b/Engine/Source/Runtime/OpenGLDrv/Private/OpenGLRenderTarget.cpp
--- a/Engine/Source/Runtime/OpenGLDrv/Private/OpenGLRenderTarget.cpp
+++ b/Engine/Source/Runtime/OpenGLDrv/Private/OpenGLRenderTarget.cpp
@@ -811,39 +811,48 @@ void FOpenGLDynamicRHI::RHIReadSurfaceData(FRHITexture* TextureRHI, FIntRect Rec
 	RHITHREAD_GLCOMMAND_EPILOGUE();
 }
 
-void FOpenGLDynamicRHI::RHIMapStagingSurface(FRHITexture* TextureRHI, FRHIGPUFence* FenceRHI, void*& OutData, int32& OutWidth, int32& OutHeight, uint32 GPUIndex)
+void FOpenGLDynamicRHI::RHIMapStagingSurface_RenderThread(class FRHICommandListImmediate& RHICmdList, FRHITexture* TextureRHI, uint32 GPUIndex, FRHIGPUFence* Fence, void*& OutData, int32& OutWidth, int32& OutHeight)
 {
-	FRHICommandListImmediate& RHICmdList = FRHICommandListExecutor::GetImmediateCommandList();
+	// Fence is not important, GL driver handles synchonization
+	RHICmdList.ImmediateFlush(EImmediateFlushType::FlushRHIThread);
 
 	RHITHREAD_GLCOMMAND_PROLOGUE();
 
 	VERIFY_GL_SCOPE();
-	
+
 	FOpenGLTexture* Texture = ResourceCast(TextureRHI->GetTexture2D());
 	check(Texture);
 	check(EnumHasAnyFlags(Texture->GetDesc().Flags, TexCreate_CPUReadback));
 
 	OutWidth = Texture->GetSizeX();
 	OutHeight = Texture->GetSizeY();
-	
+
 	uint32 Stride = 0;
 	OutData = Texture->Lock( 0, 0, RLM_ReadOnly, Stride );
 	RHITHREAD_GLCOMMAND_EPILOGUE();
 }
 
-void FOpenGLDynamicRHI::RHIUnmapStagingSurface(FRHITexture* TextureRHI, uint32 GPUIndex)
+void FOpenGLDynamicRHI::RHIUnmapStagingSurface_RenderThread(class FRHICommandListImmediate& RHICmdList, FRHITexture* TextureRHI, uint32 GPUIndex)
 {
-	FRHICommandListImmediate& RHICmdList = FRHICommandListExecutor::GetImmediateCommandList();
-
-	RHITHREAD_GLCOMMAND_PROLOGUE();
+	RunOnGLRenderContextThread([TextureRHI = TextureRHI]()
+	{
+		VERIFY_GL_SCOPE();
+		FOpenGLTexture* Texture = ResourceCast(TextureRHI->GetTexture2D());
+		check(Texture);
+		Texture->Unlock( 0, 0 );
+	});
+}
 
-	VERIFY_GL_SCOPE();
-	
-	FOpenGLTexture* Texture = ResourceCast(TextureRHI->GetTexture2D());
-	check(Texture);
+void FOpenGLDynamicRHI::RHIMapStagingSurface(FRHITexture* TextureRHI, FRHIGPUFence* FenceRHI, void*& OutData, int32& OutWidth, int32& OutHeight, uint32 GPUIndex)
+{
+	// Everything is handled in RHIMapStagingSurface_RenderThread. This function should not be called directly.
+	checkNoEntry(); 
+}
 
-	Texture->Unlock( 0, 0 );
-	RHITHREAD_GLCOMMAND_EPILOGUE();
+void FOpenGLDynamicRHI::RHIUnmapStagingSurface(FRHITexture* TextureRHI, uint32 GPUIndex)
+{
+	// Everything is handled in RHIUnmapStagingSurface_RenderThread. This function should not be called directly.
+	checkNoEntry(); 
 }
 
 void FOpenGLDynamicRHI::RHIReadSurfaceFloatData(FRHITexture* TextureRHI,FIntRect Rect,TArray<FFloat16Color>& OutData,ECubeFace CubeFace,int32 ArrayIndex,int32 MipIndex)
diff --git a/Engine/Source/Runtime/OpenGLDrv/Public/OpenGLDrv.h b/Engine/Source/Runtime/OpenGLDrv/Public/OpenGLDrv.h
--- a/Engine/Source/Runtime/OpenGLDrv/Public/OpenGLDrv.h
+++ b/Engine/Source/Runtime/OpenGLDrv/Public/OpenGLDrv.h
@@ -497,6 +497,8 @@ public:
 	virtual void RHIUnlockStagingBuffer(FRHIStagingBuffer* StagingBuffer) final override;
 	virtual void* LockStagingBuffer_RenderThread(class FRHICommandListImmediate& RHICmdList, FRHIStagingBuffer* StagingBuffer, FRHIGPUFence* Fence, uint32 Offset, uint32 SizeRHI) final override;
 	virtual void UnlockStagingBuffer_RenderThread(class FRHICommandListImmediate& RHICmdList, FRHIStagingBuffer* StagingBuffer) final override;
+	virtual void RHIMapStagingSurface_RenderThread(class FRHICommandListImmediate& RHICmdList, FRHITexture* Texture, uint32 GPUIndex, FRHIGPUFence* Fence, void*& OutData, int32& OutWidth, int32& OutHeight) final override;
+	virtual void RHIUnmapStagingSurface_RenderThread(class FRHICommandListImmediate& RHICmdList, FRHITexture* Texture, uint32 GPUIndex) final override;
 
 
 	virtual FRenderQueryPoolRHIRef RHICreateRenderQueryPool(ERenderQueryType QueryType, uint32 NumQueries = UINT32_MAX) final override;
