Restore Android GLES debug message functionality.
specify '-OpenGLDebugLevel=[1-5]' on the command line to enable debugging in all configs except shipping.

diff --git a/Engine/Source/Runtime/OpenGLDrv/Private/Android/AndroidEGL.cpp b/Engine/Source/Runtime/OpenGLDrv/Private/Android/AndroidEGL.cpp
--- a/Engine/Source/Runtime/OpenGLDrv/Private/Android/AndroidEGL.cpp
+++ b/Engine/Source/Runtime/OpenGLDrv/Private/Android/AndroidEGL.cpp
@@ -102,6 +102,7 @@ struct AndroidESPImpl
 	GLuint ResolveFrameBuffer;
 	GLuint DummyFrameBuffer;
 	FPlatformRect CachedWindowRect;
+	bool bIsDebug = false;
 
 	AndroidESPImpl();
 };
@@ -738,22 +739,24 @@ void AndroidEGL::ReInit()
 	SetCurrentSharedContext();
 }
 
-void AndroidEGL::Init(APIVariant API, uint32 MajorVersion, uint32 MinorVersion, bool bDebug)
+void AndroidEGL::Init(APIVariant API, uint32 MajorVersion, uint32 MinorVersion)
 {
 	check(IsInGameThread());
-
+	const bool bDebug = IsOGLDebugOutputEnabled();
 	if (PImplData->Initalized)
 	{
+		ensure(bDebug == PImplData->bIsDebug); // if this fires you would need to tear down the previous context and recreate to honour the debug change.
 		return;
 	}
-	InitEGL(API);
 
+	InitEGL(API);
+	PImplData->bIsDebug = bDebug;
 	if (bSupportsKHRCreateContext)
 	{
 		const uint32 MaxElements = 13;
 		uint32 Flags = 0;
 
-		Flags |= bDebug ? EGL_CONTEXT_OPENGL_DEBUG_BIT_KHR : 0;
+		Flags |= PImplData->bIsDebug ? EGL_CONTEXT_OPENGL_DEBUG_BIT_KHR : 0;
 
 		ContextAttributes = new int[MaxElements];
 		uint32 Element = 0;
diff --git a/Engine/Source/Runtime/OpenGLDrv/Private/Android/AndroidEGL.h b/Engine/Source/Runtime/OpenGLDrv/Private/Android/AndroidEGL.h
--- a/Engine/Source/Runtime/OpenGLDrv/Private/Android/AndroidEGL.h
+++ b/Engine/Source/Runtime/OpenGLDrv/Private/Android/AndroidEGL.h
@@ -71,7 +71,7 @@ public:
 	bool IsInitialized();
 	void InitBackBuffer();
 	void DestroyBackBuffer();
-	void Init( APIVariant API, uint32 MajorVersion, uint32 MinorVersion, bool bDebug);
+	void Init( APIVariant API, uint32 MajorVersion, uint32 MinorVersion);
 	void ReInit();
 	void UnBind();
 	void UnBindRender();
diff --git a/Engine/Source/Runtime/OpenGLDrv/Private/Android/AndroidOpenGL.cpp b/Engine/Source/Runtime/OpenGLDrv/Private/Android/AndroidOpenGL.cpp
--- a/Engine/Source/Runtime/OpenGLDrv/Private/Android/AndroidOpenGL.cpp
+++ b/Engine/Source/Runtime/OpenGLDrv/Private/Android/AndroidOpenGL.cpp
@@ -1104,7 +1104,7 @@ void FAndroidMisc::GetValidTargetPlatforms(TArray<FString>& TargetPlatformNames)
 void FAndroidAppEntry::PlatformInit()
 {
 	// Try to create an ES3.2 EGL here for gpu queries and don't have to recreate the GL context.
-	AndroidEGL::GetInstance()->Init(AndroidEGL::AV_OpenGLES, 3, 2, false);
+	AndroidEGL::GetInstance()->Init(AndroidEGL::AV_OpenGLES, 3, 2);
 }
 
 void FAndroidAppEntry::ReleaseEGL()
diff --git a/Engine/Source/Runtime/OpenGLDrv/Private/Linux/OpenGLLinux.cpp b/Engine/Source/Runtime/OpenGLDrv/Private/Linux/OpenGLLinux.cpp
--- a/Engine/Source/Runtime/OpenGLDrv/Private/Linux/OpenGLLinux.cpp
+++ b/Engine/Source/Runtime/OpenGLDrv/Private/Linux/OpenGLLinux.cpp
@@ -147,11 +147,7 @@ void Linux_PlatformCreateDummyGLWindow( FPlatformOpenGLContext *OutContext )
  */
 bool Linux_PlatformOpenGLDebugCtx()
 {
-#if UE_BUILD_DEBUG
-	return ! FParse::Param(FCommandLine::Get(),TEXT("openglNoDebug"));
-#else
-	return FParse::Param(FCommandLine::Get(),TEXT("openglDebug"));
-#endif
+	return IsOGLDebugOutputEnabled();
 }
 
 /**
diff --git a/Engine/Source/Runtime/OpenGLDrv/Private/OpenGLDevice.cpp b/Engine/Source/Runtime/OpenGLDrv/Private/OpenGLDevice.cpp
--- a/Engine/Source/Runtime/OpenGLDrv/Private/OpenGLDevice.cpp
+++ b/Engine/Source/Runtime/OpenGLDrv/Private/OpenGLDevice.cpp
@@ -386,6 +386,9 @@ static const TCHAR* GetOpenGLDebugSeverityStringARB(GLenum Severity)
 	#ifndef GL_APIENTRY
 	#define GL_APIENTRY APIENTRY
 	#endif
+	#ifndef GL_DEBUG_SEVERITY_MEDIUM_ARB
+	#define GL_DEBUG_SEVERITY_MEDIUM_ARB GL_DEBUG_SEVERITY_LOW_ARB
+	#endif
 static void GL_APIENTRY OpenGLDebugMessageCallbackARB(
 #else
 static void APIENTRY OpenGLDebugMessageCallbackARB(
@@ -416,7 +419,7 @@ static void APIENTRY OpenGLDebugMessageCallbackARB(
 		if (!LogRHI.IsSuppressed(Verbosity))
 		{
 			FMsg::Logf(__FILE__, __LINE__, LogRHI.GetCategoryName(), Verbosity,
-				TEXT("[%s][%s][%s][%u] %s"),
+				TEXT("GL_DBG: [%s][%s][%s][%u] %s"),
 				SourceStr,
 				TypeStr,
 				SeverityStr,
@@ -719,7 +722,12 @@ void InitDebugContext()
 	// Set the debug output callback if the driver supports it.
 	VERIFY_GL(__FUNCTION__);
 	bool bDebugOutputInitialized = false;
-#if !ENABLE_VERIFY_GL
+	if (!IsOGLDebugOutputEnabled())
+	{
+		return;
+	}
+
+#if ENABLE_DEBUG_OUTPUT
 	#if defined(GL_ARB_debug_output)
 		if (glDebugMessageCallbackARB)
 		{
@@ -743,12 +751,16 @@ void InitDebugContext()
 			bDebugOutputInitialized = (glGetError() == GL_NO_ERROR);
 		}
 	#endif // GL_AMD_debug_output
-#endif // !ENABLE_VERIFY_GL
+ #endif // ENABLE_DEBUG_OUTPUT
 
 	if (!bDebugOutputInitialized)
 	{
 		UE_LOG(LogRHI,Warning,TEXT("OpenGL debug output extension not supported!"));
 	}
+	else
+	{
+		UE_LOG(LogRHI, Log, TEXT("OpenGL debug output extension enabled!"));
+	}
 
 	// this is to suppress feeding back of the debug markers and groups to the log, since those originate in the app anyways...
 #if ENABLE_OPENGL_DEBUG_GROUPS && defined(GL_ARB_debug_output) && GL_ARB_debug_output && GL_KHR_debug
@@ -769,6 +781,19 @@ void InitDebugContext()
 		glDebugMessageControlKHR(GL_DEBUG_SOURCE_APPLICATION_KHR, GL_DEBUG_TYPE_PUSH_GROUP_KHR, GL_DONT_CARE, 0, NULL, GL_FALSE);
 		glDebugMessageControlKHR(GL_DEBUG_SOURCE_APPLICATION_KHR, GL_DEBUG_TYPE_POP_GROUP_KHR, GL_DONT_CARE, 0, NULL, GL_FALSE);
 		glDebugMessageControlKHR(GL_DEBUG_SOURCE_API_KHR, GL_DEBUG_TYPE_OTHER_KHR, GL_DEBUG_SEVERITY_NOTIFICATION, 0, NULL, GL_FALSE);
+
+		GLenum Severity;
+		switch (GetOGLDebugOutputLevel())
+		{
+			case 5: Severity = GL_DONT_CARE; break;
+			case 4: Severity = GL_DEBUG_SEVERITY_NOTIFICATION; break;
+			case 3: Severity = GL_DEBUG_SEVERITY_LOW_ARB; break;
+			case 2: Severity = GL_DEBUG_SEVERITY_MEDIUM_ARB; break;
+			case 1:
+				[[fallthrough]];
+			default: Severity = GL_DEBUG_SEVERITY_HIGH_ARB; break;
+		}
+		glDebugMessageControlKHR(GL_DEBUG_SOURCE_API_KHR, GL_DONT_CARE, Severity, 0, NULL, GL_TRUE);
 		UE_LOG(LogRHI,Verbose,TEXT("disabling reporting back of debug groups and markers to the OpenGL debug output callback"));
 	}
 #endif
diff --git a/Engine/Source/Runtime/OpenGLDrv/Private/OpenGLDrv.cpp b/Engine/Source/Runtime/OpenGLDrv/Private/OpenGLDrv.cpp
--- a/Engine/Source/Runtime/OpenGLDrv/Private/OpenGLDrv.cpp
+++ b/Engine/Source/Runtime/OpenGLDrv/Private/OpenGLDrv.cpp
@@ -95,7 +95,7 @@ bool FOpenGLDynamicRHI::RHIEGLSupportsNoErrorContext() const
 
 void FOpenGLDynamicRHI::RHIInitEGLInstanceGLES2()
 {
-	AndroidEGL::GetInstance()->Init(AndroidEGL::AV_OpenGLES, 2, 0, false);
+	AndroidEGL::GetInstance()->Init(AndroidEGL::AV_OpenGLES, 2, 0);
 	AndroidEGL::GetInstance()->InitSurface(false, false);
 }
 
diff --git a/Engine/Source/Runtime/OpenGLDrv/Private/OpenGLUtil.cpp b/Engine/Source/Runtime/OpenGLDrv/Private/OpenGLUtil.cpp
--- a/Engine/Source/Runtime/OpenGLDrv/Private/OpenGLUtil.cpp
+++ b/Engine/Source/Runtime/OpenGLDrv/Private/OpenGLUtil.cpp
@@ -10,6 +10,42 @@
 #include "OpenGLDrvPrivate.h"
 #include "RHICoreStats.h"
 
+#if ENABLE_DEBUG_OUTPUT
+
+// Enable GL debug output
+static int32 GOGLDebugOutputLevel = -1;
+
+bool IsOGLDebugOutputEnabled()
+{
+	return GetOGLDebugOutputLevel() != 0;
+}
+
+int32 GetOGLDebugOutputLevel()
+{
+	if (GOGLDebugOutputLevel < 0)
+	{
+		// this can happen super early
+		if(FCommandLine::IsInitialized())
+		{
+			int32 DebugLevel = 0;
+			GOGLDebugOutputLevel = FParse::Value(FCommandLine::Get(), TEXT("OpenGLDebugLevel="), DebugLevel) ? DebugLevel : 0;
+
+			if(FParse::Param(FCommandLine::Get(), TEXT("openglDebug")))
+			{
+				GOGLDebugOutputLevel = 1;
+			}
+		}
+		else
+		{
+			return 0;
+		}
+	}
+
+	return GOGLDebugOutputLevel;
+}
+#endif
+
+
 void VerifyOpenGLResult(GLenum ErrorCode, const TCHAR* Msg1, const TCHAR* Msg2, const TCHAR* Filename, uint32 Line)
 {
 	if (ErrorCode != GL_NO_ERROR)
diff --git a/Engine/Source/Runtime/OpenGLDrv/Private/Windows/OpenGLWindows.cpp b/Engine/Source/Runtime/OpenGLDrv/Private/Windows/OpenGLWindows.cpp
--- a/Engine/Source/Runtime/OpenGLDrv/Private/Windows/OpenGLWindows.cpp
+++ b/Engine/Source/Runtime/OpenGLDrv/Private/Windows/OpenGLWindows.cpp
@@ -217,11 +217,7 @@ static void GetOpenGLVersionForCoreProfile(int& OutMajorVersion, int& OutMinorVe
  */
 static bool PlatformOpenGLDebugCtx()
 {
-#if UE_BUILD_DEBUG
-	return ! FParse::Param(FCommandLine::Get(),TEXT("openglNoDebug"));
-#else
-	return FParse::Param(FCommandLine::Get(),TEXT("openglDebug"));
-#endif
+	return IsOGLDebugOutputEnabled();
 }
 
 
diff --git a/Engine/Source/Runtime/OpenGLDrv/Public/OpenGLUtil.h b/Engine/Source/Runtime/OpenGLDrv/Public/OpenGLUtil.h
--- a/Engine/Source/Runtime/OpenGLDrv/Public/OpenGLUtil.h
+++ b/Engine/Source/Runtime/OpenGLDrv/Public/OpenGLUtil.h
@@ -11,6 +11,17 @@
 #define ENABLE_VERIFY_GL (0 & DO_CHECK)
 #define ENABLE_VERIFY_GL_TRACE 0
 
+// Include GL debug output functionality on everything but shipping configs.
+// to enable the debug output specify '-OpenGLDebugLevel=[1-5]' via the command line.
+#define ENABLE_DEBUG_OUTPUT	(!UE_BUILD_SHIPPING)
+#if !ENABLE_DEBUG_OUTPUT
+inline bool IsOGLDebugOutputEnabled() { return false; }
+inline int32 GetOGLDebugOutputLevel() { return 0; }
+#else
+bool IsOGLDebugOutputEnabled();
+int32 GetOGLDebugOutputLevel();
+#endif
+
 // Additional check that our GL calls are occurring on the expected thread
 #define ENABLE_VERIFY_GL_THREAD (UE_BUILD_DEBUG)
 
