Android GLES, Clear buffer resource when ownership is released.

diff --git a/Engine/Source/Runtime/OpenGLDrv/Public/OpenGLResources.h b/Engine/Source/Runtime/OpenGLDrv/Public/OpenGLResources.h
--- a/Engine/Source/Runtime/OpenGLDrv/Public/OpenGLResources.h
+++ b/Engine/Source/Runtime/OpenGLDrv/Public/OpenGLResources.h
@@ -494,11 +494,6 @@ public:
 
 		if (Resource != 0)
 		{
-			if (BaseType::OnDelete(Resource, BaseType::GetSize(), false, 0))
-			{
-				FOpenGL::DeleteBuffers(1, &Resource);
-			}
-
 			if (LockBuffer != NULL)
 			{
 				if (bLockBufferWasAllocated)
@@ -511,6 +506,12 @@ public:
 				}
 			}
 
+			if (BaseType::OnDelete(Resource, BaseType::GetSize(), false, 0))
+			{
+				FOpenGL::DeleteBuffers(1, &Resource);
+				Resource = 0;
+			}
+
 			LockBuffer = nullptr;
 			OpenGLBufferStats::UpdateBufferStats(BaseType::GetDesc(), false);
 
