Fixed case where GlslCodeString was pointing to an invalid memory

diff --git a/Engine/Source/Runtime/OpenGLDrv/Private/OpenGLShaders.cpp b/Engine/Source/Runtime/OpenGLDrv/Private/OpenGLShaders.cpp
--- a/Engine/Source/Runtime/OpenGLDrv/Private/OpenGLShaders.cpp
+++ b/Engine/Source/Runtime/OpenGLDrv/Private/OpenGLShaders.cpp
@@ -1054,7 +1054,7 @@ FOpenGLShader::FOpenGLShader(TArrayView<const uint8> Code, const FSHAHash& Hash,
 		check(GlslCodeFinal.Num());
 	}
 	GlslCode = MoveTemp(GlslCodeFinal);
-	GlslCodeString = GlslCodeFinal.GetData();
+	GlslCodeString = GlslCode.GetData();
 #endif
 
 	// The shader is compiled when we link program
