Refactoring OpenGLRHI binding to not go through Context methods.
- Writing custom version of UE::RHICore::ApplyStaticUniformBuffers, UE::RHICore::RHISetShaderParametersShared, and UE::RHICore::RHISetShaderUnbindsShared for OpenGLRHI to avoid redundant overhead.
- This removed the last uses of RHISetShaderTexture, RHISetShaderSampler, RHISetUAVParameter, RHISetShaderResourceViewParameter, and RHISetShaderParameter from OpenGLRHI. Those functions are now gone.

diff --git a/Engine/Source/Runtime/OpenGLDrv/Private/OpenGLCommands.cpp b/Engine/Source/Runtime/OpenGLDrv/Private/OpenGLCommands.cpp
--- a/Engine/Source/Runtime/OpenGLDrv/Private/OpenGLCommands.cpp
+++ b/Engine/Source/Runtime/OpenGLDrv/Private/OpenGLCommands.cpp
@@ -165,73 +165,29 @@ static FORCEINLINE GLint ModifyFilterByMips(GLint Filter, bool bHasMips)
 	return Filter;
 }
 
-static FORCEINLINE EShaderFrequency GetShaderFrequency(FRHIGraphicsShader* ShaderRHI)
+static FORCEINLINE bool ValidateShader(FRHIComputeShader* ShaderRHI)
 {
-	switch (ShaderRHI->GetFrequency())
-	{
-	case SF_Vertex:
-		VALIDATE_BOUND_SHADER(ShaderRHI, Vertex);
-		return SF_Vertex;
-#if PLATFORM_SUPPORTS_GEOMETRY_SHADERS
-	case SF_Geometry:
-		VALIDATE_BOUND_SHADER(ShaderRHI, Geometry);
-		return SF_Geometry;
-#endif
-	case SF_Pixel:
-		VALIDATE_BOUND_SHADER(ShaderRHI, Pixel);
-		return SF_Pixel;
-	default:
-		checkf(0, TEXT("Undefined FRHIShader Type %d!"), (int32)ShaderRHI->GetFrequency());
-	}
-
-	return SF_NumFrequencies;
-}
-
-static FORCEINLINE CrossCompiler::EShaderStage GetShaderCrossCompilerStage(FRHIGraphicsShader* ShaderRHI)
-{
-	switch (ShaderRHI->GetFrequency())
-	{
-	case SF_Vertex:
-		VALIDATE_BOUND_SHADER(ShaderRHI, Vertex);
-		return CrossCompiler::SHADER_STAGE_VERTEX;
-#if PLATFORM_SUPPORTS_GEOMETRY_SHADERS
-	case SF_Geometry:
-		VALIDATE_BOUND_SHADER(ShaderRHI, Geometry);
-		return CrossCompiler::SHADER_STAGE_GEOMETRY;
-#endif
-	case SF_Pixel:
-		VALIDATE_BOUND_SHADER(ShaderRHI, Pixel);
-		return CrossCompiler::SHADER_STAGE_PIXEL;
-	default:
-		checkf(0, TEXT("Undefined FRHIShader Type %d!"), (int32)ShaderRHI->GetFrequency());
-	}
-
-	return CrossCompiler::NUM_SHADER_STAGES;
+	return true;
 }
 
-static FORCEINLINE void GetShaderStageIndexAndMaxUnits(FRHIGraphicsShader* ShaderRHI, GLint& OutIndex, GLint& OutMaxUnits)
+static FORCEINLINE bool ValidateShader(FRHIGraphicsShader* ShaderRHI)
 {
 	switch (ShaderRHI->GetFrequency())
 	{
 	case SF_Vertex:
 		VALIDATE_BOUND_SHADER(ShaderRHI, Vertex);
-		OutIndex = FOpenGL::GetFirstVertexTextureUnit();
-		OutMaxUnits = FOpenGL::GetMaxVertexTextureImageUnits();
-		break;
+		return true;
 #if PLATFORM_SUPPORTS_GEOMETRY_SHADERS
 	case SF_Geometry:
 		VALIDATE_BOUND_SHADER(ShaderRHI, Geometry);
-		OutIndex = FOpenGL::GetFirstGeometryTextureUnit();
-		OutMaxUnits = FOpenGL::GetMaxGeometryTextureImageUnits();
-		break;
+		return true;
 #endif
 	case SF_Pixel:
 		VALIDATE_BOUND_SHADER(ShaderRHI, Pixel);
-		OutIndex = FOpenGL::GetFirstPixelTextureUnit();
-		OutMaxUnits = FOpenGL::GetMaxTextureImageUnits();
-		break;
+		return true;
 	default:
 		checkf(0, TEXT("Undefined FRHIShader Type %d!"), (int32)ShaderRHI->GetFrequency());
+		return false;
 	}
 }
 
@@ -416,47 +372,6 @@ void FOpenGLDynamicRHI::RHISetBoundShaderState(FRHIBoundShaderState* BoundShader
 	BoundShaderStateHistory.Add(BoundShaderState);
 }
 
-void FOpenGLDynamicRHI::RHISetUAVParameter(FRHIPixelShader* PixelShaderRHI, uint32 UAVIndex, FRHIUnorderedAccessView* UnorderedAccessViewRHI)
-{
-	checkNoEntry();//UAV-PS port: not yet implemented
-}
-
-
-void FOpenGLDynamicRHI::RHISetUAVParameter(FRHIComputeShader* ComputeShaderRHI, uint32 InUAVIndex, FRHIUnorderedAccessView* UnorderedAccessViewRHI)
-{
-	VERIFY_GL_SCOPE();
-
-	GLint UAVIndex = FOpenGL::GetFirstComputeUAVUnit() + InUAVIndex;
-	
-	if (UnorderedAccessViewRHI == nullptr)
-	{
-		InternalSetShaderBufferUAV(UAVIndex, 0);
-		return;
-	}
-
-	FOpenGLUnorderedAccessView* UnorderedAccessView = ResourceCast(UnorderedAccessViewRHI);
-	if (UnorderedAccessView->Resource)
-	{
-		GLuint	Resource = UnorderedAccessView->Resource;
-		GLenum	Format = UnorderedAccessView->Format;
-		bool	bLayered = UnorderedAccessView->IsLayered();
-		GLint	Layer = UnorderedAccessView->GetLayer();
-		GLenum	Access = GL_READ_WRITE;
-		InternalSetShaderImageUAV(UAVIndex, Format, Resource, bLayered, Layer, Access);
-	}
-	else
-	{
-		GLuint Resource = UnorderedAccessView->BufferResource;
-		InternalSetShaderBufferUAV(UAVIndex, Resource);
-	}
-}
-
-void FOpenGLDynamicRHI::RHISetUAVParameter(FRHIComputeShader* ComputeShaderRHI,uint32 UAVIndex, FRHIUnorderedAccessView* UAVRHI, uint32 InitialCount )
-{
-	// TODO: Implement for OpenGL
-	check(0);
-}
-
 void FOpenGLDynamicRHI::InternalSetShaderTexture(FOpenGLTexture* Texture, FOpenGLShaderResourceView* SRV, GLint TextureIndex, GLenum Target, GLuint Resource, int NumMips, int LimitMip)
 {
 	auto& PendingTextureState = PendingState.Textures[TextureIndex];
@@ -976,120 +891,6 @@ void FOpenGLDynamicRHI::UpdateSRV(FOpenGLShaderResourceView* SRV)
 #endif
 }
 
-void FOpenGLDynamicRHI::RHISetShaderResourceViewParameter(FRHIGraphicsShader* ShaderRHI,uint32 TextureIndex, FRHIShaderResourceView* SRVRHI)
-{
-	VERIFY_GL_SCOPE();
-
-	GLint Index = 0;
-	GLint MaxUnits = 0;
-	GetShaderStageIndexAndMaxUnits(ShaderRHI, Index, MaxUnits);
-
-	ensureMsgf((int32)TextureIndex < MaxUnits, TEXT("Using more texture units (%d) than allowed (%d) on Frequency %d!"), TextureIndex, MaxUnits, (int32)ShaderRHI->GetFrequency());
-	FOpenGLShaderResourceView* SRV = ResourceCast(SRVRHI);
-	
-	GLuint Resource = 0;
-	GLenum Target = GL_TEXTURE_BUFFER;
-	int32 LimitMip = -1;
-	if (SRV)
-	{
-		Target = SRV->Target;
-		
-		if (Target == GL_SHADER_STORAGE_BUFFER)
-		{
-			Index = (ShaderRHI->GetFrequency() == SF_Pixel) ? FOpenGL::GetFirstPixelUAVUnit() :  FOpenGL::GetFirstVertexUAVUnit();
-			InternalSetShaderBufferUAV(Index + TextureIndex, SRV->Resource);
-			return;
-		}
-		else
-		{
-			Resource = SRV->Resource;
-			LimitMip = SRV->LimitMip;
-			UpdateSRV(SRV);
-		}
-	}
-
-	ensureMsgf((int32)TextureIndex < MaxUnits, TEXT("Using more textures (%d) than allowed (%d)!"), TextureIndex, MaxUnits);
-	InternalSetShaderTexture(NULL, SRV, Index + TextureIndex, Target, Resource, 0, LimitMip);
-	RHISetShaderSampler(ShaderRHI, TextureIndex, PointSamplerState);
-}
-
-void FOpenGLDynamicRHI::RHISetShaderResourceViewParameter(FRHIComputeShader* ComputeShaderRHI,uint32 TextureIndex, FRHIShaderResourceView* SRVRHI)
-{
-	VERIFY_GL_SCOPE();
-
-	FOpenGLShaderResourceView* SRV = ResourceCast(SRVRHI);
-	GLuint Resource = 0;
-	GLenum Target = GL_TEXTURE_BUFFER;
-	int32 LimitMip = -1;
-	if (SRV)
-	{
-		Target = SRV->Target;
-		
-		if (Target == GL_SHADER_STORAGE_BUFFER)
-		{
-			InternalSetShaderBufferUAV(FOpenGL::GetFirstComputeUAVUnit() + TextureIndex, SRV->Resource);
-			return;
-		}
-		else
-		{
-			Resource = SRV->Resource;
-			LimitMip = SRV->LimitMip;
-			UpdateSRV(SRV);
-		}
-	}
-
-	ensureMsgf((int32)TextureIndex < FOpenGL::GetMaxComputeTextureImageUnits(), TEXT("Using more compute texture units (%d) than allowed (%d)!"), TextureIndex, FOpenGL::GetMaxComputeTextureImageUnits());
-	InternalSetShaderTexture(NULL, SRV, FOpenGL::GetFirstComputeTextureUnit() + TextureIndex, Target, Resource, 0, LimitMip);
-	RHISetShaderSampler(ComputeShaderRHI,TextureIndex,PointSamplerState);
-}
-
-void FOpenGLDynamicRHI::RHISetShaderTexture(FRHIGraphicsShader* ShaderRHI,uint32 TextureIndex, FRHITexture* NewTextureRHI)
-{
-	VERIFY_GL_SCOPE();
-	FOpenGLTexture* NewTexture = ResourceCast(NewTextureRHI);
-
-	GLint Index = 0;
-	GLint MaxUnits = 0;
-	GetShaderStageIndexAndMaxUnits(ShaderRHI, Index, MaxUnits);
-
-	ensureMsgf((int32)TextureIndex < MaxUnits, TEXT("Using more texture units (%d) than allowed (%d) on Frequency %d!"), TextureIndex, MaxUnits, (int32)ShaderRHI->GetFrequency());
-	if (NewTexture)
-	{
-		InternalSetShaderTexture(NewTexture, nullptr, Index + TextureIndex, NewTexture->Target, NewTexture->GetResource(), NewTextureRHI->GetNumMips(), -1);
-	}
-	else
-	{
-		InternalSetShaderTexture(nullptr, nullptr, Index + TextureIndex, 0, 0, 0, -1);
-	}
-}
-
-void FOpenGLDynamicRHI::RHISetShaderSampler(FRHIGraphicsShader* ShaderRHI,uint32 SamplerIndex, FRHISamplerState* NewStateRHI)
-{
-	VERIFY_GL_SCOPE();
-	FOpenGLSamplerState* NewState = ResourceCast(NewStateRHI);
-
-	GLint Index = 0;
-	GLint MaxUnits = 0;
-	GetShaderStageIndexAndMaxUnits(ShaderRHI, Index, MaxUnits);
-
-	InternalSetSamplerStates(Index + SamplerIndex, NewState);
-}
-
-void FOpenGLDynamicRHI::RHISetShaderTexture(FRHIComputeShader* ComputeShaderRHI,uint32 TextureIndex, FRHITexture* NewTextureRHI)
-{
-	VERIFY_GL_SCOPE();
-	FOpenGLTexture* NewTexture = ResourceCast(NewTextureRHI);
-	ensureMsgf((int32)TextureIndex < FOpenGL::GetMaxComputeTextureImageUnits(), TEXT("Using more compute texture units (%d) than allowed (%d)!"), TextureIndex, FOpenGL::GetMaxComputeTextureImageUnits());
-	if (NewTexture)
-	{
-		InternalSetShaderTexture(NewTexture, nullptr, FOpenGL::GetFirstComputeTextureUnit() + TextureIndex, NewTexture->Target, NewTexture->GetResource(), NewTextureRHI->GetNumMips(), -1);
-	}
-	else
-	{
-		InternalSetShaderTexture(nullptr, nullptr, FOpenGL::GetFirstComputeTextureUnit() + TextureIndex, 0, 0, 0, -1);
-	}
-}
-
 void FOpenGLDynamicRHI::RHISetStaticUniformBuffers(const FUniformBufferStaticBindings& InUniformBuffers)
 {
 	FMemory::Memzero(GlobalUniformBuffers.GetData(), GlobalUniformBuffers.Num() * sizeof(FRHIUniformBuffer*));
@@ -1152,94 +953,278 @@ void FOpenGLDynamicRHI::RHISetUniformBufferDynamicOffset(FUniformBufferStaticSlo
 	}
 }
 
-void FOpenGLDynamicRHI::RHISetShaderUniformBuffer(FRHIGraphicsShader* ShaderRHI,uint32 BufferIndex, FRHIUniformBuffer* BufferRHI)
+void FOpenGLDynamicRHI::BindUniformBuffer(EShaderFrequency ShaderFrequency, uint32 BufferIndex, FRHIUniformBuffer* BufferRHI)
 {
-	VERIFY_GL_SCOPE();
-	EShaderFrequency Stage = GetShaderFrequency(ShaderRHI);
-	if (Stage != SF_NumFrequencies)
+	PendingState.BoundUniformBuffers[ShaderFrequency][BufferIndex] = BufferRHI;
+	PendingState.BoundUniformBuffersDynamicOffset[ShaderFrequency][BufferIndex] = 0u;
+	PendingState.DirtyUniformBuffers[ShaderFrequency] |= 1 << BufferIndex;
+	PendingState.bAnyDirtyGraphicsUniformBuffers = true;
+
+	if (!GUseEmulatedUniformBuffers || !((FOpenGLUniformBuffer*)BufferRHI)->bIsEmulatedUniformBuffer)
 	{
-		PendingState.BoundUniformBuffers[Stage][BufferIndex] = BufferRHI;
-		PendingState.BoundUniformBuffersDynamicOffset[Stage][BufferIndex] = 0u;
-		PendingState.DirtyUniformBuffers[Stage] |= 1 << BufferIndex;
-		PendingState.bAnyDirtyGraphicsUniformBuffers = true;
-		
-		if (!GUseEmulatedUniformBuffers || !((FOpenGLUniformBuffer*)BufferRHI)->bIsEmulatedUniformBuffer)
-		{
-			PendingState.bAnyDirtyRealUniformBuffers[Stage] = true;
-		}
+		PendingState.bAnyDirtyRealUniformBuffers[ShaderFrequency] = true;
 	}
 }
 
-void FOpenGLDynamicRHI::RHISetShaderSampler(FRHIComputeShader* ComputeShaderRHI,uint32 SamplerIndex, FRHISamplerState* NewStateRHI)
+struct FOpenGLResourceBinder
 {
-	VERIFY_GL_SCOPE();
-	FOpenGLSamplerState* NewState = ResourceCast(NewStateRHI);
-	InternalSetSamplerStates(FOpenGL::GetFirstComputeTextureUnit() + SamplerIndex, NewState);
-}
+	FOpenGLDynamicRHI& RHI;
+	EShaderFrequency ShaderFrequency;
+	CrossCompiler::EShaderStage ShaderStage = CrossCompiler::NUM_SHADER_STAGES;
+	GLint FirstTextureUnit = 0;
+	GLint MaxTextureImageUnits = 0;
+	GLint FirstUAVUnit = 0;
+	GLint MaxUAVUnits = 0;
+	bool bClearTextureSamplers = false;
+
+	FOpenGLResourceBinder() = delete;
+	FOpenGLResourceBinder(FOpenGLDynamicRHI& InRHI, EShaderFrequency InShaderFrequency, bool bInClearTextureSamplers = false)
+		: RHI(InRHI)
+		, ShaderFrequency(InShaderFrequency)
+		, bClearTextureSamplers(bInClearTextureSamplers)
+	{
+		Init();
+	}
+
+	void Init()
+	{
+		switch (ShaderFrequency)
+		{
+		case SF_Vertex:
+			ShaderStage = CrossCompiler::SHADER_STAGE_VERTEX;
+			FirstTextureUnit = FOpenGL::GetFirstVertexTextureUnit();
+			MaxTextureImageUnits = FOpenGL::GetMaxVertexTextureImageUnits();
+			FirstUAVUnit = FOpenGL::GetFirstVertexUAVUnit();
+			MaxUAVUnits = FOpenGL::GetMaxPixelUAVUnits();
+			break;
+#if PLATFORM_SUPPORTS_GEOMETRY_SHADERS
+		case SF_Geometry:
+			ShaderStage = CrossCompiler::SHADER_STAGE_GEOMETRY;
+			FirstTextureUnit = FOpenGL::GetFirstGeometryTextureUnit();
+			MaxTextureImageUnits = FOpenGL::GetMaxGeometryTextureImageUnits();
+			FirstUAVUnit = 0;
+			MaxUAVUnits = 0;
+			break;
+#endif
+		case SF_Pixel:
+			ShaderStage = CrossCompiler::SHADER_STAGE_PIXEL;
+			FirstTextureUnit = FOpenGL::GetFirstPixelTextureUnit();
+			MaxTextureImageUnits = FOpenGL::GetMaxTextureImageUnits();
+			FirstUAVUnit = FOpenGL::GetFirstPixelUAVUnit();
+			MaxUAVUnits = FOpenGL::GetMaxPixelUAVUnits();
+			break;
+		case SF_Compute:
+			ShaderStage = CrossCompiler::SHADER_STAGE_COMPUTE;
+			FirstTextureUnit = FOpenGL::GetFirstComputeTextureUnit();
+			MaxTextureImageUnits = FOpenGL::GetMaxComputeTextureImageUnits();
+			FirstUAVUnit = FOpenGL::GetFirstComputeUAVUnit();
+			MaxUAVUnits = FOpenGL::GetMaxComputeUAVUnits();
+			break;
+		default:
+			checkf(0, TEXT("Undefined FRHIShader Type %d!"), (int32)ShaderFrequency);
+			break;
+		}
+	}
+
+	void SetUAV(FRHIUnorderedAccessView* RHIUAV, uint8 Index)
+	{
+		ensureMsgf(Index < MaxUAVUnits
+			, TEXT("Using more %s image units (%d) than allowed (%d) on a shader unit!")
+			, GetShaderFrequencyString(ShaderFrequency, false)
+			, Index
+			, MaxUAVUnits
+		);
+
+		const GLint UAVIndex = FirstUAVUnit + Index;
 
-void FOpenGLDynamicRHI::RHISetShaderUniformBuffer(FRHIComputeShader* ComputeShaderRHI,uint32 BufferIndex, FRHIUniformBuffer* BufferRHI)
+		if (FOpenGLUnorderedAccessView* UAV = FOpenGLDynamicRHI::ResourceCast(RHIUAV))
+		{
+			if (UAV->Resource)
+			{
+				const GLenum Access = (ShaderFrequency == SF_Compute) ? GL_READ_WRITE : GL_WRITE_ONLY;
+				RHI.InternalSetShaderImageUAV(UAVIndex, UAV->Format, UAV->Resource, UAV->IsLayered(), UAV->GetLayer(), Access);
+			}
+			else
+			{
+				RHI.InternalSetShaderBufferUAV(UAVIndex, UAV->BufferResource);
+			}
+		}
+		else
+		{
+			RHI.InternalSetShaderBufferUAV(UAVIndex, 0);
+		}
+	}
+
+	void SetSRV(FRHIShaderResourceView* RHISRV, uint8 Index)
+	{
+		ensureMsgf(Index < MaxTextureImageUnits
+			, TEXT("Using more %s texture units (%d) than allowed (%d) on a shader unit!")
+			, GetShaderFrequencyString(ShaderFrequency, false)
+			, Index
+			, MaxTextureImageUnits
+		);
+
+		if (FOpenGLShaderResourceView* SRV = FOpenGLDynamicRHI::ResourceCast(RHISRV))
+		{
+			if (SRV->Target == GL_SHADER_STORAGE_BUFFER)
+			{
+				RHI.InternalSetShaderBufferUAV(FirstUAVUnit + Index, SRV->Resource);
+			}
+			else
+			{
+				RHI.InternalSetShaderTexture(nullptr, SRV, FirstTextureUnit + Index, SRV->Target, SRV->Resource, 0, SRV->LimitMip);
+				SetSampler(RHI.GetPointSamplerState(), Index);
+			}
+		}
+		else
+		{
+			RHI.InternalSetShaderTexture(nullptr, nullptr, FirstTextureUnit + Index, GL_TEXTURE_BUFFER, 0, 0, -1);
+			SetSampler(RHI.GetPointSamplerState(), Index);
+		}
+	}
+
+	void SetTexture(FRHITexture* TextureRHI, uint8 Index)
+	{
+		ensureMsgf(Index < MaxTextureImageUnits
+			, TEXT("Using more %s texture units (%d) than allowed (%d) on a shader unit!")
+			, GetShaderFrequencyString(ShaderFrequency, false)
+			, Index
+			, MaxTextureImageUnits
+		);
+
+		if (FOpenGLTexture* Texture = FOpenGLDynamicRHI::ResourceCast(TextureRHI))
+		{
+			RHI.InternalSetShaderTexture(Texture, nullptr, FirstTextureUnit + Index, Texture->Target, Texture->GetResource(), Texture->GetNumMips(), -1);
+		}
+		else
+		{
+			RHI.InternalSetShaderTexture(nullptr, nullptr, FirstTextureUnit + Index, 0, 0, 0, -1);
+		}
+
+		if (bClearTextureSamplers)
+		{
+			// clear any previous sampler state
+			RHI.InternalSetSamplerStates(FirstTextureUnit + Index, nullptr);
+		}
+	}
+
+	void SetSampler(FRHISamplerState* Sampler, uint8 Index)
+	{
+		RHI.InternalSetSamplerStates(FirstTextureUnit + Index, FOpenGLDynamicRHI::ResourceCast(Sampler));
+	}
+};
+
+void FOpenGLDynamicRHI::SetShaderParametersCommon(EShaderFrequency ShaderFrequency, TConstArrayView<uint8> InParametersData, TConstArrayView<FRHIShaderParameter> InParameters, TConstArrayView<FRHIShaderParameterResource> InResourceParameters)
 {
-	VERIFY_GL_SCOPE();
-	PendingState.BoundUniformBuffers[SF_Compute][BufferIndex] = BufferRHI;
-	PendingState.BoundUniformBuffersDynamicOffset[SF_Compute][BufferIndex] = 0u;
-	PendingState.DirtyUniformBuffers[SF_Compute] |= 1 << BufferIndex;
+	FOpenGLResourceBinder Binder(*this, ShaderFrequency);
 
-	if (!GUseEmulatedUniformBuffers || !((FOpenGLUniformBuffer*)BufferRHI)->bIsEmulatedUniformBuffer)
+	if (InParameters.Num())
 	{
-		PendingState.bAnyDirtyRealUniformBuffers[SF_Compute] = true;
+		for (const FRHIShaderParameter& Parameter : InParameters)
+		{
+			PendingState.ShaderParameters[Binder.ShaderStage].Set(Parameter.BufferIndex, Parameter.BaseIndex, Parameter.ByteSize, &InParametersData[Parameter.ByteOffset]);
+			PendingState.LinkedProgramAndDirtyFlag = nullptr;
+		}
+	}
+
+	for (const FRHIShaderParameterResource& Parameter : InResourceParameters)
+	{
+		if (Parameter.Type == FRHIShaderParameterResource::EType::UnorderedAccessView)
+		{
+			Binder.SetUAV(static_cast<FRHIUnorderedAccessView*>(Parameter.Resource), Parameter.Index);
+		}
+	}
+
+	for (const FRHIShaderParameterResource& Parameter : InResourceParameters)
+	{
+		switch (Parameter.Type)
+		{
+		case FRHIShaderParameterResource::EType::Texture:
+			Binder.SetTexture(static_cast<FRHITexture*>(Parameter.Resource), Parameter.Index);
+			break;
+		case FRHIShaderParameterResource::EType::ResourceView:
+			Binder.SetSRV(static_cast<FRHIShaderResourceView*>(Parameter.Resource), Parameter.Index);
+			break;
+		case FRHIShaderParameterResource::EType::UnorderedAccessView:
+			break;
+		case FRHIShaderParameterResource::EType::Sampler:
+			Binder.SetSampler(static_cast<FRHISamplerState*>(Parameter.Resource), Parameter.Index);
+			break;
+		case FRHIShaderParameterResource::EType::UniformBuffer:
+			BindUniformBuffer(Binder.ShaderFrequency, Parameter.Index, static_cast<FRHIUniformBuffer*>(Parameter.Resource));
+			break;
+		default:
+			checkf(false, TEXT("Unhandled resource type?"));
+			break;
+		}
 	}
 }
 
-void FOpenGLDynamicRHI::RHISetShaderParameter(FRHIGraphicsShader* ShaderRHI,uint32 BufferIndex,uint32 BaseIndex,uint32 NumBytes,const void* NewValue)
+void FOpenGLDynamicRHI::RHISetShaderParameters(FRHIGraphicsShader* Shader, TConstArrayView<uint8> InParametersData, TConstArrayView<FRHIShaderParameter> InParameters, TConstArrayView<FRHIShaderParameterResource> InResourceParameters, TConstArrayView<FRHIShaderParameterResource> InBindlessParameters)
 {
 	VERIFY_GL_SCOPE();
-	CrossCompiler::EShaderStage Stage = GetShaderCrossCompilerStage(ShaderRHI);
-	if (Stage != CrossCompiler::NUM_SHADER_STAGES)
+	if (ValidateShader(Shader))
 	{
-		PendingState.ShaderParameters[Stage].Set(BufferIndex, BaseIndex, NumBytes, NewValue);
-		PendingState.LinkedProgramAndDirtyFlag = nullptr;
+		SetShaderParametersCommon(
+			Shader->GetFrequency()
+			, InParametersData
+			, InParameters
+			, InResourceParameters
+		);
 	}
 }
 
-void FOpenGLDynamicRHI::RHISetShaderParameter(FRHIComputeShader* ComputeShaderRHI,uint32 BufferIndex,uint32 BaseIndex,uint32 NumBytes,const void* NewValue)
-{ 
+void FOpenGLDynamicRHI::RHISetShaderParameters(FRHIComputeShader* Shader, TConstArrayView<uint8> InParametersData, TConstArrayView<FRHIShaderParameter> InParameters, TConstArrayView<FRHIShaderParameterResource> InResourceParameters, TConstArrayView<FRHIShaderParameterResource> InBindlessParameters)
+{
 	VERIFY_GL_SCOPE();
-	PendingState.ShaderParameters[CrossCompiler::SHADER_STAGE_COMPUTE].Set(BufferIndex, BaseIndex, NumBytes, NewValue);
-	PendingState.LinkedProgramAndDirtyFlag = nullptr;
+	if (ValidateShader(Shader))
+	{
+		SetShaderParametersCommon(
+			SF_Compute
+			, InParametersData
+			, InParameters
+			, InResourceParameters
+		);
+	}
 }
 
-void FOpenGLDynamicRHI::RHISetShaderParameters(FRHIGraphicsShader* Shader, TConstArrayView<uint8> InParametersData, TConstArrayView<FRHIShaderParameter> InParameters, TConstArrayView<FRHIShaderParameterResource> InResourceParameters, TConstArrayView<FRHIShaderParameterResource> InBindlessParameters)
+void FOpenGLDynamicRHI::SetShaderUnbindsCommon(EShaderFrequency ShaderFrequency, TConstArrayView<FRHIShaderParameterUnbind> InUnbinds)
 {
-	UE::RHICore::RHISetShaderParametersShared(
-		*this
-		, Shader
-		, InParametersData
-		, InParameters
-		, InResourceParameters
-		, InBindlessParameters
-	);
-}
+	FOpenGLResourceBinder Binder(*this, ShaderFrequency);
 
-void FOpenGLDynamicRHI::RHISetShaderParameters(FRHIComputeShader* Shader, TConstArrayView<uint8> InParametersData, TConstArrayView<FRHIShaderParameter> InParameters, TConstArrayView<FRHIShaderParameterResource> InResourceParameters, TConstArrayView<FRHIShaderParameterResource> InBindlessParameters)
-{
-	UE::RHICore::RHISetShaderParametersShared(
-		*this
-		, Shader
-		, InParametersData
-		, InParameters
-		, InResourceParameters
-		, InBindlessParameters
-	);
+	for (const FRHIShaderParameterUnbind& Unbind : InUnbinds)
+	{
+		switch (Unbind.Type)
+		{
+		case FRHIShaderParameterUnbind::EType::ResourceView:
+			Binder.SetSRV(nullptr, Unbind.Index);
+			break;
+		case FRHIShaderParameterUnbind::EType::UnorderedAccessView:
+			Binder.SetUAV(nullptr, Unbind.Index);
+			break;
+		default:
+			checkf(false, TEXT("Unhandled unbind resource type?"));
+			break;
+		}
+	}
 }
 
 void FOpenGLDynamicRHI::RHISetShaderUnbinds(FRHIComputeShader* Shader, TConstArrayView<FRHIShaderParameterUnbind> InUnbinds)
 {
-	UE::RHICore::RHISetShaderUnbindsShared(*this, Shader, InUnbinds);
+	VERIFY_GL_SCOPE();
+	if (ValidateShader(Shader))
+	{
+		SetShaderUnbindsCommon(SF_Compute, InUnbinds);
+	}
 }
 
 void FOpenGLDynamicRHI::RHISetShaderUnbinds(FRHIGraphicsShader* Shader, TConstArrayView<FRHIShaderParameterUnbind> InUnbinds)
 {
-	UE::RHICore::RHISetShaderUnbindsShared(*this, Shader, InUnbinds);
+	VERIFY_GL_SCOPE();
+	if (ValidateShader(Shader))
+	{
+		SetShaderUnbindsCommon(Shader->GetFrequency(), InUnbinds);
+	}
 }
 
 void FOpenGLDynamicRHI::RHISetDepthStencilState(FRHIDepthStencilState* NewStateRHI,uint32 StencilRef)
@@ -2179,34 +2164,6 @@ void FOpenGLDynamicRHI::CommitComputeShaderConstants(FOpenGLComputeShader* Compu
 	PendingState.LinkedProgramAndDirtyFlag = nullptr;
 }
 
-template <EShaderFrequency Frequency>
-uint32 GetFirstTextureUnit();
-
-template <> FORCEINLINE uint32 GetFirstTextureUnit<SF_Vertex>() { return FOpenGL::GetFirstVertexTextureUnit(); }
-template <> FORCEINLINE uint32 GetFirstTextureUnit<SF_Pixel>() { return FOpenGL::GetFirstPixelTextureUnit(); }
-template <> FORCEINLINE uint32 GetFirstTextureUnit<SF_Geometry>() { return FOpenGL::GetFirstGeometryTextureUnit(); }
-template <> FORCEINLINE uint32 GetFirstTextureUnit<SF_Compute>() { return FOpenGL::GetFirstComputeTextureUnit(); }
-
-template <EShaderFrequency Frequency>
-uint32 GetNumTextureUnits();
-
-template <> FORCEINLINE uint32 GetNumTextureUnits<SF_Vertex>() { return FOpenGL::GetMaxVertexTextureImageUnits(); }
-template <> FORCEINLINE uint32 GetNumTextureUnits<SF_Pixel>() { return FOpenGL::GetMaxTextureImageUnits(); }
-template <> FORCEINLINE uint32 GetNumTextureUnits<SF_Geometry>() { return FOpenGL::GetMaxGeometryTextureImageUnits(); }
-template <> FORCEINLINE uint32 GetNumTextureUnits<SF_Compute>() { return FOpenGL::GetMaxComputeTextureImageUnits(); }
-
-template <EShaderFrequency Frequency>
-uint32 GetFirstUAVUnit() { return 0; }
-template <> FORCEINLINE uint32 GetFirstUAVUnit<SF_Vertex>() { return FOpenGL::GetFirstVertexUAVUnit(); }
-template <> FORCEINLINE uint32 GetFirstUAVUnit<SF_Pixel>() { return FOpenGL::GetFirstPixelUAVUnit(); }
-template <> FORCEINLINE uint32 GetFirstUAVUnit<SF_Compute>() { return FOpenGL::GetFirstComputeUAVUnit(); }
-
-template <EShaderFrequency Frequency>
-uint32 GetNumUAVUnits() { return 0; }
-template <> FORCEINLINE uint32 GetNumUAVUnits<SF_Compute>()	{ return FOpenGL::GetMaxComputeUAVUnits(); }
-template <> FORCEINLINE uint32 GetNumUAVUnits<SF_Pixel>()	{ return FOpenGL::GetMaxPixelUAVUnits(); }
-template <> FORCEINLINE uint32 GetNumUAVUnits<SF_Vertex>()	{ return FOpenGL::GetMaxPixelUAVUnits(); }
-
 template <class ShaderType>
 FORCEINLINE void FOpenGLDynamicRHI::SetResourcesFromTables(ShaderType* Shader)
 {
@@ -2215,89 +2172,8 @@ FORCEINLINE void FOpenGLDynamicRHI::SetResourcesFromTables(ShaderType* Shader)
 
 	static constexpr EShaderFrequency Frequency = ShaderType::Frequency;
 
-	struct FUniformResourceBinder
-	{
-		FOpenGLDynamicRHI& RHI;
-
-		void SetUAV(FRHIUnorderedAccessView* RHIUAV, uint8 Index)
-		{
-			ensureMsgf(Index < GetNumUAVUnits<Frequency>()
-				, TEXT("Using more %s image units (%d) than allowed (%d) on a shader unit!")
-				, GetShaderFrequencyString(Frequency, false)
-				, Index
-				, GetNumUAVUnits<Frequency>()
-			);
-
-			auto UAV = FOpenGLDynamicRHI::ResourceCast(RHIUAV);
-			GLint UAVIndex = GetFirstUAVUnit<Frequency>() + Index;
-
-			if (UAV->Resource)
-			{
-				GLenum Access = (Frequency == SF_Compute) ? GL_READ_WRITE : GL_WRITE_ONLY;
-				// TODO: This must be true for 3D textures
-				bool bLayered = false;
-				GLint Layer = 0;
-				RHI.InternalSetShaderImageUAV(UAVIndex, UAV->Format, UAV->Resource, bLayered, Layer, Access);
-			}
-			else
-			{
-				RHI.InternalSetShaderBufferUAV(UAVIndex, UAV->BufferResource);
-			}
-		}
-
-		void SetSRV(FRHIShaderResourceView* RHISRV, uint8 Index)
-		{
-			ensureMsgf(Index < GetNumTextureUnits<Frequency>()
-				, TEXT("Using more %s texture units (%d) than allowed (%d) on a shader unit!")
-				, GetShaderFrequencyString(Frequency, false)
-				, Index
-				, GetNumTextureUnits<Frequency>()
-			);
-
-			FOpenGLShaderResourceView* SRV = FOpenGLDynamicRHI::ResourceCast(RHISRV);
-			if (SRV->Target == GL_SHADER_STORAGE_BUFFER)
-			{
-				RHI.InternalSetShaderBufferUAV(GetFirstUAVUnit<Frequency>() + Index, SRV->Resource);
-			}
-			else
-			{
-				RHI.InternalSetShaderTexture(nullptr, SRV, GetFirstTextureUnit<Frequency>() + Index, SRV->Target, SRV->Resource, 0, SRV->LimitMip);
-				SetSampler(RHI.GetPointSamplerState(), Index);
-			}
-		}
-
-		void SetTexture(FRHITexture* TextureRHI, uint8 Index)
-		{
-			ensureMsgf(Index < GetNumTextureUnits<Frequency>()
-				, TEXT("Using more %s texture units (%d) than allowed (%d) on a shader unit!")
-				, GetShaderFrequencyString(Frequency, false)
-				, Index
-				, GetNumTextureUnits<Frequency>()
-			);
-
-			FOpenGLTexture* Texture = ResourceCast(TextureRHI);
-
-			if (Texture)
-			{
-				RHI.InternalSetShaderTexture(Texture, nullptr, GetFirstTextureUnit<Frequency>() + Index, Texture->Target, Texture->GetResource(), Texture->GetNumMips(), -1);
-			}
-			else
-			{
-				RHI.InternalSetShaderTexture(Texture, nullptr, GetFirstTextureUnit<Frequency>() + Index, 0, 0, 0, -1);
-			}
-
-			// clear any previous sampler state
-			RHI.InternalSetSamplerStates(GetFirstTextureUnit<Frequency>() + Index, nullptr);
-		}
-
-		void SetSampler(FRHISamplerState* Sampler, uint8 Index)
-		{
-			RHI.InternalSetSamplerStates(GetFirstTextureUnit<Frequency>() + Index, ResourceCast(Sampler));
-		}
-	};
-
 	UE::RHICore::SetResourcesFromTables(
-		  FUniformResourceBinder { *this }
+		FOpenGLResourceBinder(*this, Frequency, true)
 		, *Shader
 		, Shader->Bindings.ShaderResourceTable
 		, PendingState.DirtyUniformBuffers[Frequency]
@@ -2872,13 +2748,21 @@ uint32 FOpenGLDynamicRHI::RHIGetGPUFrameCycles(uint32 GPUIndex)
 	return GGPUFrameTime;
 }
 
-template <typename TRHIShader, typename TRHIProxyShader>
-void FOpenGLDynamicRHI::ApplyStaticUniformBuffers(TRHIShader* Shader, TRHIProxyShader* ProxyShader)
+template <typename TRHIShader>
+void FOpenGLDynamicRHI::ApplyStaticUniformBuffers(TRHIShader* Shader)
 {
-	if (ProxyShader)
+	if (Shader)
 	{
-		check(Shader);
-		UE::RHICore::ApplyStaticUniformBuffers(this, Shader, ProxyShader->StaticSlots, ProxyShader->Bindings.ShaderResourceTable.ResourceTableLayoutHashes, GlobalUniformBuffers);
+		const EShaderFrequency ShaderFrequency = Shader->GetFrequency();
+
+		auto* ProxyShader = ResourceCast(Shader);
+		check(ProxyShader);
+
+		UE::RHICore::ApplyStaticUniformBuffers(Shader, ProxyShader->StaticSlots, ProxyShader->Bindings.ShaderResourceTable.ResourceTableLayoutHashes, GlobalUniformBuffers,
+			[this, ShaderFrequency](int32 BufferIndex, FRHIUniformBuffer* Buffer)
+			{
+				BindUniformBuffer(ShaderFrequency, BufferIndex, Buffer);
+			});
 	}
 }
 
@@ -2915,9 +2799,9 @@ void FOpenGLDynamicRHI::RHISetGraphicsPipelineState(FRHIGraphicsPipelineState* G
 
 	if (bApplyAdditionalState)
 	{
-		ApplyStaticUniformBuffers(PsoInit.BoundShaderState.VertexShaderRHI, ResourceCast(PsoInit.BoundShaderState.VertexShaderRHI));
-		ApplyStaticUniformBuffers(PsoInit.BoundShaderState.GetGeometryShader(), ResourceCast(PsoInit.BoundShaderState.GetGeometryShader()));
-		ApplyStaticUniformBuffers(PsoInit.BoundShaderState.PixelShaderRHI, ResourceCast(PsoInit.BoundShaderState.PixelShaderRHI));
+		ApplyStaticUniformBuffers(PsoInit.BoundShaderState.VertexShaderRHI);
+		ApplyStaticUniformBuffers(PsoInit.BoundShaderState.GetGeometryShader());
+		ApplyStaticUniformBuffers(PsoInit.BoundShaderState.PixelShaderRHI);
 	}
 
 	// Store the PSO's primitive (after since IRHICommandContext::RHISetGraphicsPipelineState sets the BSS)
@@ -2933,7 +2817,7 @@ void FOpenGLDynamicRHI::RHISetComputeShader(FRHIComputeShader* ComputeShaderRHI)
 
 	PendingState.CurrentComputeShader = ComputeShaderRHI;
 
-	ApplyStaticUniformBuffers(ComputeShaderRHI, ResourceCast(ComputeShaderRHI));
+	ApplyStaticUniformBuffers(ComputeShaderRHI);
 }
 
 void FOpenGLDynamicRHI::RHIDispatchComputeShader(uint32 ThreadGroupCountX, uint32 ThreadGroupCountY, uint32 ThreadGroupCountZ)
diff --git a/Engine/Source/Runtime/OpenGLDrv/Public/OpenGLDrv.h b/Engine/Source/Runtime/OpenGLDrv/Public/OpenGLDrv.h
--- a/Engine/Source/Runtime/OpenGLDrv/Public/OpenGLDrv.h
+++ b/Engine/Source/Runtime/OpenGLDrv/Public/OpenGLDrv.h
@@ -44,6 +44,7 @@
 
 class FOpenGLDynamicRHI;
 class FResourceBulkDataInterface;
+struct FOpenGLResourceBinder;
 struct Rect;
 
 template<class T> struct TOpenGLResourceTraits;
@@ -343,6 +344,11 @@ public:
 		}
 	}
 
+	void BindUniformBuffer(EShaderFrequency ShaderFrequency, uint32 BufferIndex, FRHIUniformBuffer* BufferRHI);
+
+	void SetShaderParametersCommon(EShaderFrequency ShaderFrequency, TConstArrayView<uint8> InParametersData, TConstArrayView<FRHIShaderParameter> InParameters, TConstArrayView<FRHIShaderParameterResource> InResourceParameters);
+	void SetShaderUnbindsCommon(EShaderFrequency ShaderFrequency, TConstArrayView<FRHIShaderParameterUnbind> InUnbinds);
+
 	virtual FRasterizerStateRHIRef RHICreateRasterizerState(const FRasterizerStateInitializerRHI& Initializer) final override;
 	virtual FDepthStencilStateRHIRef RHICreateDepthStencilState(const FDepthStencilStateInitializerRHI& Initializer) final override;
 	virtual FBlendStateRHIRef RHICreateBlendState(const FBlendStateInitializerRHI& Initializer) final override;
@@ -443,22 +449,9 @@ public:
 	virtual void RHISetViewport(float MinX, float MinY, float MinZ, float MaxX, float MaxY, float MaxZ) final override;
 	virtual void RHISetScissorRect(bool bEnable, uint32 MinX, uint32 MinY, uint32 MaxX, uint32 MaxY) final override;
 	virtual void RHISetBoundShaderState(FRHIBoundShaderState* BoundShaderState) final override;
-	void RHISetShaderTexture(FRHIGraphicsShader* Shader, uint32 TextureIndex, FRHITexture* NewTexture);
-	void RHISetShaderTexture(FRHIComputeShader* PixelShader, uint32 TextureIndex, FRHITexture* NewTexture);
-	void RHISetShaderSampler(FRHIComputeShader* ComputeShader, uint32 SamplerIndex, FRHISamplerState* NewState);
-	void RHISetShaderSampler(FRHIGraphicsShader* Shader, uint32 SamplerIndex, FRHISamplerState* NewState);
-	void RHISetUAVParameter(FRHIPixelShader* PixelShader, uint32 UAVIndex, FRHIUnorderedAccessView* UAV);
-	void RHISetUAVParameter(FRHIComputeShader* ComputeShader, uint32 UAVIndex, FRHIUnorderedAccessView* UAV);
-	void RHISetUAVParameter(FRHIComputeShader* ComputeShader, uint32 UAVIndex, FRHIUnorderedAccessView* UAV, uint32 InitialCount);
-	void RHISetShaderResourceViewParameter(FRHIGraphicsShader* Shader, uint32 SamplerIndex, FRHIShaderResourceView* SRV);
-	void RHISetShaderResourceViewParameter(FRHIComputeShader* ComputeShader, uint32 SamplerIndex, FRHIShaderResourceView* SRV);
 	virtual void RHISetStaticUniformBuffers(const FUniformBufferStaticBindings& InUniformBuffers) final override;
 	virtual void RHISetStaticUniformBuffer(FUniformBufferStaticSlot Slot, FRHIUniformBuffer* Buffer) final override;
 	virtual void RHISetUniformBufferDynamicOffset(FUniformBufferStaticSlot Slot, uint32 Offset) final override;
-	void RHISetShaderUniformBuffer(FRHIGraphicsShader* Shader, uint32 BufferIndex, FRHIUniformBuffer* Buffer);
-	void RHISetShaderUniformBuffer(FRHIComputeShader* ComputeShader, uint32 BufferIndex, FRHIUniformBuffer* Buffer);
-	void RHISetShaderParameter(FRHIGraphicsShader* Shader, uint32 BufferIndex, uint32 BaseIndex, uint32 NumBytes, const void* NewValue);
-	void RHISetShaderParameter(FRHIComputeShader* ComputeShader, uint32 BufferIndex, uint32 BaseIndex, uint32 NumBytes, const void* NewValue);
 	virtual void RHISetShaderParameters(FRHIGraphicsShader* Shader, TConstArrayView<uint8> InParametersData, TConstArrayView<FRHIShaderParameter> InParameters, TConstArrayView<FRHIShaderParameterResource> InResourceParameters, TConstArrayView<FRHIShaderParameterResource> InBindlessParameters) final override;
 	virtual void RHISetShaderParameters(FRHIComputeShader* Shader, TConstArrayView<uint8> InParametersData, TConstArrayView<FRHIShaderParameter> InParameters, TConstArrayView<FRHIShaderParameterResource> InResourceParameters, TConstArrayView<FRHIShaderParameterResource> InBindlessParameters) final override;
 	virtual void RHISetShaderUnbinds(FRHIComputeShader* Shader, TConstArrayView<FRHIShaderParameterUnbind> InUnbinds) final override;
@@ -866,8 +859,8 @@ private:
 	// Cached context type on BeginScene
 	int32 BeginSceneContextType;
 
-	template <typename TRHIShader, typename TRHIProxyShader>
-	void ApplyStaticUniformBuffers(TRHIShader* Shader, TRHIProxyShader* ProxyShader);
+	template <typename TRHIShader>
+	void ApplyStaticUniformBuffers(TRHIShader* Shader);
 
 	TArray<FRHIUniformBuffer*> GlobalUniformBuffers;
 
