Ensure the GL RHI specifies a valid SRV mip level.
 - fix for a validation error where the SRV requests -1 as a max mip.

diff --git a/Engine/Source/Runtime/OpenGLDrv/Private/OpenGLCommands.cpp b/Engine/Source/Runtime/OpenGLDrv/Private/OpenGLCommands.cpp
--- a/Engine/Source/Runtime/OpenGLDrv/Private/OpenGLCommands.cpp
+++ b/Engine/Source/Runtime/OpenGLDrv/Private/OpenGLCommands.cpp
@@ -1072,7 +1072,7 @@ struct FOpenGLResourceBinder
 			}
 			else
 			{
-				RHI.InternalSetShaderTexture(nullptr, SRV, FirstTextureUnit + Index, SRV->Target, SRV->Resource, 0, SRV->LimitMip);
+				RHI.InternalSetShaderTexture(nullptr, SRV, FirstTextureUnit + Index, SRV->Target, SRV->Resource, 1, SRV->LimitMip);
 				SetSampler(RHI.GetPointSamplerState(), Index);
 			}
 		}
