Android GLES:
Ensure GLES performance messages are only received when -OpenGLDebugLevel=5
 - Hides spam from some drivers.

diff --git a/Engine/Source/Runtime/OpenGLDrv/Private/OpenGLDevice.cpp b/Engine/Source/Runtime/OpenGLDrv/Private/OpenGLDevice.cpp
--- a/Engine/Source/Runtime/OpenGLDrv/Private/OpenGLDevice.cpp
+++ b/Engine/Source/Runtime/OpenGLDrv/Private/OpenGLDevice.cpp
@@ -783,6 +783,7 @@ void InitDebugContext()
 		glDebugMessageControlKHR(GL_DEBUG_SOURCE_API_KHR, GL_DEBUG_TYPE_OTHER_KHR, GL_DEBUG_SEVERITY_NOTIFICATION, 0, NULL, GL_FALSE);
 
 		GLenum Severity;
+		const bool bAllowPerformanceMessages = GetOGLDebugOutputLevel() >= 5;
 		switch (GetOGLDebugOutputLevel())
 		{
 			case 5: Severity = GL_DONT_CARE; break;
@@ -793,7 +794,12 @@ void InitDebugContext()
 				[[fallthrough]];
 			default: Severity = GL_DEBUG_SEVERITY_HIGH_ARB; break;
 		}
-		glDebugMessageControlKHR(GL_DEBUG_SOURCE_API_KHR, GL_DONT_CARE, Severity, 0, NULL, GL_TRUE);
+		glDebugMessageControlKHR(GL_DONT_CARE, GL_DONT_CARE, Severity, 0, NULL, GL_TRUE);
+		if( !bAllowPerformanceMessages)
+		{
+			glDebugMessageControlKHR(GL_DONT_CARE, GL_DEBUG_TYPE_PERFORMANCE_KHR, GL_DONT_CARE, 0, NULL, GL_FALSE);
+		}
+		
 		UE_LOG(LogRHI,Verbose,TEXT("disabling reporting back of debug groups and markers to the OpenGL debug output callback"));
 	}
 #endif
