Android GLES: check GL_OES_get_program_binary is set before attempting to use program binaries.
 - solves crash issues relating to RenderDoc usage

diff --git a/Engine/Source/Runtime/OpenGLDrv/Private/OpenGLES.cpp b/Engine/Source/Runtime/OpenGLDrv/Private/OpenGLES.cpp
--- a/Engine/Source/Runtime/OpenGLDrv/Private/OpenGLES.cpp
+++ b/Engine/Source/Runtime/OpenGLDrv/Private/OpenGLES.cpp
@@ -147,6 +147,9 @@ GLint FOpenGLES::MaxCombinedUAVUnits = 0;
 /** GL_EXT_texture_compression_astc_decode_mode */
 bool FOpenGLES::bSupportsASTCDecodeMode = false;
 
+// GL_OES_get_program_binary
+bool FOpenGLES::bSupportsProgramBinary = false;
+
 FOpenGLES::EFeatureLevelSupport FOpenGLES::CurrentFeatureLevelSupport = FOpenGLES::EFeatureLevelSupport::ES31;
 
 bool FOpenGLES::SupportsDisjointTimeQueries()
@@ -275,6 +278,8 @@ void FOpenGLES::ProcessExtensions(const FString& ExtensionsString)
 		MaxMSAASamplesTileMem = 1;
 	}
 
+	bSupportsProgramBinary = ExtensionsString.Contains(TEXT("GL_OES_get_program_binary"));
+
 	bSupportsETC2 = true;
 	// According to https://www.khronos.org/registry/gles/extensions/EXT/EXT_color_buffer_float.txt
 	bSupportsColorBufferHalfFloat = (bSupportsColorBufferHalfFloat || bSupportsColorBufferFloat);
diff --git a/Engine/Source/Runtime/OpenGLDrv/Public/OpenGLES.h b/Engine/Source/Runtime/OpenGLDrv/Public/OpenGLES.h
--- a/Engine/Source/Runtime/OpenGLDrv/Public/OpenGLES.h
+++ b/Engine/Source/Runtime/OpenGLDrv/Public/OpenGLES.h
@@ -502,7 +502,7 @@ struct FOpenGLES : public FOpenGLBase
 		glUniform4uiv(Location, Count, Value);
 	}
 
-	static FORCEINLINE bool SupportsProgramBinary() { return true; }
+	static FORCEINLINE bool SupportsProgramBinary() { return bSupportsProgramBinary; }
 
 	static FORCEINLINE void GetProgramBinary(GLuint Program, GLsizei BufSize, GLsizei* Length, GLenum* BinaryFormat, void* Binary)
 	{
