Android GLES: If no other GL debugging is active, add a once per frame gl error check after present.
  - any error is sent to the log.
 - r.OpenGL.PerFrameErrorCheck=0 to disable.

diff --git a/Engine/Source/Runtime/OpenGLDrv/Private/OpenGLViewport.cpp b/Engine/Source/Runtime/OpenGLDrv/Private/OpenGLViewport.cpp
--- a/Engine/Source/Runtime/OpenGLDrv/Private/OpenGLViewport.cpp
+++ b/Engine/Source/Runtime/OpenGLDrv/Private/OpenGLViewport.cpp
@@ -11,6 +11,41 @@
 #include "OpenGLDrv.h"
 #include "OpenGLDrvPrivate.h"
 
+static int32 GOpenGLPerFrameErrorCheck = 1;
+static FAutoConsoleVariableRef CVarPerFrameGLErrorCheck(
+	TEXT("r.OpenGL.PerFrameErrorCheck"),
+	GOpenGLPerFrameErrorCheck,
+	TEXT("When no other GL debugging is in use, check for GL errors once per frame.\nNot active in shipping builds.\n")
+	TEXT("0: GL errors not be checked.\n")
+	TEXT("1: any GL errors will be logged as errors. (default)\n")
+	TEXT("2: any GL errors will be fatal.\n")
+	,
+	ECVF_RenderThreadSafe
+);
+
+static void CheckForGLErrors()
+{
+#if !UE_BUILD_SHIPPING 
+	if (GOpenGLPerFrameErrorCheck && IsOGLDebugOutputEnabled() == false && !ENABLE_VERIFY_GL)
+	{
+		int32 Error = PlatformGlGetError();
+		if (Error != GL_NO_ERROR)
+		{
+			switch (GOpenGLPerFrameErrorCheck)
+			{
+			case 1:
+				UE_LOG(LogRHI, Error, TEXT("GL Error encountered during frame %d, glerror=0x%x. Set command line arg -OpenGLDebugLevel=1 for detailed debugging."), GFrameNumber, Error);
+				break;
+			default: checkNoEntry(); [[fallthrough]];
+			case 2:
+				UE_LOG(LogRHI, Fatal, TEXT("GL Error encountered during frame %d, glerror=0x%x. Set command line arg -OpenGLDebugLevel=1 for detailed debugging."), GFrameNumber, Error);
+				break;
+			}
+		}
+	}
+#endif
+}
+
 void FOpenGLDynamicRHI::RHIGetSupportedResolution(uint32 &Width, uint32 &Height)
 {
 	PlatformGetSupportedResolution(Width, Height);
@@ -188,6 +223,8 @@ void FOpenGLDynamicRHI::RHIEndDrawingViewport(FRHIViewport* ViewportRHI,bool bPr
 	}
 
 	EndFrameTick();
+
+	CheckForGLErrors();
 }
 
 
