Android EGL: Add option to enable robust buffer access.

diff --git a/Engine/Source/Runtime/OpenGLDrv/Private/Android/AndroidEGL.cpp b/Engine/Source/Runtime/OpenGLDrv/Private/Android/AndroidEGL.cpp
--- a/Engine/Source/Runtime/OpenGLDrv/Private/Android/AndroidEGL.cpp
+++ b/Engine/Source/Runtime/OpenGLDrv/Private/Android/AndroidEGL.cpp
@@ -423,11 +423,12 @@ void AndroidEGL::InitEGL(APIVariant API)
 	// Get the EGL Extension list to determine what is supported
 	FString Extensions = ANSI_TO_TCHAR( eglQueryString( PImplData->eglDisplay, EGL_EXTENSIONS));
 
-	FPlatformMisc::LowLevelOutputDebugStringf( TEXT("EGL Extensions: \n%s" ), *Extensions );
+	UE_LOG(LogAndroid, Log, TEXT("EGL Extensions: \n%s"), *Extensions);
 
 	bSupportsKHRCreateContext = Extensions.Contains(TEXT("EGL_KHR_create_context"));
 	bSupportsKHRSurfacelessContext = Extensions.Contains(TEXT("EGL_KHR_surfaceless_context"));
 	bSupportsKHRNoErrorContext = Extensions.Contains(TEXT("EGL_KHR_create_context_no_error"));
+	bSupportsEXTRobustContext = Extensions.Contains(TEXT("EGL_EXT_create_context_robustness"));
 
 	if (API == AV_OpenGLES)
 	{
@@ -768,6 +769,8 @@ void AndroidEGL::Init(APIVariant API, uint32 MajorVersion, uint32 MinorVersion)
 {
 	check(IsInGameThread());
 	const bool bDebug = IsOGLDebugOutputEnabled();
+	const FString* ConfigRulesForceRobustGLContext = FAndroidMisc::GetConfigRulesVariable(TEXT("ForceRobustGLContext"));
+	bool bWantsRobustGLContext = ConfigRulesForceRobustGLContext && ConfigRulesForceRobustGLContext->Equals("true", ESearchCase::IgnoreCase);
 	if (PImplData->Initalized)
 	{
 		ensure(bDebug == PImplData->bIsDebug); // if this fires you would need to tear down the previous context and recreate to honour the debug change.
@@ -778,7 +781,7 @@ void AndroidEGL::Init(APIVariant API, uint32 MajorVersion, uint32 MinorVersion)
 	PImplData->bIsDebug = bDebug;
 	if (bSupportsKHRCreateContext)
 	{
-		const uint32 MaxElements = 13;
+		const uint32 MaxElements = 16;
 		uint32 Flags = 0;
 
 		Flags |= PImplData->bIsDebug ? EGL_CONTEXT_OPENGL_DEBUG_BIT_KHR : 0;
@@ -797,6 +800,14 @@ void AndroidEGL::Init(APIVariant API, uint32 MajorVersion, uint32 MinorVersion)
 			ContextAttributes[Element++] = EGL_TRUE;
 		}
 #endif // USE_ANDROID_EGL_NO_ERROR_CONTEXT
+
+		if (bSupportsEXTRobustContext && bWantsRobustGLContext)
+		{
+			UE_LOG(LogAndroid, Log, TEXT("Enabling: EGL_CONTEXT_OPENGL_ROBUST_ACCESS_EXT"));
+			ContextAttributes[Element++] = EGL_CONTEXT_OPENGL_ROBUST_ACCESS_EXT;
+			ContextAttributes[Element++] = EGL_TRUE;
+		}
+
 		if (API == AV_OpenGLCore)
 		{
 			ContextAttributes[Element++] = EGL_CONTEXT_OPENGL_PROFILE_MASK_KHR;
@@ -806,7 +817,7 @@ void AndroidEGL::Init(APIVariant API, uint32 MajorVersion, uint32 MinorVersion)
 		ContextAttributes[Element++] = Flags;
 		ContextAttributes[Element++] = EGL_NONE;
 
-		checkf( Element < MaxElements, TEXT("Too many elements in config list"));
+		checkf( Element <= MaxElements, TEXT("Too many elements in config list"));
 	}
 	else
 	{
diff --git a/Engine/Source/Runtime/OpenGLDrv/Private/Android/AndroidEGL.h b/Engine/Source/Runtime/OpenGLDrv/Private/Android/AndroidEGL.h
--- a/Engine/Source/Runtime/OpenGLDrv/Private/Android/AndroidEGL.h
+++ b/Engine/Source/Runtime/OpenGLDrv/Private/Android/AndroidEGL.h
@@ -145,6 +145,7 @@ private:
 	bool bSupportsKHRCreateContext;
 	bool bSupportsKHRSurfacelessContext;
 	bool bSupportsKHRNoErrorContext;
+	bool bSupportsEXTRobustContext;
 
 	int *ContextAttributes;
 };
