Cleaning up some OpenGLDrv headers.

diff --git a/Engine/Source/Developer/ShaderFormatOpenGL/Private/OpenGLShaderCompiler.cpp b/Engine/Source/Developer/ShaderFormatOpenGL/Private/OpenGLShaderCompiler.cpp
--- a/Engine/Source/Developer/ShaderFormatOpenGL/Private/OpenGLShaderCompiler.cpp
+++ b/Engine/Source/Developer/ShaderFormatOpenGL/Private/OpenGLShaderCompiler.cpp
@@ -57,7 +57,7 @@
 	#define GL_TESS_CONTROL_SHADER 0x8E88
 	#endif
 #endif
-#include "OpenGLUtil.h"
+
 #include "OpenGLShaderResources.h"
 
 THIRD_PARTY_INCLUDES_START
diff --git a/Engine/Source/Programs/IncludeTool/IncludeTool/Rules.cs b/Engine/Source/Programs/IncludeTool/IncludeTool/Rules.cs
--- a/Engine/Source/Programs/IncludeTool/IncludeTool/Rules.cs
+++ b/Engine/Source/Programs/IncludeTool/IncludeTool/Rules.cs
@@ -363,12 +363,12 @@ namespace IncludeTool
 			"/Engine/Source/Runtime/Slate/Public/Framework/Text/GenericPlatformTextField.h",
 
 			// Base definitions for OpenGL3/4
-			"/Engine/Source/Runtime/OpenGLDrv/Public/OpenGL.h",
-			"/Engine/Source/Runtime/OpenGLDrv/Public/OpenGL3.h",
-			"/Engine/Source/Runtime/OpenGLDrv/Public/OpenGL4.h",
+			"/Engine/Source/Runtime/OpenGLDrv/Private/OpenGL.h",
+			"/Engine/Source/Runtime/OpenGLDrv/Private/OpenGL3.h",
+			"/Engine/Source/Runtime/OpenGLDrv/Private/OpenGL4.h",
 			"/Engine/Source/Runtime/OpenGLDrv/Public/OpenGLUtil.h", // Requires external OpenGL headers
-			"/Engine/Source/Runtime/OpenGLDrv/Public/OpenGLState.h", // Requires external OpenGL headers
-			"/Engine/Source/Runtime/OpenGLDrv/Public/OpenGLResources.h", // Requires external OpenGL headers
+			"/Engine/Source/Runtime/OpenGLDrv/Private/OpenGLState.h", // Requires external OpenGL headers
+			"/Engine/Source/Runtime/OpenGLDrv/Private/OpenGLResources.h", // Requires external OpenGL headers
 		};
 
 		/// <summary>
diff --git a/Engine/Source/Runtime/OpenGLDrv/Private/OpenGL.h b/Engine/Source/Runtime/OpenGLDrv/Private/OpenGL.h
new file mode 100644
--- /dev/null
+++ b/Engine/Source/Runtime/OpenGLDrv/Private/OpenGL.h
@@ -0,0 +1,648 @@
+// Copyright Epic Games, Inc. All Rights Reserved.
+
+/*=============================================================================
+	OpenGL.h: Public OpenGL base definitions for non-common functionality
+=============================================================================*/
+#pragma once
+
+// HEADER_UNIT_SKIP - REVISIT
+
+#include "Containers/ContainersFwd.h"
+#include "Logging/LogMacros.h"
+
+struct FPlatformOpenGLContext;
+struct FPlatformOpenGLDevice;
+
+/** OpenGL Logging. */
+OPENGLDRV_API DECLARE_LOG_CATEGORY_EXTERN(LogOpenGL,Log,VeryVerbose);
+
+#define UGL_REQUIRED_VOID			{ UE_LOG(LogOpenGL,Fatal,TEXT("%s is not supported."), ANSI_TO_TCHAR(__FUNCTION__)); }
+#define UGL_REQUIRED(ReturnValue)	{ UE_LOG(LogOpenGL,Fatal,TEXT("%s is not supported."), ANSI_TO_TCHAR(__FUNCTION__)); return (ReturnValue); }
+#define UGL_OPTIONAL_VOID			{ }
+#define UGL_OPTIONAL(ReturnValue)	{ return (ReturnValue); }
+
+struct FPlatformOpenGLDevice;
+struct FPlatformOpenGLContext;
+
+
+#define UGL_SUPPORTS_PIXELBUFFERS		1
+#define UGL_SUPPORTS_UNIFORMBUFFERS		1
+
+#ifndef OPENGL_ES
+#define OPENGL_ES	0
+#endif
+#ifndef OPENGL_GL3
+#define OPENGL_GL3	0
+#endif
+#ifndef OPENGL_GL4
+#define OPENGL_GL4	0
+#endif
+
+/** Official OpenGL definitions */
+#ifndef GL_HALF_FLOAT
+#define GL_HALF_FLOAT 0x140B
+#endif
+
+#ifndef GL_RGBA16F
+#define GL_RGBA16F    0x881A
+#endif
+
+
+// for platform extensions: from OpenGLShaders.h
+extern "C" struct FOpenGLShaderDeviceCapabilities;
+
+// for platform extensions: from OpenGLResources.h
+typedef TArray<ANSICHAR> FAnsiCharArray;
+
+// for platform extensions: from OpenGLDrvPrivate.h
+extern "C" struct FOpenGLTextureFormat;
+
+// Base static class
+class FOpenGLBase
+{
+public:
+	enum class EResourceLockMode : uint8
+	{
+		RLM_ReadWrite,
+		RLM_ReadOnly,
+		RLM_WriteOnly,
+		RLM_WriteOnlyUnsynchronized,
+		RLM_WriteOnlyPersistent,
+		RLM_ReadOnlyPersistent,
+	};
+
+	enum EQueryMode
+	{
+		QM_Result,
+		QM_ResultAvailable,
+	};
+
+	enum EFenceResult
+	{
+		FR_AlreadySignaled,
+		FR_TimeoutExpired,
+		FR_ConditionSatisfied,
+		FR_WaitFailed,
+	};
+
+	static void ProcessQueryGLInt();
+	static void ProcessExtensions(const FString& ExtensionsString);
+	static void SetupDefaultGLContextState(const FString& ExtensionsString) {};
+
+	static FORCEINLINE bool SupportsUniformBuffers()					{ return true; }
+	static FORCEINLINE bool SupportsStructuredBuffers()					{ return true; }
+	static FORCEINLINE bool SupportsTimestampQueries()					{ return true; }
+	static FORCEINLINE bool SupportsDisjointTimeQueries()				{ return false; } // @todo: if enabled, causes crash on PC
+	static FORCEINLINE bool SupportsExactOcclusionQueries()				{ return true; }
+	static FORCEINLINE bool SupportsDepthStencilReadSurface()			{ return true; }
+	static FORCEINLINE bool SupportsFloatReadSurface()					{ return true; }
+	static FORCEINLINE bool SupportsWideMRT()							{ return true; }
+	static FORCEINLINE bool SupportsPolygonMode()						{ return true; }
+	static FORCEINLINE bool SupportsTexture3D()							{ return true; }
+	static FORCEINLINE bool SupportsMobileMultiView()					{ return false; }
+	static FORCEINLINE bool SupportsImageExternal()						{ return false; }
+	static FORCEINLINE bool SupportsTextureLODBias()					{ return true; }
+	static FORCEINLINE bool SupportsTextureCompare()					{ return true; }
+	static FORCEINLINE bool SupportsDrawIndexOffset()					{ return true; }
+	static FORCEINLINE bool SupportsDiscardFrameBuffer()				{ return false; }
+	static FORCEINLINE bool SupportsIndexedExtensions()					{ return true; }
+	static FORCEINLINE bool SupportsColorBufferFloat()					{ return true; }
+	static FORCEINLINE bool SupportsColorBufferHalfFloat()				{ return true; }
+	static FORCEINLINE bool SupportsVolumeTextureRendering()			{ return false; }
+	static FORCEINLINE bool SupportsShaderFramebufferFetch()			{ return false; }
+	static FORCEINLINE bool SupportsShaderMRTFramebufferFetch()			{ return false; }
+	static FORCEINLINE bool SupportsShaderDepthStencilFetch()			{ return false; }
+	static FORCEINLINE bool SupportsPixelLocalStorage()					{ return false; }
+	static FORCEINLINE bool SupportsVertexArrayBGRA()					{ return true; }
+	static FORCEINLINE bool SupportsBGRA8888()							{ return true; }
+	static FORCEINLINE bool SupportsDXT()								{ return true; }
+	static FORCEINLINE bool SupportsASTC()								{ return bSupportsASTC; }
+	static FORCEINLINE bool SupportsASTCHDR()							{ return bSupportsASTCHDR; }
+	static FORCEINLINE bool SupportsETC2()								{ return false; }
+	static FORCEINLINE bool SupportsFramebufferSRGBEnable()				{ return true; }
+	static FORCEINLINE bool SupportsFastBufferData()					{ return true; }
+	static FORCEINLINE bool SupportsTextureFilterAnisotropic()			{ return bSupportsTextureFilterAnisotropic; }
+	static FORCEINLINE bool SupportsSeparateAlphaBlend()				{ return bSupportsDrawBuffersBlend; }
+	static FORCEINLINE void EnableSupportsClipControl()					{ bSupportsClipControl = true; }
+	static FORCEINLINE bool SupportsClipControl()						{ return bSupportsClipControl; }
+	static FORCEINLINE bool SupportsSeamlessCubeMap()					{ return false; }
+	static FORCEINLINE bool SupportsDrawIndirect()						{ return false; }
+	static FORCEINLINE bool SupportsGenerateMipmap()					{ return false; }
+	static FORCEINLINE bool SupportsBufferStorage()						{ return false; }
+	static FORCEINLINE bool SupportsDepthBoundsTest()					{ return false; }
+	static FORCEINLINE bool SupportsTextureRange()						{ return false; }
+	static FORCEINLINE bool SupportsBindlessTexture()					{ return false; }
+	static FORCEINLINE bool HasHardwareHiddenSurfaceRemoval()			{ return false; }
+	static FORCEINLINE bool AmdWorkaround()								{ return false; }
+	static FORCEINLINE bool SupportsProgramBinary()						{ return false; }
+	static FORCEINLINE bool SupportsDepthClamp()						{ return true; }
+	
+	static FORCEINLINE bool SupportsASTCDecodeMode()					{ return false; }
+
+	static FORCEINLINE GLenum GetDepthFormat()							{ return GL_DEPTH_COMPONENT16; }
+	static FORCEINLINE GLenum GetShadowDepthFormat()					{ return GL_DEPTH_COMPONENT16; }
+
+	static FORCEINLINE GLint GetMaxTextureImageUnits()			{ check(MaxTextureImageUnits != -1); return MaxTextureImageUnits; }
+	static FORCEINLINE GLint GetMaxVertexTextureImageUnits()	{ check(MaxVertexTextureImageUnits != -1); return MaxVertexTextureImageUnits; }
+	static FORCEINLINE GLint GetMaxGeometryTextureImageUnits()	{ check(MaxGeometryTextureImageUnits != -1); return MaxGeometryTextureImageUnits; }
+	static FORCEINLINE GLint GetMaxComputeTextureImageUnits()	{ check(MaxComputeTextureImageUnits != -1); return MaxComputeTextureImageUnits; }
+	static FORCEINLINE GLint GetMaxCombinedTextureImageUnits()	{ check(MaxCombinedTextureImageUnits != -1); return MaxCombinedTextureImageUnits; }
+	static FORCEINLINE GLint GetTextureBufferAlignment()		{ return TextureBufferAlignment; }
+
+
+	// Indices per unit are set in this order [Pixel, Vertex, Geometry]
+	static FORCEINLINE GLint GetFirstPixelTextureUnit()			{ return 0; }
+	static FORCEINLINE GLint GetFirstVertexTextureUnit()		{ return GetFirstPixelTextureUnit() + GetMaxTextureImageUnits(); }
+	static FORCEINLINE GLint GetFirstGeometryTextureUnit()		{ return GetFirstVertexTextureUnit() + GetMaxVertexTextureImageUnits(); }
+
+	static FORCEINLINE GLint GetFirstComputeTextureUnit()		{ return 0; }
+	
+	// Image load/store units
+	static FORCEINLINE GLint GetFirstComputeUAVUnit()			{ return 0; }
+	static FORCEINLINE GLint GetMaxComputeUAVUnits()			{ return 0; }
+	static FORCEINLINE GLint GetFirstVertexUAVUnit()			{ return 0; }
+	static FORCEINLINE GLint GetFirstPixelUAVUnit()				{ return 0; }
+	static FORCEINLINE GLint GetMaxPixelUAVUnits()				{ return 0; }
+	static FORCEINLINE GLint GetMaxCombinedUAVUnits()			{ return 0; }
+	
+	static FORCEINLINE GLint GetMaxVaryingVectors()				{ check(MaxVaryingVectors != -1); return MaxVaryingVectors; }
+	static FORCEINLINE GLint GetMaxPixelUniformComponents()		{ check(MaxPixelUniformComponents != -1); return MaxPixelUniformComponents; }
+	static FORCEINLINE GLint GetMaxVertexUniformComponents()	{ check(MaxVertexUniformComponents != -1); return MaxVertexUniformComponents; }
+	static FORCEINLINE GLint GetMaxGeometryUniformComponents()	{ check(MaxGeometryUniformComponents != -1); return MaxGeometryUniformComponents; }
+	static FORCEINLINE GLint GetMaxComputeUniformComponents()	{ return 0; }
+
+	static FORCEINLINE uint64 GetVideoMemorySize()				{ return 0; }
+
+	static FORCEINLINE bool IsDebugContent()					{ return false; }
+	static FORCEINLINE void InitDebugContext()					{ }
+
+	static FORCEINLINE int32 GetReadHalfFloatPixelsEnum() UGL_REQUIRED(0)
+
+	static FORCEINLINE GLint GetMaxMSAASamplesTileMem()			{ return 0; /* not supported */ }
+
+	// Silently ignored if not implemented:
+	static FORCEINLINE void QueryTimestampCounter(GLuint QueryID) UGL_OPTIONAL_VOID
+	static FORCEINLINE void BeginQuery(GLenum QueryType, GLuint QueryId) UGL_OPTIONAL_VOID
+	static FORCEINLINE void EndQuery(GLenum QueryType) UGL_OPTIONAL_VOID
+	static FORCEINLINE void GetQueryObject(GLuint QueryId, EQueryMode QueryMode, uint64 *OutResult) UGL_OPTIONAL_VOID
+	static FORCEINLINE void BindFragDataLocation(GLuint Program, GLuint Color, const GLchar *Name) UGL_OPTIONAL_VOID
+	static FORCEINLINE void ReadBuffer(GLenum Mode) UGL_OPTIONAL_VOID
+	static FORCEINLINE void DrawBuffer(GLenum Mode) UGL_OPTIONAL_VOID
+	static FORCEINLINE void DeleteSync(UGLsync Sync) UGL_OPTIONAL_VOID
+	static FORCEINLINE UGLsync FenceSync(GLenum Condition, GLbitfield Flags) UGL_OPTIONAL(UGLsync())
+	static FORCEINLINE bool IsSync(UGLsync Sync) UGL_OPTIONAL(false)
+	static FORCEINLINE EFenceResult ClientWaitSync(UGLsync Sync, GLbitfield Flags, GLuint64 Timeout) UGL_OPTIONAL(FR_WaitFailed)
+	static FORCEINLINE void GenSamplers(GLsizei Count, GLuint *Samplers) UGL_OPTIONAL_VOID
+	static FORCEINLINE void DeleteSamplers(GLsizei Count, GLuint *Samplers) UGL_OPTIONAL_VOID
+	static FORCEINLINE void SetSamplerParameter(GLuint Sampler, GLenum Parameter, GLint Value) UGL_OPTIONAL_VOID
+	static FORCEINLINE void BindSampler(GLuint Unit, GLuint Sampler) UGL_OPTIONAL_VOID
+	static FORCEINLINE void PolygonMode(GLenum Face, GLenum Mode) UGL_OPTIONAL_VOID
+	static FORCEINLINE void VertexAttribDivisor(GLuint Index, GLuint Divisor) UGL_OPTIONAL_VOID
+	static FORCEINLINE void PushGroupMarker(const ANSICHAR* Name) UGL_OPTIONAL_VOID
+	static FORCEINLINE void PopGroupMarker() UGL_OPTIONAL_VOID
+	static FORCEINLINE void LabelObject(GLenum Type, GLuint Object, const ANSICHAR* Name) UGL_OPTIONAL_VOID
+	static FORCEINLINE GLsizei GetLabelObject(GLenum Type, GLuint Object, GLsizei BufferSize, ANSICHAR* OutName) UGL_OPTIONAL(0)
+	static FORCEINLINE void InvalidateFramebuffer(GLenum Target, GLsizei NumAttachments, const GLenum* Attachments) UGL_OPTIONAL_VOID
+
+	// Will assert at run-time if not implemented:
+	static FORCEINLINE void* MapBufferRange(GLenum Type, uint32 InOffset, uint32 InSize, EResourceLockMode LockMode) UGL_REQUIRED(NULL)
+	static FORCEINLINE void UnmapBufferRange(GLenum Type, uint32 InOffset, uint32 InSize) UGL_REQUIRED_VOID
+	static FORCEINLINE void UnmapBuffer(GLenum Type) UGL_REQUIRED_VOID
+	static FORCEINLINE void GenQueries(GLsizei NumQueries, GLuint* QueryIDs) UGL_REQUIRED_VOID
+	static FORCEINLINE void DeleteQueries(GLsizei NumQueries, const GLuint* QueryIDs) UGL_REQUIRED_VOID
+	static FORCEINLINE void GetQueryObject(GLuint QueryId, EQueryMode QueryMode, GLuint *OutResult) UGL_REQUIRED_VOID
+	static FORCEINLINE void BindBufferBase(GLenum Target, GLuint Index, GLuint Buffer) UGL_REQUIRED_VOID
+	static FORCEINLINE void BindBufferRange(GLenum Target, GLuint Index, GLuint Buffer, GLintptr Offset, GLsizeiptr Size) UGL_REQUIRED_VOID
+	static FORCEINLINE GLuint GetUniformBlockIndex(GLuint Program, const GLchar *UniformBlockName) UGL_REQUIRED(-1)
+	static FORCEINLINE void UniformBlockBinding(GLuint Program, GLuint UniformBlockIndex, GLuint UniformBlockBinding) UGL_REQUIRED_VOID
+	static FORCEINLINE void Uniform4uiv(GLint Location, GLsizei Count, const GLuint* Value) UGL_REQUIRED_VOID
+	static FORCEINLINE void TexParameter(GLenum Target, GLenum Parameter, GLint Value) UGL_REQUIRED_VOID
+	static FORCEINLINE void FramebufferTexture(GLenum Target, GLenum Attachment, GLuint Texture, GLint Level) UGL_REQUIRED_VOID
+	static FORCEINLINE void FramebufferTexture2D(GLenum Target, GLenum Attachment, GLenum TexTarget, GLuint Texture, GLint Level)
+	{
+		glFramebufferTexture2D(Target, Attachment, TexTarget, Texture, Level);
+	}
+	static FORCEINLINE void FramebufferTexture2DMultisample(GLenum Target, GLenum Attachment, GLenum TexTarget, GLuint Texture, GLint Level, GLint NumSamples) UGL_REQUIRED_VOID
+	static FORCEINLINE void FramebufferTexture3D(GLenum Target, GLenum Attachment, GLenum TexTarget, GLuint Texture, GLint Level, GLint ZOffset) UGL_REQUIRED_VOID
+	static FORCEINLINE void FramebufferTextureLayer(GLenum Target, GLenum Attachment, GLuint Texture, GLint Level, GLint Layer) UGL_REQUIRED_VOID
+	static FORCEINLINE void FramebufferRenderbuffer(GLenum Target, GLenum Attachment, GLenum RenderBufferTarget, GLuint RenderBuffer)
+	{
+		glFramebufferRenderbuffer(Target, Attachment, RenderBufferTarget, RenderBuffer);
+	}
+	static FORCEINLINE void FramebufferTextureMultiviewOVR(GLenum Target, GLenum Attachment, GLuint Texture, GLint Level, GLint BaseViewIndex, GLsizei NumViews) UGL_REQUIRED_VOID
+	static FORCEINLINE void FramebufferTextureMultisampleMultiviewOVR(GLenum Target, GLenum Attachment, GLuint Texture, GLint Level, GLsizei NumSamples, GLint BaseViewIndex, GLsizei NumViews) UGL_REQUIRED_VOID
+	static FORCEINLINE void BlitFramebuffer(GLint SrcX0, GLint SrcY0, GLint SrcX1, GLint SrcY1, GLint DstX0, GLint DstY0, GLint DstX1, GLint DstY1, GLbitfield Mask, GLenum Filter) UGL_REQUIRED_VOID
+	static FORCEINLINE void DrawBuffers(GLsizei NumBuffers, const GLenum *Buffers) UGL_REQUIRED_VOID
+	static FORCEINLINE void DepthRange(GLdouble Near, GLdouble Far) UGL_REQUIRED_VOID
+	static FORCEINLINE void EnableIndexed(GLenum Parameter, GLuint Index) UGL_REQUIRED_VOID
+	static FORCEINLINE void DisableIndexed(GLenum Parameter, GLuint Index) UGL_REQUIRED_VOID
+	static FORCEINLINE void ColorMaskIndexed(GLuint Index, GLboolean Red, GLboolean Green, GLboolean Blue, GLboolean Alpha) UGL_REQUIRED_VOID
+	static FORCEINLINE void VertexAttribPointer(GLuint Index, GLint Size, GLenum Type, GLboolean Normalized, GLsizei Stride, const GLvoid* Pointer) UGL_REQUIRED_VOID
+	static FORCEINLINE void VertexAttribIPointer(GLuint Index, GLint Size, GLenum Type, GLsizei Stride, const GLvoid* Pointer) UGL_REQUIRED_VOID
+	static FORCEINLINE void VertexAttrib4Nsv(GLuint AttributeIndex, const GLshort* Values) UGL_REQUIRED_VOID
+	static FORCEINLINE void VertexAttrib4sv(GLuint AttributeIndex, const GLshort* Values) UGL_REQUIRED_VOID
+	static FORCEINLINE void VertexAttribI4sv(GLuint AttributeIndex, const GLshort* Values) UGL_REQUIRED_VOID
+	static FORCEINLINE void VertexAttribI4usv(GLuint AttributeIndex, const GLushort* Values) UGL_REQUIRED_VOID
+	static FORCEINLINE void VertexAttrib4Nubv(GLuint AttributeIndex, const GLubyte* Values) UGL_REQUIRED_VOID
+	static FORCEINLINE void VertexAttrib4ubv(GLuint AttributeIndex, const GLubyte* Values) UGL_REQUIRED_VOID
+	static FORCEINLINE void VertexAttribI4ubv(GLuint AttributeIndex, const GLubyte* Values) UGL_REQUIRED_VOID
+	static FORCEINLINE void VertexAttrib4Nbv(GLuint AttributeIndex, const GLbyte* Values) UGL_REQUIRED_VOID
+	static FORCEINLINE void VertexAttrib4bv(GLuint AttributeIndex, const GLbyte* Values) UGL_REQUIRED_VOID
+	static FORCEINLINE void VertexAttribI4bv(GLuint AttributeIndex, const GLbyte* Values) UGL_REQUIRED_VOID
+	static FORCEINLINE void VertexAttrib4dv(GLuint AttributeIndex, const GLdouble* Values) UGL_REQUIRED_VOID
+	static FORCEINLINE void VertexAttribI4iv(GLuint AttributeIndex, const GLint* Values) UGL_REQUIRED_VOID
+	static FORCEINLINE void VertexAttribI4uiv(GLuint AttributeIndex, const GLuint* Values) UGL_REQUIRED_VOID
+	static FORCEINLINE void DrawArraysInstanced(GLenum Mode, GLint First, GLsizei Count, GLsizei InstanceCount) UGL_REQUIRED_VOID
+	static FORCEINLINE void DrawElementsInstanced(GLenum Mode, GLsizei Count, GLenum Type, const GLvoid* Indices, GLsizei InstanceCount) UGL_REQUIRED_VOID
+	static FORCEINLINE void DrawRangeElements(GLenum Mode, GLuint Start, GLuint End, GLsizei Count, GLenum Type, const GLvoid* Indices) UGL_REQUIRED_VOID
+	static FORCEINLINE void ClearBufferfv(GLenum Buffer, GLint DrawBufferIndex, const GLfloat* Value) UGL_REQUIRED_VOID
+	static FORCEINLINE void ClearBufferfi(GLenum Buffer, GLint DrawBufferIndex, GLfloat Depth, GLint Stencil) UGL_REQUIRED_VOID
+	static FORCEINLINE void ClearBufferiv(GLenum Buffer, GLint DrawBufferIndex, const GLint* Value) UGL_REQUIRED_VOID
+	static FORCEINLINE void ClearDepth(GLdouble Depth) UGL_REQUIRED_VOID
+	static FORCEINLINE void TexImage3D(GLenum Target, GLint Level, GLint InternalFormat, GLsizei Width, GLsizei Height, GLsizei Depth, GLint Border, GLenum Format, GLenum Type, const GLvoid* PixelData) UGL_REQUIRED_VOID
+	static FORCEINLINE void CompressedTexImage3D(GLenum Target, GLint Level, GLenum InternalFormat, GLsizei Width, GLsizei Height, GLsizei Depth, GLint Border, GLsizei ImageSize, const GLvoid* PixelData) UGL_REQUIRED_VOID
+	static FORCEINLINE void TexImage2DMultisample(GLenum Target, GLsizei Samples, GLint InternalFormat, GLsizei Width, GLsizei Height, GLboolean FixedSampleLocations) UGL_REQUIRED_VOID
+	static FORCEINLINE void TexBuffer(GLenum Target, GLenum InternalFormat, GLuint Buffer) UGL_REQUIRED_VOID
+	static FORCEINLINE void TexBufferRange(GLenum Target, GLenum InternalFormat, GLuint Buffer, GLintptr Offset, GLsizeiptr Size) UGL_REQUIRED_VOID
+	static FORCEINLINE void TexSubImage3D(GLenum Target, GLint Level, GLint XOffset, GLint YOffset, GLint ZOffset, GLsizei Width, GLsizei Height, GLsizei Depth, GLenum Format, GLenum Type, const GLvoid* PixelData) UGL_REQUIRED_VOID
+	static FORCEINLINE void	CopyTexSubImage2D(GLenum Target, GLint Level, GLint XOffset, GLint YOffset, GLint X, GLint Y, GLsizei Width, GLsizei Height) UGL_REQUIRED_VOID
+	static FORCEINLINE void	CopyTexSubImage3D(GLenum Target, GLint Level, GLint XOffset, GLint YOffset, GLint ZOffset, GLint X, GLint Y, GLsizei Width, GLsizei Height) UGL_REQUIRED_VOID
+	static FORCEINLINE void GetCompressedTexImage(GLenum Target, GLint Level, GLvoid* OutImageData) UGL_REQUIRED_VOID
+	static FORCEINLINE void GetTexImage(GLenum Target, GLint Level, GLenum Format, GLenum Type, GLvoid* OutPixelData) UGL_REQUIRED_VOID
+	static FORCEINLINE void CopyBufferSubData(GLenum ReadTarget, GLenum WriteTarget, GLintptr ReadOffset, GLintptr WriteOffset, GLsizeiptr Size) UGL_REQUIRED_VOID
+	static FORCEINLINE const ANSICHAR* GetStringIndexed(GLenum Name, GLuint Index) UGL_REQUIRED(NULL)
+	static FORCEINLINE GLuint GetMajorVersion() UGL_REQUIRED(0)
+	static FORCEINLINE GLuint GetMinorVersion() UGL_REQUIRED(0)
+	static FORCEINLINE ERHIFeatureLevel::Type GetFeatureLevel() UGL_REQUIRED(ERHIFeatureLevel::SM5)
+	static FORCEINLINE EShaderPlatform GetShaderPlatform() UGL_REQUIRED(SP_NumPlatforms)
+	static FORCEINLINE FString GetAdapterName() UGL_REQUIRED(TEXT(""))
+	static FORCEINLINE void BlendFuncSeparatei(GLuint Buf, GLenum SrcRGB, GLenum DstRGB, GLenum SrcAlpha, GLenum DstAlpha) UGL_REQUIRED_VOID
+	static FORCEINLINE void BlendEquationSeparatei(GLuint Buf, GLenum ModeRGB, GLenum ModeAlpha) UGL_REQUIRED_VOID
+	static FORCEINLINE void BlendFunci(GLuint Buf, GLenum Src, GLenum Dst) UGL_REQUIRED_VOID
+	static FORCEINLINE void BlendEquationi(GLuint Buf, GLenum Mode) UGL_REQUIRED_VOID
+	static FORCEINLINE void PatchParameteri(GLenum Pname, GLint Value) UGL_REQUIRED_VOID
+	static FORCEINLINE void BindImageTexture(GLuint Unit, GLuint Texture, GLint Level, GLboolean Layered, GLint Layer, GLenum Access, GLenum Format) UGL_REQUIRED_VOID
+	static FORCEINLINE void DispatchCompute(GLuint NumGroupsX, GLuint NumGroupsY, GLuint NumGroupsZ) UGL_REQUIRED_VOID
+	static FORCEINLINE void DispatchComputeIndirect(GLintptr Offset) UGL_REQUIRED_VOID
+	static FORCEINLINE void MemoryBarrier(GLbitfield Barriers) UGL_REQUIRED_VOID
+	static FORCEINLINE bool TexStorage2D(GLenum Target, GLint Levels, GLint InternalFormat, GLsizei Width, GLsizei Height, GLenum Format, GLenum Type, ETextureCreateFlags Flags) UGL_OPTIONAL(false)
+	static FORCEINLINE bool TexStorage2DMultisample(GLenum Target, GLsizei Samples, GLint InternalFormat, GLsizei Width, GLsizei Height, GLboolean FixedSampleLocations) UGL_OPTIONAL(false)
+	static FORCEINLINE void RenderbufferStorageMultisample(GLenum Target, GLsizei Samples, GLint InternalFormat, GLsizei Width, GLsizei Height) UGL_REQUIRED_VOID
+	static FORCEINLINE void TexStorage3D(GLenum Target, GLint Levels, GLint InternalFormat, GLsizei Width, GLsizei Height, GLsizei Depth, GLenum Format, GLenum Type) UGL_REQUIRED_VOID
+	static FORCEINLINE void CompressedTexSubImage3D(GLenum Target, GLint Level, GLint XOffset, GLint YOffset, GLint ZOffset, GLsizei Width, GLsizei Height, GLsizei Depth, GLenum Format, GLsizei ImageSize, const GLvoid* PixelData) UGL_REQUIRED_VOID
+	static FORCEINLINE void CopyImageSubData(GLuint SrcName, GLenum SrcTarget, GLint SrcLevel, GLint SrcX, GLint SrcY, GLint SrcZ, GLuint DstName, GLenum DstTarget, GLint DstLevel, GLint DstX, GLint DstY, GLint DstZ, GLsizei Width, GLsizei Height, GLsizei Depth) UGL_REQUIRED_VOID
+	static FORCEINLINE void TextureView(GLuint ViewName, GLenum ViewTarget, GLuint SrcName, GLenum InternalFormat, GLuint MinLevel, GLuint NumLevels, GLuint MinLayer, GLuint NumLayers) UGL_REQUIRED_VOID
+	static FORCEINLINE void DrawArraysIndirect(GLenum Mode, const void *Offset) UGL_REQUIRED_VOID
+	static FORCEINLINE void DrawElementsIndirect(GLenum Mode, GLenum Type, const void *Offset) UGL_REQUIRED_VOID
+	static FORCEINLINE void GenerateMipmap(GLenum Target) UGL_REQUIRED_VOID
+	static FORCEINLINE void BindVertexBuffer(GLuint BindingIndex, GLuint Nuffer, GLintptr Offset, GLsizei Stride) UGL_REQUIRED_VOID
+	static FORCEINLINE void VertexAttribFormat(GLuint AttribIndex, GLint Size, GLenum Type, GLboolean Normalized, GLuint RelativeOffset) UGL_REQUIRED_VOID
+	static FORCEINLINE void VertexAttribIFormat(GLuint AttribIndex, GLint Size, GLenum Type, GLuint RelativeOffset) UGL_REQUIRED_VOID
+	static FORCEINLINE void VertexAttribBinding(GLuint AttribIndex, GLuint BindingIndex) UGL_REQUIRED_VOID
+	static FORCEINLINE void ClearBufferData(GLenum Target, GLenum InternalFormat, GLenum Format, GLenum Type, const uint32* Data) UGL_REQUIRED_VOID
+	static FORCEINLINE void VertexBindingDivisor(GLuint BindingIndex, GLuint Divisor) UGL_REQUIRED_VOID
+	static FORCEINLINE void BufferStorage(GLenum Target, GLsizeiptr Size, const void *Data, GLbitfield Flags) UGL_REQUIRED_VOID
+	static FORCEINLINE void DepthBounds(GLfloat Min, GLfloat Max) UGL_REQUIRED_VOID
+	static FORCEINLINE void TextureRange(GLenum Target, GLsizei Length, const GLvoid *Pointer) UGL_OPTIONAL_VOID
+	static FORCEINLINE void ProgramParameter (GLuint Program, GLenum PName, GLint Value) UGL_OPTIONAL_VOID
+	static FORCEINLINE void UseProgramStages(GLuint Pipeline, GLbitfield Stages, GLuint Program) { glAttachShader(Pipeline, Program); }
+	static FORCEINLINE void BindProgramPipeline(GLuint Pipeline) { glUseProgram(Pipeline); }
+	static FORCEINLINE void DeleteShader(GLuint Program) { glDeleteShader(Program); }
+	static FORCEINLINE void DeleteProgramPipelines(GLsizei Number, const GLuint *Pipelines) { for(GLsizei i = 0; i < Number; i++) { glDeleteProgram(Pipelines[i]); } }
+	static FORCEINLINE void GenProgramPipelines(GLsizei Number, GLuint *Pipelines) { check(Pipelines); for(GLsizei i = 0; i < Number; i++) { Pipelines[i] = glCreateProgram(); } }
+	static FORCEINLINE void ProgramUniform1i(GLuint Program, GLint Location, GLint V0) { glUniform1i( Location, V0 ); }
+	static FORCEINLINE void ProgramUniform4iv(GLuint Program, GLint Location, GLsizei Count, const GLint *Value) { glUniform4iv(Location, Count, Value); }
+	static FORCEINLINE void ProgramUniform4fv(GLuint Program, GLint Location, GLsizei Count, const GLfloat *Value) { glUniform4fv(Location, Count, Value); }
+	static FORCEINLINE void ProgramUniform4uiv(GLuint Program, GLint Location, GLsizei Count, const GLuint *Value) UGL_REQUIRED_VOID
+	static FORCEINLINE void GetProgramPipelineiv(GLuint Pipeline, GLenum Pname, GLint* Params) UGL_OPTIONAL_VOID
+	static FORCEINLINE void ValidateProgramPipeline(GLuint Pipeline) UGL_OPTIONAL_VOID
+	static FORCEINLINE void GetProgramPipelineInfoLog(GLuint Pipeline, GLsizei BufSize, GLsizei* Length, GLchar* InfoLog) UGL_OPTIONAL_VOID
+	static FORCEINLINE bool IsProgramPipeline(GLuint Pipeline) UGL_OPTIONAL(false)
+
+	static FORCEINLINE GLuint64 GetTextureSamplerHandle(GLuint Texture, GLuint Sampler) UGL_REQUIRED(0)
+	static FORCEINLINE GLuint64 GetTextureHandle(GLuint Texture) UGL_REQUIRED(0)
+	static FORCEINLINE void MakeTextureHandleResident(GLuint64 TextureHandle) UGL_REQUIRED_VOID
+	static FORCEINLINE void MakeTextureHandleNonResident(GLuint64 TextureHandle) UGL_REQUIRED_VOID
+	static FORCEINLINE void UniformHandleui64(GLint Location, GLuint64 Value) UGL_REQUIRED_VOID
+
+	static FORCEINLINE void GetProgramBinary(GLuint Program, GLsizei BufSize, GLsizei *Length, GLenum *BinaryFormat, void *Binary) UGL_OPTIONAL_VOID
+	static FORCEINLINE void ProgramBinary(GLuint Program, GLenum BinaryFormat, const void *Binary, GLsizei Length) UGL_OPTIONAL_VOID
+
+	static FORCEINLINE void FrameBufferFetchBarrier() UGL_OPTIONAL_VOID
+	
+	static FPlatformOpenGLDevice*	CreateDevice() UGL_REQUIRED(NULL)
+	static FPlatformOpenGLContext*	CreateContext( FPlatformOpenGLDevice* Device, void* WindowHandle ) UGL_REQUIRED(NULL)
+
+	static FORCEINLINE void CheckFrameBuffer()
+	{
+#if UE_BUILD_DEBUG || UE_BUILD_DEVELOPMENT 
+		GLenum CompleteResult = glCheckFramebufferStatus(GL_FRAMEBUFFER);
+		if (CompleteResult != GL_FRAMEBUFFER_COMPLETE)
+		{
+				UE_LOG(LogRHI, Fatal,TEXT("Framebuffer not complete. Status = 0x%x"), CompleteResult);
+		}
+#endif 
+	}
+
+	static FORCEINLINE void BufferSubData(GLenum Target, GLintptr Offset, GLsizeiptr Size, const GLvoid* Data)	{ glBufferSubData(Target, Offset, Size, Data); }
+	static FORCEINLINE void DeleteBuffers(GLsizei Number, const GLuint* Buffers)	{ glDeleteBuffers(Number, Buffers); }
+	static FORCEINLINE void DeleteTextures(GLsizei Number, const GLuint* Textures)	{ glDeleteTextures(Number, Textures); }
+	static FORCEINLINE void Flush()													{ glFlush(); }
+	static FORCEINLINE GLuint CreateShader(GLenum Type)								{ return glCreateShader(Type); }
+	static FORCEINLINE GLuint CreateProgram()										{ return glCreateProgram(); }
+	static FORCEINLINE bool TimerQueryDisjoint()									{ return false; }
+
+	// Calling glBufferData() to discard-reupload is slower than calling glBufferSubData() on some platforms,
+	// because changing glBufferData() with a different size (from before) may incur extra validation.
+	// To use glBufferData() discard trick: set to this to true - otherwise, glBufferSubData() will used.
+	static FORCEINLINE bool DiscardFrameBufferToResize()							{ return true; }
+
+	static FORCEINLINE EPixelFormat PreferredPixelFormatHint(EPixelFormat PreferredPixelFormat)
+	{
+		// Use a default pixel format if none was specified	
+		if (PreferredPixelFormat == EPixelFormat::PF_Unknown)
+		{
+			PreferredPixelFormat = EPixelFormat::PF_B8G8R8A8;
+		}
+		return PreferredPixelFormat;
+	}
+
+	// for platform extensions
+	static void PE_GetCurrentOpenGLShaderDeviceCapabilities(FOpenGLShaderDeviceCapabilities& Capabilities);
+	static bool PE_GLSLToDeviceCompatibleGLSL(FAnsiCharArray& GlslCodeOriginal, const FString& ShaderName, GLenum TypeEnum, const FOpenGLShaderDeviceCapabilities& Capabilities, FAnsiCharArray& GlslCode) UGL_OPTIONAL(false)
+	static void PE_SetupTextureFormat(void(*SetupTextureFormat)(EPixelFormat, const FOpenGLTextureFormat&)) UGL_OPTIONAL_VOID
+
+	static GLenum GetPlatfrom5551Format() UGL_OPTIONAL(GL_UNSIGNED_SHORT_5_5_5_1)
+protected:
+	static GLint MaxTextureImageUnits;
+	static GLint MaxCombinedTextureImageUnits;
+	static GLint MaxComputeTextureImageUnits;
+	static GLint MaxVertexTextureImageUnits;
+	static GLint MaxGeometryTextureImageUnits;
+	static GLint MaxVertexUniformComponents;
+	static GLint MaxPixelUniformComponents;
+	static GLint MaxGeometryUniformComponents;
+	static GLint MaxVaryingVectors;
+	static GLint TextureBufferAlignment;
+
+	/** GL_ARB_clip_control */
+	static bool bSupportsClipControl;
+
+	/** GL_KHR_texture_compression_astc_ldr */
+	static bool bSupportsASTC;
+
+	/** GL_KHR_texture_compression_astc_hdr */
+	static bool bSupportsASTCHDR;
+
+	/** GL_ARB_seamless_cube_map */
+	static bool bSupportsSeamlessCubemap;
+	
+	/** Can we render to texture 2D array or 3D */
+	static bool bSupportsVolumeTextureRendering;
+	
+	/** GL_EXT_texture_filter_anisotropic Can we use anisotropic filtering? */
+	static bool bSupportsTextureFilterAnisotropic;
+
+	/** GL_ARB_draw_buffers_blend */
+	static bool bSupportsDrawBuffersBlend;
+
+	/** Workaround AMD driver issues. */
+	static bool bAmdWorkaround;
+
+};
+
+/** Unreal tokens that maps to different OpenGL tokens by platform. */
+#define UGL_DRAW_FRAMEBUFFER	GL_DRAW_FRAMEBUFFER
+#define UGL_READ_FRAMEBUFFER	GL_READ_FRAMEBUFFER
+#define UGL_ABGR8				GL_UNSIGNED_INT_8_8_8_8_REV
+#define UGL_ANY_SAMPLES_PASSED	GL_ANY_SAMPLES_PASSED
+#define UGL_SAMPLES_PASSED		GL_SAMPLES_PASSED
+#define UGL_TIME_ELAPSED		GL_TIME_ELAPSED
+#define UGL_CLAMP_TO_BORDER		GL_CLAMP_TO_BORDER
+
+/** OpenGL extensions */
+
+// http://www.opengl.org/registry/specs/EXT/texture_compression_s3tc.txt
+#if !defined(GL_EXT_texture_compression_s3tc)
+#define GL_COMPRESSED_RGB_S3TC_DXT1_EXT			0x83F0
+#define GL_COMPRESSED_RGBA_S3TC_DXT1_EXT		0x83F1
+#define GL_COMPRESSED_RGBA_S3TC_DXT3_EXT		0x83F2
+#define GL_COMPRESSED_RGBA_S3TC_DXT5_EXT		0x83F3
+#endif
+
+// http://www.opengl.org/registry/specs/EXT/texture_sRGB.txt
+#if !defined(GL_EXT_texture_sRGB)
+#define GL_SRGB_EXT                       0x8C40
+#define GL_SRGB8_EXT                      0x8C41
+#define GL_SRGB_ALPHA_EXT                 0x8C42
+#define GL_SRGB8_ALPHA8_EXT               0x8C43
+#define GL_SLUMINANCE_ALPHA_EXT           0x8C44
+#define GL_SLUMINANCE8_ALPHA8_EXT         0x8C45
+#define GL_SLUMINANCE_EXT                 0x8C46
+#define GL_SLUMINANCE8_EXT                0x8C47
+#define GL_COMPRESSED_SRGB_EXT            0x8C48
+#define GL_COMPRESSED_SRGB_ALPHA_EXT      0x8C49
+#define GL_COMPRESSED_SLUMINANCE_EXT      0x8C4A
+#define GL_COMPRESSED_SLUMINANCE_ALPHA_EXT 0x8C4B
+#define GL_COMPRESSED_SRGB_S3TC_DXT1_EXT  0x8C4C
+#define GL_COMPRESSED_SRGB_ALPHA_S3TC_DXT1_EXT 0x8C4D
+#define GL_COMPRESSED_SRGB_ALPHA_S3TC_DXT3_EXT 0x8C4E
+#define GL_COMPRESSED_SRGB_ALPHA_S3TC_DXT5_EXT 0x8C4F
+#endif
+
+// http://www.opengl.org/registry/specs/ARB/texture_compression_rgtc.txt
+#if !defined(GL_ARB_texture_compression_rgtc)
+#define GL_COMPRESSED_RED_RGTC1           0x8DBB
+#define GL_COMPRESSED_SIGNED_RED_RGTC1    0x8DBC
+#define GL_COMPRESSED_RG_RGTC2            0x8DBD
+#define GL_COMPRESSED_SIGNED_RG_RGTC2     0x8DBE
+#endif
+
+/* http://www.khronos.org/registry/gles/extensions/NV/NV_sRGB_formats.txt */
+#if !defined(GL_NV_sRGB_formats)
+#define GL_SLUMINANCE_NV                                        0x8C46
+#define GL_SLUMINANCE_ALPHA_NV                                  0x8C44
+#define GL_SRGB8_NV                                             0x8C41
+#define GL_SLUMINANCE8_NV                                       0x8C47
+#define GL_SLUMINANCE8_ALPHA8_NV                                0x8C45
+#define GL_COMPRESSED_SRGB_S3TC_DXT1_NV                         0x8C4C
+#define GL_COMPRESSED_SRGB_ALPHA_S3TC_DXT1_NV                   0x8C4D
+#define GL_COMPRESSED_SRGB_ALPHA_S3TC_DXT3_NV                   0x8C4E
+#define GL_COMPRESSED_SRGB_ALPHA_S3TC_DXT5_NV                   0x8C4F
+#endif
+
+// http://www.opengl.org/registry/specs/KHR/texture_compression_astc_ldr.txt
+#if !defined(GL_KHR_texture_compression_astc_ldr)
+#define GL_COMPRESSED_RGBA_ASTC_4x4_KHR            0x93B0
+#define GL_COMPRESSED_RGBA_ASTC_5x4_KHR            0x93B1
+#define GL_COMPRESSED_RGBA_ASTC_5x5_KHR            0x93B2
+#define GL_COMPRESSED_RGBA_ASTC_6x5_KHR            0x93B3
+#define GL_COMPRESSED_RGBA_ASTC_6x6_KHR            0x93B4
+#define GL_COMPRESSED_RGBA_ASTC_8x5_KHR            0x93B5
+#define GL_COMPRESSED_RGBA_ASTC_8x6_KHR            0x93B6
+#define GL_COMPRESSED_RGBA_ASTC_8x8_KHR            0x93B7
+#define GL_COMPRESSED_RGBA_ASTC_10x5_KHR           0x93B8
+#define GL_COMPRESSED_RGBA_ASTC_10x6_KHR           0x93B9
+#define GL_COMPRESSED_RGBA_ASTC_10x8_KHR           0x93BA
+#define GL_COMPRESSED_RGBA_ASTC_10x10_KHR          0x93BB
+#define GL_COMPRESSED_RGBA_ASTC_12x10_KHR          0x93BC
+#define GL_COMPRESSED_RGBA_ASTC_12x12_KHR          0x93BD
+
+#define GL_COMPRESSED_SRGB8_ALPHA8_ASTC_4x4_KHR    0x93D0
+#define GL_COMPRESSED_SRGB8_ALPHA8_ASTC_5x4_KHR    0x93D1
+#define GL_COMPRESSED_SRGB8_ALPHA8_ASTC_5x5_KHR    0x93D2
+#define GL_COMPRESSED_SRGB8_ALPHA8_ASTC_6x5_KHR    0x93D3
+#define GL_COMPRESSED_SRGB8_ALPHA8_ASTC_6x6_KHR    0x93D4
+#define GL_COMPRESSED_SRGB8_ALPHA8_ASTC_8x5_KHR    0x93D5
+#define GL_COMPRESSED_SRGB8_ALPHA8_ASTC_8x6_KHR    0x93D6
+#define GL_COMPRESSED_SRGB8_ALPHA8_ASTC_8x8_KHR    0x93D7
+#define GL_COMPRESSED_SRGB8_ALPHA8_ASTC_10x5_KHR   0x93D8
+#define GL_COMPRESSED_SRGB8_ALPHA8_ASTC_10x6_KHR   0x93D9
+#define GL_COMPRESSED_SRGB8_ALPHA8_ASTC_10x8_KHR   0x93DA
+#define GL_COMPRESSED_SRGB8_ALPHA8_ASTC_10x10_KHR  0x93DB
+#define GL_COMPRESSED_SRGB8_ALPHA8_ASTC_12x10_KHR  0x93DC
+#define GL_COMPRESSED_SRGB8_ALPHA8_ASTC_12x12_KHR  0x93DD
+#endif
+
+#if !defined(GL_TESS_EVALUATION_SHADER)
+#define GL_TESS_EVALUATION_SHADER					0x8E87
+#endif
+#if !defined(GL_TESS_CONTROL_SHADER)
+#define GL_TESS_CONTROL_SHADER						0x8E88
+#endif
+#if !defined(GL_PATCHES)
+#define GL_PATCHES									0x000E
+#endif
+#if !defined(GL_PATCH_VERTICES)
+#define GL_PATCH_VERTICES							0x8E72
+#endif
+#if !defined(GL_MAX_GEOMETRY_TEXTURE_IMAGE_UNITS)
+#define GL_MAX_GEOMETRY_TEXTURE_IMAGE_UNITS			0x8C29
+#endif
+#if !defined(GL_MAX_TESS_CONTROL_TEXTURE_IMAGE_UNITS)
+#define GL_MAX_TESS_CONTROL_TEXTURE_IMAGE_UNITS		0x8E81
+#endif
+#if !defined(GL_MAX_TESS_EVALUATION_TEXTURE_IMAGE_UNITS)
+#define GL_MAX_TESS_EVALUATION_TEXTURE_IMAGE_UNITS	0x8E82
+#endif
+#if !defined(GL_READ_WRITE)
+#define GL_READ_WRITE								0x88BA
+#endif
+#if !defined(GL_ALL_BARRIER_BITS)
+#define GL_ALL_BARRIER_BITS							0xFFFFFFFF
+#endif
+#if !defined(GL_TEXTURE_CUBE_MAP_ARRAY)
+#define GL_TEXTURE_CUBE_MAP_ARRAY					0x9009
+#endif
+#if !defined(GL_UNIFORM_BLOCK_REFERENCED_BY_TESS_CONTROL_SHADER)
+#define GL_UNIFORM_BLOCK_REFERENCED_BY_TESS_CONTROL_SHADER		0x84F0
+#endif
+#if !defined(GL_UNIFORM_BLOCK_REFERENCED_BY_TESS_EVALUATION_SHADER)
+#define GL_UNIFORM_BLOCK_REFERENCED_BY_TESS_EVALUATION_SHADER	0x84F1
+#endif
+#if !defined(GL_UNIFORM_BLOCK_REFERENCED_BY_COMPUTE_SHADER)
+#define GL_UNIFORM_BLOCK_REFERENCED_BY_COMPUTE_SHADER			0x90EC
+#endif
+#ifndef GL_ARB_seamless_cube_map
+#define GL_TEXTURE_CUBE_MAP_SEAMLESS				0x884F
+#endif
+#ifndef GL_TIME_ELAPSED_EXT
+#define GL_TIME_ELAPSED_EXT							0x88BF
+#endif
+#ifndef GL_TIMESTAMP_EXT
+#define GL_TIMESTAMP_EXT							0x8E28
+#endif
+#ifndef GL_DISPATCH_INDIRECT_BUFFER
+#define GL_DISPATCH_INDIRECT_BUFFER					0x90EE
+#endif
+#ifndef GL_DRAW_INDIRECT_BUFFER
+#define GL_DRAW_INDIRECT_BUFFER						0x8F3F
+#endif
+#ifndef GL_MAP_WRITE_BIT
+#define GL_MAP_WRITE_BIT							0x0002
+#endif
+#ifndef GL_MAP_PERSISTENT_BIT
+#define GL_MAP_PERSISTENT_BIT						0x0040
+#endif
+#ifndef GL_MAP_COHERENT_BIT
+#define GL_MAP_COHERENT_BIT							0x0080
+#endif
+#ifndef GL_DEPTH_BOUNDS_TEST_EXT
+#define GL_DEPTH_BOUNDS_TEST_EXT					0x8890
+#endif
+#ifndef GL_DEPTH_STENCIL_TEXTURE_MODE
+#define GL_DEPTH_STENCIL_TEXTURE_MODE				0x90EA
+#endif
+#ifndef GL_TEXTURE_SWIZZLE_R
+#define GL_TEXTURE_SWIZZLE_R						0x8E42
+#endif
+#ifndef GL_TEXTURE_SWIZZLE_G
+#define GL_TEXTURE_SWIZZLE_G						0x8E43
+#endif
+#ifndef GL_TEXTURE_SWIZZLE_B
+#define GL_TEXTURE_SWIZZLE_B						0x8E44
+#endif
+#ifndef GL_TEXTURE_SWIZZLE_A
+#define GL_TEXTURE_SWIZZLE_A						0x8E45
+#endif
+#ifndef GL_RED
+#define GL_RED										0x1903
+#endif
+#ifndef GL_BLUE
+#define GL_BLUE										0x1905
+#endif
+#ifndef GL_STENCIL_INDEX
+#define GL_STENCIL_INDEX							0x1901
+#endif
+#ifndef GL_RGBA_INTEGER
+#define GL_RGBA_INTEGER								0x8D99
+#endif
+
+#ifndef GL_ARB_compute_shader
+#define GL_COMPUTE_SHADER                 0x91B9
+#define GL_MAX_COMPUTE_UNIFORM_BLOCKS     0x91BB
+#define GL_MAX_COMPUTE_TEXTURE_IMAGE_UNITS 0x91BC
+#define GL_MAX_COMPUTE_IMAGE_UNIFORMS     0x91BD
+#define GL_MAX_COMPUTE_SHARED_MEMORY_SIZE 0x8262
+#define GL_MAX_COMPUTE_UNIFORM_COMPONENTS 0x8263
+#define GL_MAX_COMPUTE_ATOMIC_COUNTER_BUFFERS 0x8264
+#define GL_MAX_COMPUTE_ATOMIC_COUNTERS    0x8265
+#define GL_MAX_COMBINED_COMPUTE_UNIFORM_COMPONENTS 0x8266
+#define GL_MAX_COMPUTE_LOCAL_INVOCATIONS  0x90EB
+#define GL_MAX_COMPUTE_WORK_GROUP_COUNT   0x91BE
+#define GL_MAX_COMPUTE_WORK_GROUP_SIZE    0x91BF
+#define GL_COMPUTE_LOCAL_WORK_SIZE        0x8267
+#define GL_UNIFORM_BLOCK_REFERENCED_BY_COMPUTE_SHADER 0x90EC
+#define GL_ATOMIC_COUNTER_BUFFER_REFERENCED_BY_COMPUTE_SHADER 0x90ED
+#define GL_DISPATCH_INDIRECT_BUFFER       0x90EE
+#define GL_DISPATCH_INDIRECT_BUFFER_BINDING 0x90EF
+#define GL_COMPUTE_SHADER_BIT             0x00000020
+#endif
+
+#ifndef GL_GPU_DISJOINT_EXT
+#define GL_GPU_DISJOINT_EXT							0x8FBB
+#endif
+
+#ifndef GL_APPLE_client_storage
+#define GL_UNPACK_CLIENT_STORAGE_APPLE				0x85B2
+#endif
+
+#ifndef GL_APPLE_texture_range
+#define GL_TEXTURE_RANGE_LENGTH_APPLE     0x85B7
+#define GL_TEXTURE_RANGE_POINTER_APPLE    0x85B8
+#define GL_TEXTURE_STORAGE_HINT_APPLE     0x85BC
+#define GL_TEXTURE_MINIMIZE_STORAGE_APPLE 0x85B6
+#define GL_STORAGE_PRIVATE_APPLE          0x85BD
+#define GL_STORAGE_CACHED_APPLE           0x85BE
+#define GL_STORAGE_SHARED_APPLE           0x85BF
+#endif
+
+#ifndef TEXTURE_ASTC_DECODE_PRECISION_EXT
+#define TEXTURE_ASTC_DECODE_PRECISION_EXT 0x8F69
+#endif
+
+#ifndef GL_EXT_shader_pixel_local_storage
+#define GL_MAX_SHADER_PIXEL_LOCAL_STORAGE_FAST_SIZE_EXT 0x8F63
+#define GL_MAX_SHADER_PIXEL_LOCAL_STORAGE_SIZE_EXT 0x8F67
+#define GL_SHADER_PIXEL_LOCAL_STORAGE_EXT 0x8F64
+#endif
\ No newline at end of file
diff --git a/Engine/Source/Runtime/OpenGLDrv/Private/OpenGL3.h b/Engine/Source/Runtime/OpenGLDrv/Private/OpenGL3.h
new file mode 100644
--- /dev/null
+++ b/Engine/Source/Runtime/OpenGLDrv/Private/OpenGL3.h
@@ -0,0 +1,654 @@
+// Copyright Epic Games, Inc. All Rights Reserved.
+
+/*=============================================================================
+	OpenGL3.h: Public OpenGL 3.2 definitions for non-common functionality
+=============================================================================*/
+
+#pragma once
+
+#include "UObject/UObjectHierarchyFwd.h"
+#include "Misc/AssertionMacros.h"
+#include "HAL/UnrealMemory.h"
+#include "Containers/Array.h"
+#include "Containers/UnrealString.h"
+#include "Misc/Parse.h"
+#include "Containers/StringConv.h"
+#include "CoreGlobals.h"
+#include "Containers/Map.h"
+#include "CoreMinimal.h"
+#include "Misc/CommandLine.h"
+
+struct FPlatformOpenGLContext;
+struct FPlatformOpenGLDevice;
+
+template<typename KeyType,typename ValueType,typename SetAllocator ,typename KeyFuncs > class TMap;
+
+typedef GLsync UGLsync;
+
+#define OPENGL_GL3		1
+
+#include "OpenGL.h"
+
+#define USE_OPENGL_NAME_CACHE 1
+#define OPENGL_NAME_CACHE_SIZE 1024
+
+struct FOpenGL3 : public FOpenGLBase
+{
+	static FORCEINLINE bool IsDebugContent()						{ return bDebugContext; }
+
+	static FORCEINLINE bool SupportsTimestampQueries()				{ return TimestampQueryBits > 0; }
+	static FORCEINLINE bool SupportsSeamlessCubeMap()				{ return bSupportsSeamlessCubemap; }
+	static FORCEINLINE bool SupportsVolumeTextureRendering()		{ return bSupportsVolumeTextureRendering; }
+	static FORCEINLINE bool SupportsGenerateMipmap()				{ return true; }
+	static FORCEINLINE bool AmdWorkaround()							{ return bAmdWorkaround; }
+	static FORCEINLINE bool SupportsBufferStorage()					{ return true; }
+
+	// Optional
+	static FORCEINLINE void QueryTimestampCounter(GLuint QueryID)
+	{
+		glQueryCounter(QueryID, GL_TIMESTAMP);
+	}
+
+	static FORCEINLINE void BeginQuery(GLenum QueryType, GLuint QueryId)
+	{
+		glBeginQuery( QueryType, QueryId );
+	}
+
+	static FORCEINLINE void EndQuery(GLenum QueryType)
+	{
+		glEndQuery( QueryType );
+	}
+
+	static FORCEINLINE void GetQueryObject(GLuint QueryId, EQueryMode QueryMode, GLuint64* OutResult)
+	{
+		GLenum QueryName = (QueryMode == QM_Result) ? GL_QUERY_RESULT : GL_QUERY_RESULT_AVAILABLE;
+		GLuint64 Result = 0;
+		glGetQueryObjectui64v(QueryId, QueryName, &Result);
+		*OutResult = Result;
+	}
+
+	static FORCEINLINE void ReadBuffer(GLenum Mode)
+	{
+		glReadBuffer( Mode );
+	}
+
+	static FORCEINLINE void DrawBuffer(GLenum Mode)
+	{
+		glDrawBuffer( Mode );
+	}
+
+	static FORCEINLINE void DeleteSync(UGLsync Sync)
+	{
+		glDeleteSync( Sync );
+	}
+
+	static FORCEINLINE UGLsync FenceSync(GLenum Condition, GLbitfield Flags)
+	{
+		return glFenceSync( Condition, Flags );
+	}
+
+	static FORCEINLINE bool IsSync(UGLsync Sync)
+	{
+		return (glIsSync( Sync ) == GL_TRUE) ? true : false;
+	}
+
+	static FORCEINLINE EFenceResult ClientWaitSync(UGLsync Sync, GLbitfield Flags, GLuint64 Timeout)
+	{
+		GLenum Result = glClientWaitSync( Sync, Flags, Timeout );
+		switch (Result)
+		{
+			case GL_ALREADY_SIGNALED:		return FR_AlreadySignaled;
+			case GL_TIMEOUT_EXPIRED:		return FR_TimeoutExpired;
+			case GL_CONDITION_SATISFIED:	return FR_ConditionSatisfied;
+		}
+		return FR_WaitFailed;
+	}
+
+	static FORCEINLINE void GenSamplers(GLsizei Count, GLuint* Samplers)
+	{
+		glGenSamplers(Count, Samplers);
+	}
+
+	static FORCEINLINE void DeleteSamplers(GLsizei Count, GLuint* Samplers)
+	{
+		glDeleteSamplers(Count, Samplers);
+	}
+
+	static FORCEINLINE void SetSamplerParameter(GLuint Sampler, GLenum Parameter, GLint Value)
+	{
+		glSamplerParameteri(Sampler, Parameter, Value);
+	}
+
+	static FORCEINLINE void BindSampler(GLuint Unit, GLuint Sampler)
+	{
+		glBindSampler(Unit, Sampler);
+	}
+
+	static FORCEINLINE void PolygonMode(GLenum Face, GLenum Mode)
+	{
+		glPolygonMode(Face, Mode);
+	}
+
+	static FORCEINLINE void VertexAttribDivisor(GLuint Index, GLuint Divisor)
+	{
+		glVertexAttribDivisor(Index, Divisor);
+	}
+
+	// Required
+	static FORCEINLINE void* MapBufferRange(GLenum Type, uint32 InOffset, uint32 InSize, EResourceLockMode LockMode)
+	{
+		GLenum Access;
+		switch ( LockMode )
+		{
+			case EResourceLockMode::RLM_ReadOnly:
+				Access = GL_MAP_READ_BIT;
+				break;
+			case EResourceLockMode::RLM_ReadOnlyPersistent:
+				Access = (GL_MAP_READ_BIT | GL_MAP_PERSISTENT_BIT | GL_MAP_COHERENT_BIT);
+				break;
+			case EResourceLockMode::RLM_WriteOnly:
+				Access = (GL_MAP_INVALIDATE_BUFFER_BIT | GL_MAP_WRITE_BIT);
+#if 1
+				// Temp workaround for synchrnoization when a UBO is discarded while being referenced
+				Access |= GL_MAP_UNSYNCHRONIZED_BIT;
+#endif
+				break;
+			case EResourceLockMode::RLM_WriteOnlyUnsynchronized:
+				Access = (GL_MAP_WRITE_BIT | GL_MAP_UNSYNCHRONIZED_BIT);
+				break;
+			case EResourceLockMode::RLM_WriteOnlyPersistent:
+				Access = (GL_MAP_WRITE_BIT | GL_MAP_PERSISTENT_BIT | GL_MAP_COHERENT_BIT);
+				break;
+			case EResourceLockMode::RLM_ReadWrite:
+			default:
+				Access = (GL_MAP_READ_BIT | GL_MAP_WRITE_BIT);
+		}
+		return glMapBufferRange(Type, InOffset, InSize, Access);
+	}
+
+	static FORCEINLINE void UnmapBuffer(GLenum Type)
+	{
+		glUnmapBuffer(Type);
+	}
+
+	static FORCEINLINE void UnmapBufferRange(GLenum Type, uint32 InOffset, uint32 InSize)
+	{
+		UnmapBuffer(Type);
+	}
+
+	static FORCEINLINE void GenQueries(GLsizei NumQueries, GLuint* QueryIDs)
+	{
+		glGenQueries(NumQueries, QueryIDs);
+	}
+
+	static FORCEINLINE void DeleteQueries(GLsizei NumQueries, const GLuint* QueryIDs )
+	{
+		glDeleteQueries( NumQueries, QueryIDs );
+	}
+
+	static FORCEINLINE void GetQueryObject(GLuint QueryId, EQueryMode QueryMode, GLuint* OutResult)
+	{
+		GLenum QueryName = (QueryMode == QM_Result) ? GL_QUERY_RESULT : GL_QUERY_RESULT_AVAILABLE;
+		glGetQueryObjectuiv(QueryId, QueryName, OutResult);
+	}
+
+	static FORCEINLINE void BindBufferBase(GLenum Target, GLuint Index, GLuint Buffer)
+	{
+		glBindBufferBase(Target, Index, Buffer);
+	}
+
+	static FORCEINLINE void BindBufferRange(GLenum Target, GLuint Index, GLuint Buffer, GLintptr Offset, GLsizeiptr Size)
+	{
+		glBindBufferRange(Target, Index, Buffer, Offset, Size);
+	}
+
+	static FORCEINLINE GLuint GetUniformBlockIndex(GLuint Program, const GLchar* UniformBlockName)
+	{
+		return glGetUniformBlockIndex(Program, UniformBlockName);
+	}
+
+	static FORCEINLINE void UniformBlockBinding(GLuint Program, GLuint UniformBlockIndex, GLuint UniformBlockBinding)
+	{
+		glUniformBlockBinding(Program, UniformBlockIndex, UniformBlockBinding);
+	}
+
+	static FORCEINLINE void BindFragDataLocation(GLuint Program, GLuint Color, const GLchar* Name)
+	{
+		glBindFragDataLocation(Program, Color, Name);
+	}
+
+	static FORCEINLINE void TexParameter(GLenum Target, GLenum Parameter, GLint Value)
+	{
+		glTexParameteri(Target, Parameter, Value);
+	}
+
+	static FORCEINLINE void FramebufferTexture(GLenum Target, GLenum Attachment, GLuint Texture, GLint Level)
+	{
+		glFramebufferTexture(Target, Attachment, Texture, Level);
+	}
+
+	static FORCEINLINE void FramebufferTexture3D(GLenum Target, GLenum Attachment, GLenum TexTarget, GLuint Texture, GLint Level, GLint ZOffset)
+	{
+		glFramebufferTexture3D(Target, Attachment, TexTarget, Texture, Level, ZOffset);
+	}
+
+	static FORCEINLINE void FramebufferTextureLayer(GLenum Target, GLenum Attachment, GLuint Texture, GLint Level, GLint Layer)
+	{
+		glFramebufferTextureLayer(Target, Attachment, Texture, Level, Layer);
+	}
+
+	static FORCEINLINE void Uniform4uiv(GLint Location, GLsizei Count, const GLuint* Value)
+	{
+		glUniform4uiv(Location, Count, Value);
+	}
+
+	static FORCEINLINE void BlitFramebuffer(GLint SrcX0, GLint SrcY0, GLint SrcX1, GLint SrcY1, GLint DstX0, GLint DstY0, GLint DstX1, GLint DstY1, GLbitfield Mask, GLenum Filter)
+	{
+		glBlitFramebuffer(SrcX0, SrcY0, SrcX1, SrcY1, DstX0, DstY0, DstX1, DstY1, Mask, Filter);
+	}
+
+	static FORCEINLINE void DrawBuffers(GLsizei NumBuffers, const GLenum* Buffers)
+	{
+		glDrawBuffers(NumBuffers, Buffers);
+	}
+
+	static FORCEINLINE void DepthRange(GLdouble Near, GLdouble Far)
+	{
+		glDepthRange(Near, Far);
+	}
+
+	static FORCEINLINE void EnableIndexed(GLenum Parameter, GLuint Index)
+	{
+		glEnablei(Parameter, Index);
+	}
+
+	static FORCEINLINE void DisableIndexed(GLenum Parameter, GLuint Index)
+	{
+		glDisablei(Parameter, Index);
+	}
+
+	static FORCEINLINE void ColorMaskIndexed(GLuint Index, GLboolean Red, GLboolean Green, GLboolean Blue, GLboolean Alpha)
+	{
+		glColorMaski(Index, Red, Green, Blue, Alpha);
+	}
+
+	static FORCEINLINE void VertexAttribPointer(GLuint Index, GLint Size, GLenum Type, GLboolean Normalized, GLsizei Stride, const GLvoid* Pointer)
+	{
+		glVertexAttribPointer(Index, Size, Type, Normalized, Stride, Pointer);
+	}
+
+	static FORCEINLINE void VertexAttribIPointer(GLuint Index, GLint Size, GLenum Type, GLsizei Stride, const GLvoid* Pointer)
+	{
+		glVertexAttribIPointer(Index, Size, Type, Stride, Pointer);
+	}
+
+	static FORCEINLINE void VertexAttrib4Nsv(GLuint AttributeIndex, const GLshort* Values)
+	{
+		glVertexAttrib4Nsv(AttributeIndex, Values);
+	}
+
+	static FORCEINLINE void VertexAttrib4sv(GLuint AttributeIndex, const GLshort* Values)
+	{
+		glVertexAttrib4sv(AttributeIndex, Values);
+	}
+
+	static FORCEINLINE void VertexAttribI4sv(GLuint AttributeIndex, const GLshort* Values)
+	{
+		glVertexAttribI4sv(AttributeIndex, Values);
+	}
+
+	static FORCEINLINE void VertexAttribI4usv(GLuint AttributeIndex, const GLushort* Values)
+	{
+		glVertexAttribI4usv(AttributeIndex, Values);
+	}
+
+	static FORCEINLINE void VertexAttrib4Nubv(GLuint AttributeIndex, const GLubyte* Values)
+	{
+		glVertexAttrib4Nubv(AttributeIndex, Values);
+	}
+
+	static FORCEINLINE void VertexAttrib4ubv(GLuint AttributeIndex, const GLubyte* Values)
+	{
+		glVertexAttrib4ubv(AttributeIndex, Values);
+	}
+
+	static FORCEINLINE void VertexAttribI4ubv(GLuint AttributeIndex, const GLubyte* Values)
+	{
+		glVertexAttribI4ubv(AttributeIndex, Values);
+	}
+
+	static FORCEINLINE void VertexAttrib4Nbv(GLuint AttributeIndex, const GLbyte* Values)
+	{
+		glVertexAttrib4Nbv(AttributeIndex, Values);
+	}
+
+	static FORCEINLINE void VertexAttrib4bv(GLuint AttributeIndex, const GLbyte* Values)
+	{
+		glVertexAttrib4bv(AttributeIndex, Values);
+	}
+
+	static FORCEINLINE void VertexAttribI4bv(GLuint AttributeIndex, const GLbyte* Values)
+	{
+		glVertexAttribI4bv(AttributeIndex, Values);
+	}
+
+	static FORCEINLINE void VertexAttrib4dv(GLuint AttributeIndex, const GLdouble* Values)
+	{
+		glVertexAttrib4dv(AttributeIndex, Values);
+	}
+
+	static FORCEINLINE void VertexAttribI4iv(GLuint AttributeIndex, const GLint* Values)
+	{
+		glVertexAttribI4iv(AttributeIndex, Values);
+	}
+
+	static FORCEINLINE void VertexAttribI4uiv(GLuint AttributeIndex, const GLuint* Values)
+	{
+		glVertexAttribI4uiv(AttributeIndex, Values);
+	}
+
+	static FORCEINLINE void DrawArraysInstanced(GLenum Mode, GLint First, GLsizei Count, GLsizei InstanceCount)
+	{
+		glDrawArraysInstanced(Mode, First, Count, InstanceCount);
+	}
+
+	static FORCEINLINE void DrawElementsInstanced(GLenum Mode, GLsizei Count, GLenum Type, const GLvoid* Indices, GLsizei InstanceCount)
+	{
+		glDrawElementsInstanced(Mode, Count, Type, Indices, InstanceCount);
+	}
+
+	static FORCEINLINE void DrawRangeElements(GLenum Mode, GLuint Start, GLuint End, GLsizei Count, GLenum Type, const GLvoid* Indices)
+	{
+		glDrawRangeElements(Mode, Start, End, Count, Type, Indices);
+	}
+
+	static FORCEINLINE void ClearBufferfv(GLenum Buffer, GLint DrawBufferIndex, const GLfloat* Value)
+	{
+		glClearBufferfv(Buffer, DrawBufferIndex, Value);
+	}
+
+	static FORCEINLINE void ClearBufferfi(GLenum Buffer, GLint DrawBufferIndex, GLfloat Depth, GLint Stencil)
+	{
+		glClearBufferfi(Buffer, DrawBufferIndex, Depth, Stencil);
+	}
+
+	static FORCEINLINE void ClearBufferiv(GLenum Buffer, GLint DrawBufferIndex, const GLint* Value)
+	{
+		glClearBufferiv(Buffer, DrawBufferIndex, Value);
+	}
+
+	static FORCEINLINE void ClearDepth(GLdouble Depth)
+	{
+		glClearDepth(Depth);
+	}
+
+	static FORCEINLINE void TexImage3D(GLenum Target, GLint Level, GLint InternalFormat, GLsizei Width, GLsizei Height, GLsizei Depth, GLint Border, GLenum Format, GLenum Type, const GLvoid* PixelData)
+	{
+		glTexImage3D(Target, Level, InternalFormat, Width, Height, Depth, Border, Format, Type, PixelData);
+	}
+
+	static FORCEINLINE void CompressedTexImage3D(GLenum Target, GLint Level, GLenum InternalFormat, GLsizei Width, GLsizei Height, GLsizei Depth, GLint Border, GLsizei ImageSize, const GLvoid* PixelData)
+	{
+		glCompressedTexImage3D(Target, Level, InternalFormat, Width, Height, Depth, Border, ImageSize, PixelData);
+	}
+
+	static FORCEINLINE void CompressedTexSubImage2D(GLenum Target, GLint Level, GLsizei Width, GLsizei Height, GLenum Format, GLsizei ImageSize, const GLvoid* PixelData)
+	{
+		glCompressedTexSubImage2D(Target, Level, 0, 0, Width, Height, Format, ImageSize, PixelData);
+	}
+
+	static FORCEINLINE void TexImage2DMultisample(GLenum Target, GLsizei Samples, GLint InternalFormat, GLsizei Width, GLsizei Height, GLboolean FixedSampleLocations)
+	{
+		glTexImage2DMultisample(Target, Samples, InternalFormat, Width, Height, FixedSampleLocations);
+	}
+
+	static FORCEINLINE void TexBuffer(GLenum Target, GLenum InternalFormat, GLuint Buffer)
+	{
+		glTexBuffer(Target, InternalFormat, Buffer);
+	}
+
+	static FORCEINLINE void TexSubImage3D(GLenum Target, GLint Level, GLint XOffset, GLint YOffset, GLint ZOffset, GLsizei Width, GLsizei Height, GLsizei Depth, GLenum Format, GLenum Type, const GLvoid* PixelData)
+	{
+		glTexSubImage3D(Target, Level, XOffset, YOffset, ZOffset, Width, Height, Depth, Format, Type, PixelData);
+	}
+
+	static FORCEINLINE void	CopyTexSubImage1D(GLenum Target, GLint Level, GLint XOffset, GLint X, GLint Y, GLsizei Width)
+	{
+		glCopyTexSubImage1D(Target, Level, XOffset, X, Y, Width);
+	}
+
+	static FORCEINLINE void	CopyTexSubImage2D(GLenum Target, GLint Level, GLint XOffset, GLint YOffset, GLint X, GLint Y, GLsizei Width, GLsizei Height)
+	{
+		glCopyTexSubImage2D(Target, Level, XOffset, YOffset, X, Y, Width, Height);
+	}
+
+	static FORCEINLINE void	CopyTexSubImage3D(GLenum Target, GLint Level, GLint XOffset, GLint YOffset, GLint ZOffset, GLint X, GLint Y, GLsizei Width, GLsizei Height)
+	{
+		glCopyTexSubImage3D(Target, Level, XOffset, YOffset, ZOffset, X, Y, Width, Height);
+	}
+
+	static FORCEINLINE void GetCompressedTexImage(GLenum Target, GLint Level, GLvoid* OutImageData)
+	{
+		glGetCompressedTexImage(Target, Level, OutImageData);
+	}
+
+	static FORCEINLINE void GetTexImage(GLenum Target, GLint Level, GLenum Format, GLenum Type, GLvoid* OutPixelData)
+	{
+		glGetTexImage(Target, Level, Format, Type, OutPixelData);
+	}
+
+	static FORCEINLINE void CopyBufferSubData(GLenum ReadTarget, GLenum WriteTarget, GLintptr ReadOffset, GLintptr WriteOffset, GLsizeiptr Size)
+	{
+		glCopyBufferSubData(ReadTarget, WriteTarget, ReadOffset, WriteOffset, Size);
+	}
+	
+	static FORCEINLINE GLuint CreateShader(GLenum Type)
+	{
+#if USE_OPENGL_NAME_CACHE
+		static TMap<GLenum, TArray<GLuint>> ShaderNames;
+		TArray<GLuint>& Shaders = ShaderNames.FindOrAdd(Type);
+		if(!Shaders.Num())
+		{
+			while(Shaders.Num() < OPENGL_NAME_CACHE_SIZE)
+			{
+				GLuint Resource = glCreateShader(Type);
+				Shaders.Add(Resource);
+			}
+		}
+		return Shaders.Pop();
+#else
+		return glCreateShader(Type);
+#endif
+	}
+	
+	static FORCEINLINE GLuint CreateProgram()
+	{
+#if USE_OPENGL_NAME_CACHE
+		static TArray<GLuint> ProgramNames;
+		if(!ProgramNames.Num())
+		{
+			while(ProgramNames.Num() < OPENGL_NAME_CACHE_SIZE)
+			{
+				GLuint Resource = glCreateProgram();
+				ProgramNames.Add(Resource);
+			}
+		}
+		return ProgramNames.Pop();
+#else
+		return glCreateProgram();
+#endif
+	}
+
+	static FORCEINLINE void GenBuffers( GLsizei n, GLuint *buffers)
+	{
+#if USE_OPENGL_NAME_CACHE
+		if ( n < OPENGL_NAME_CACHE_SIZE - NextBufferName)
+		{
+			FMemory::Memcpy( buffers, &BufferNamesCache[NextBufferName], sizeof(GLuint)*n);
+			NextBufferName += n;
+		}
+		else
+		{
+			if ( n >= OPENGL_NAME_CACHE_SIZE)
+			{
+				glGenBuffers( n, buffers);
+			}
+			else
+			{
+				GLsizei Leftover = OPENGL_NAME_CACHE_SIZE - NextBufferName;
+
+				FMemory::Memcpy( buffers, &BufferNamesCache[NextBufferName], sizeof(GLuint)*Leftover);
+
+				glGenBuffers( OPENGL_NAME_CACHE_SIZE, BufferNamesCache);
+
+				n -= Leftover;
+				buffers += Leftover;
+
+				FMemory::Memcpy( buffers, BufferNamesCache, sizeof(GLuint)*n);
+				NextBufferName = n;
+			}
+		}
+#else
+		glGenBuffers( n, buffers);
+#endif
+	}
+
+	static FORCEINLINE void GenTextures( GLsizei n, GLuint *textures)
+	{
+#if USE_OPENGL_NAME_CACHE
+		if ( n < OPENGL_NAME_CACHE_SIZE - NextTextureName)
+		{
+			FMemory::Memcpy( textures, &TextureNamesCache[NextTextureName], sizeof(GLuint)*n);
+			NextTextureName += n;
+		}
+		else
+		{
+			if ( n >= OPENGL_NAME_CACHE_SIZE)
+			{
+				glGenTextures( n, textures);
+			}
+			else
+			{
+				GLsizei Leftover = OPENGL_NAME_CACHE_SIZE - NextTextureName;
+
+				FMemory::Memcpy( textures, &TextureNamesCache[NextTextureName], sizeof(GLuint)*Leftover);
+
+				glGenTextures( OPENGL_NAME_CACHE_SIZE, TextureNamesCache);
+
+				n -= Leftover;
+				textures += Leftover;
+
+				FMemory::Memcpy( textures, TextureNamesCache, sizeof(GLuint)*n);
+				NextTextureName = n;
+			}
+		}
+#else
+		glGenTextures( n, textures);
+#endif
+	}
+	
+	static FORCEINLINE void CompressedTexSubImage3D(GLenum Target, GLint Level, GLint XOffset, GLint YOffset, GLint ZOffset, GLsizei Width, GLsizei Height, GLsizei Depth, GLenum Format, GLsizei ImageSize, const GLvoid* PixelData)
+	{
+		glCompressedTexSubImage3D( Target, Level, XOffset, YOffset, ZOffset, Width, Height, Depth, Format, ImageSize, PixelData);
+	}
+
+	static FORCEINLINE void GenerateMipmap( GLenum Target )
+	{
+		glGenerateMipmap( Target);
+	}
+
+	static FORCEINLINE const ANSICHAR* GetStringIndexed(GLenum Name, GLuint Index)
+	{
+		return (const ANSICHAR*)glGetStringi(Name, Index);
+	}
+
+	static FORCEINLINE GLuint GetMajorVersion()
+	{
+		GLint MajorVersion = 0;
+		glGetIntegerv(GL_MAJOR_VERSION, &MajorVersion);
+		return MajorVersion;
+	}
+
+	static FORCEINLINE GLuint GetMinorVersion()
+	{
+		GLint MinorVersion = 0;
+		glGetIntegerv(GL_MINOR_VERSION, &MinorVersion);
+		return MinorVersion;
+	}
+	static FORCEINLINE void ProgramParameter(GLuint Program, GLenum PName, GLint Value)
+	{
+		check(PName != GL_PROGRAM_SEPARABLE);
+		glProgramParameteri(Program, PName, Value);
+	}
+	static FORCEINLINE void ProgramUniform1i(GLuint Program, GLint Location, GLint V0)
+	{
+		glUniform1i(Location, V0);
+	}
+	static FORCEINLINE void ProgramUniform4iv(GLuint Program, GLint Location, GLsizei Count, const GLint *Value)
+	{
+		glUniform4iv(Location, Count, Value);
+	}
+	static FORCEINLINE void ProgramUniform4fv(GLuint Program, GLint Location, GLsizei Count, const GLfloat *Value)
+	{
+		glUniform4fv(Location, Count, Value);
+	}
+	static FORCEINLINE void ProgramUniform4uiv(GLuint Program, GLint Location, GLsizei Count, const GLuint *Value)
+	{
+		glUniform4uiv(Location, Count, Value);
+	}
+	static FORCEINLINE void GetProgramPipelineiv(GLuint Pipeline, GLenum Pname, GLint* Params)
+	{
+		glGetProgramPipelineiv(Pipeline, Pname, Params);
+	}
+	static FORCEINLINE void ValidateProgramPipeline(GLuint Pipeline)
+	{
+		glValidateProgramPipeline(Pipeline);
+	}
+	static FORCEINLINE void GetProgramPipelineInfoLog(GLuint Pipeline, GLsizei BufSize, GLsizei* Length, GLchar* InfoLog)
+	{
+		glGetProgramPipelineInfoLog(Pipeline, BufSize, Length, InfoLog);
+	}
+	static FORCEINLINE bool IsProgramPipeline(GLuint Pipeline)
+	{
+		return (glIsProgramPipeline(Pipeline) == GL_TRUE);
+	}
+
+	static FORCEINLINE ERHIFeatureLevel::Type GetFeatureLevel()
+	{
+		return ERHIFeatureLevel::ES3_1;
+	}
+
+	static FORCEINLINE EShaderPlatform GetShaderPlatform()
+	{
+		return bAndroidGLESCompatibilityMode ? SP_OPENGL_ES3_1_ANDROID : SP_OPENGL_PCES3_1;
+	}
+
+	static FORCEINLINE FString GetAdapterName()
+	{
+		return ANSI_TO_TCHAR((const ANSICHAR*)glGetString(GL_RENDERER));
+	}
+
+	static FPlatformOpenGLDevice*	CreateDevice()	UGL_REQUIRED(NULL)
+	static FPlatformOpenGLContext*	CreateContext( FPlatformOpenGLDevice* Device, void* WindowHandle )	UGL_REQUIRED(NULL)
+
+	static void ProcessQueryGLInt();
+	static void ProcessExtensions(const FString& ExtensionsString);
+	
+	static FORCEINLINE int32 GetReadHalfFloatPixelsEnum() { return GL_HALF_FLOAT; }
+
+	static FORCEINLINE bool IsAndroidGLESCompatibilityModeEnabled() { return bAndroidGLESCompatibilityMode; }
+
+	static GLenum GetPlatfrom5551Format() { return GL_UNSIGNED_SHORT_1_5_5_5_REV; }
+protected:
+	static GLsizei NextTextureName;
+	static GLuint TextureNamesCache[OPENGL_NAME_CACHE_SIZE];
+	static GLsizei NextBufferName;
+	static GLuint BufferNamesCache[OPENGL_NAME_CACHE_SIZE];
+	static GLsizei NextPipelineName;
+	static GLuint PipelineNamesCache[OPENGL_NAME_CACHE_SIZE];
+
+	static GLint TimestampQueryBits;
+	
+	static bool bDebugContext;
+	static bool bAndroidGLESCompatibilityMode;
+};
diff --git a/Engine/Source/Runtime/OpenGLDrv/Private/OpenGL4.h b/Engine/Source/Runtime/OpenGLDrv/Private/OpenGL4.h
new file mode 100644
--- /dev/null
+++ b/Engine/Source/Runtime/OpenGLDrv/Private/OpenGL4.h
@@ -0,0 +1,123 @@
+// Copyright Epic Games, Inc. All Rights Reserved.
+
+/*=============================================================================
+	OpenGL4.h: Public OpenGL 4.3 definitions for non-common functionality
+=============================================================================*/
+
+#pragma once
+
+#define OPENGL_GL4		1
+
+#include "UObject/UObjectHierarchyFwd.h"
+#include "Misc/AssertionMacros.h"
+#include "OpenGL3.h"
+
+struct FOpenGL4 : public FOpenGL3
+{
+	static FORCEINLINE bool SupportsDrawIndirect()						{ return true; }
+
+	// Optional
+
+	// Required
+	static FORCEINLINE void BlendFuncSeparatei(GLuint Buf, GLenum SrcRGB, GLenum DstRGB, GLenum SrcAlpha, GLenum DstAlpha)
+	{
+		glBlendFuncSeparatei(Buf, SrcRGB, DstRGB, SrcAlpha, DstAlpha);
+	}
+	static FORCEINLINE void BlendEquationSeparatei(GLuint Buf, GLenum ModeRGB, GLenum ModeAlpha)
+	{
+		glBlendEquationSeparatei(Buf, ModeRGB, ModeAlpha);
+	}
+	static FORCEINLINE void BlendFunci(GLuint Buf, GLenum Src, GLenum Dst)
+	{
+		glBlendFunci(Buf, Src, Dst);
+	}
+	static FORCEINLINE void BlendEquationi(GLuint Buf, GLenum Mode)
+	{
+		glBlendEquationi(Buf, Mode);
+	}
+	static FORCEINLINE void PatchParameteri(GLenum Pname, GLint Value)
+	{
+		glPatchParameteri(Pname, Value);
+	}
+	static FORCEINLINE void BindImageTexture(GLuint Unit, GLuint Texture, GLint Level, GLboolean Layered, GLint Layer, GLenum Access, GLenum Format)
+	{
+		glBindImageTexture(Unit, Texture, Level, Layered, Layer, Access, Format);
+	}
+	static FORCEINLINE void DispatchCompute(GLuint NumGroupsX, GLuint NumGroupsY, GLuint NumGroupsZ)
+	{
+		glDispatchCompute(NumGroupsX, NumGroupsY, NumGroupsZ);
+	}
+	static FORCEINLINE void DispatchComputeIndirect(GLintptr Offset)
+	{
+		glDispatchComputeIndirect(Offset);
+	}
+	static FORCEINLINE void MemoryBarrier(GLbitfield Barriers)
+	{
+		glMemoryBarrier(Barriers);
+	}
+	static FORCEINLINE void DrawArraysIndirect (GLenum Mode, const void *Offset)
+	{
+		glDrawArraysIndirect( Mode, Offset);
+	}
+	static FORCEINLINE void DrawElementsIndirect (GLenum Mode, GLenum Type, const void *Offset)
+	{
+		glDrawElementsIndirect( Mode, Type, Offset);
+	}
+	static FORCEINLINE void BindVertexBuffer(GLuint BindingIndex, GLuint Buffer, GLintptr Offset, GLsizei Stride)
+	{
+		glBindVertexBuffer(BindingIndex, Buffer, Offset, Stride);
+	}
+	static FORCEINLINE void VertexAttribFormat(GLuint AttribIndex, GLint Size, GLenum Type, GLboolean Normalized, GLuint RelativeOffset)
+	{
+		glVertexAttribFormat(AttribIndex, Size, Type, Normalized, RelativeOffset);
+	}
+	static FORCEINLINE void VertexAttribIFormat(GLuint AttribIndex, GLint Size, GLenum Type, GLuint RelativeOffset)
+	{
+		glVertexAttribIFormat(AttribIndex, Size, Type, RelativeOffset);
+	}
+	static FORCEINLINE void VertexAttribBinding(GLuint AttribIndex, GLuint BindingIndex)
+	{
+		glVertexAttribBinding(AttribIndex, BindingIndex);
+	}
+	static FORCEINLINE void VertexBindingDivisor(GLuint BindingIndex, GLuint Divisor)
+	{
+		glVertexBindingDivisor(BindingIndex, Divisor);
+	}
+	static FORCEINLINE void TexBufferRange(GLenum Target, GLenum InternalFormat, GLuint Buffer, GLintptr Offset, GLsizeiptr Size)
+	{
+		glTexBufferRange(Target, InternalFormat, Buffer, Offset, Size);
+	}
+	static FORCEINLINE void TextureView(GLuint ViewName, GLenum ViewTarget, GLuint SrcName, GLenum InternalFormat, GLuint MinLevel, GLuint NumLevels, GLuint MinLayer, GLuint NumLayers)
+	{
+		glTextureView(ViewName, ViewTarget, SrcName, InternalFormat, MinLevel, NumLevels, MinLayer, NumLayers);
+	}
+	static FORCEINLINE void ClearBufferData(GLenum Target, GLenum InternalFormat, GLenum Format, GLenum Type, const uint32* Data)
+	{
+		glClearBufferData(Target, InternalFormat, Format, Type, Data);
+	}
+
+	static void ProcessQueryGLInt();
+	static void ProcessExtensions( const FString& ExtensionsString );
+
+	static FORCEINLINE GLint GetMaxComputeUniformComponents() { check(MaxComputeUniformComponents != -1); return MaxComputeUniformComponents; }
+
+	static FORCEINLINE GLint GetFirstComputeUAVUnit()			{ return 0; }
+	static FORCEINLINE GLint GetMaxComputeUAVUnits()			{ check(MaxComputeUAVUnits != -1); return MaxComputeUAVUnits; }
+	static FORCEINLINE GLint GetFirstVertexUAVUnit()			{ return 0; }
+	static FORCEINLINE GLint GetFirstPixelUAVUnit()				{ return 0; }
+	static FORCEINLINE GLint GetMaxPixelUAVUnits()				{ check(MaxPixelUAVUnits != -1); return MaxPixelUAVUnits; }
+	static FORCEINLINE GLint GetMaxCombinedUAVUnits()			{ return MaxCombinedUAVUnits; }
+		
+	static uint64 GetVideoMemorySize();
+
+	static FORCEINLINE int32 GetReadHalfFloatPixelsEnum() { return GL_HALF_FLOAT; }
+
+protected:
+	static GLint MaxComputeUniformComponents;
+
+	static GLint MaxCombinedUAVUnits;
+	static GLint MaxComputeUAVUnits;
+	static GLint MaxPixelUAVUnits;
+
+	static bool bSupportsGPUMemoryInfo;
+};
diff --git a/Engine/Source/Runtime/OpenGLDrv/Private/OpenGLDrv.h b/Engine/Source/Runtime/OpenGLDrv/Private/OpenGLDrv.h
new file mode 100644
--- /dev/null
+++ b/Engine/Source/Runtime/OpenGLDrv/Private/OpenGLDrv.h
@@ -0,0 +1,961 @@
+// Copyright Epic Games, Inc. All Rights Reserved.
+
+/*=============================================================================
+	OpenGLDrv.h: Public OpenGL RHI definitions.
+=============================================================================*/
+
+#pragma once
+
+// Dependencies.
+
+#include "CoreMinimal.h"
+#include "Containers/IndirectArray.h"
+#include "IOpenGLDynamicRHI.h"
+#include "RHIDefinitions.h"
+#include "RHI.h"
+#include "GPUProfiler.h"
+#include "RenderResource.h"
+#include "Templates/EnableIf.h"
+#include "BoundShaderStateHistory.h"
+
+// @todo platplug: Replace all of these includes with a call to COMPILED_PLATFORM_HEADER(OpenGLDrvPrivate.h)
+//TODO: Move these to OpenGLDrvPrivate.h
+#if PLATFORM_WINDOWS
+	#include "Windows/OpenGLWindows.h"
+#elif PLATFORM_LINUX
+	#include "Linux/OpenGLLinux.h"
+#elif PLATFORM_ANDROID
+	#include "Android/AndroidOpenGL.h"
+#else
+#include COMPILED_PLATFORM_HEADER(OpenGLDrvPrivate.h)
+#endif
+
+// Define here so don't have to do platform filtering
+#ifndef GL_TEXTURE_EXTERNAL_OES
+#define GL_TEXTURE_EXTERNAL_OES 0x8D65
+#endif
+
+// OpenGL RHI public headers.
+#include "OpenGLUtil.h"
+#include "OpenGLState.h"
+#include "RenderUtils.h"
+
+#define FOpenGLCachedUniformBuffer_Invalid 0xFFFFFFFF
+
+class FOpenGLDynamicRHI;
+class FResourceBulkDataInterface;
+struct FOpenGLResourceBinder;
+struct Rect;
+
+template<class T> struct TOpenGLResourceTraits;
+
+// This class has multiple inheritance but really FGPUTiming is a static class
+class FOpenGLBufferedGPUTiming : public FGPUTiming
+{
+public:
+
+	/**
+	 * Constructor.
+	 *
+	 * @param InOpenGLRHI			RHI interface
+	 * @param InBufferSize		Number of buffered measurements
+	 */
+	FOpenGLBufferedGPUTiming(class FOpenGLDynamicRHI* InOpenGLRHI, int32 BufferSize);
+
+	void	StartTiming();
+
+	/**
+	 * End a GPU timing measurement.
+	 * The timing for this particular measurement will be resolved at a later time by the GPU.
+	 */
+	void	EndTiming();
+
+	/**
+	 * Retrieves the most recently resolved timing measurement.
+	 * The unit is the same as for FPlatformTime::Cycles(). Returns 0 if there are no resolved measurements.
+	 *
+	 * @return	Value of the most recently resolved timing, or 0 if no measurements have been resolved by the GPU yet.
+	 */
+	uint64	GetTiming(bool bGetCurrentResultsAndBlock = false);
+
+	void InitResources();
+	void ReleaseResources();
+
+
+private:
+
+	/**
+	 * Initializes the static variables, if necessary.
+	 */
+	static void PlatformStaticInitialize(void* UserData);
+
+	/** RHI interface */
+	FOpenGLDynamicRHI*					OpenGLRHI;
+	/** Number of timestamps created in 'StartTimestamps' and 'EndTimestamps'. */
+	int32								BufferSize;
+	/** Current timing being measured on the CPU. */
+	int32								CurrentTimestamp;
+	/** Number of measurements in the buffers (0 - BufferSize). */
+	int32								NumIssuedTimestamps;
+	/** Timestamps for all StartTimings. */
+	TArray<FOpenGLRenderQuery *>		StartTimestamps;
+	/** Timestamps for all EndTimings. */
+	TArray<FOpenGLRenderQuery *>		EndTimestamps;
+	/** Whether we are currently timing the GPU: between StartTiming() and EndTiming(). */
+	bool								bIsTiming;
+};
+
+/**
+  * Used to track whether a period was disjoint on the GPU, which means GPU timings are invalid.
+  * OpenGL lacks this concept at present, so the class is just a placeholder
+  * Timings are all assumed to be non-disjoint
+  */
+class FOpenGLDisjointTimeStampQuery
+{
+public:
+	FOpenGLDisjointTimeStampQuery(class FOpenGLDynamicRHI* InOpenGLRHI=NULL);
+
+	void Init(class FOpenGLDynamicRHI* InOpenGLRHI)
+	{
+		OpenGLRHI = InOpenGLRHI;
+		InitResources();
+	}
+
+	void StartTracking();
+	void EndTracking();
+	bool IsResultValid();
+	bool GetResult(uint64* OutResult=NULL);
+	static uint64 GetTimingFrequency()
+	{
+		return 1000000000ull;
+	}
+	static bool IsSupported()
+	{
+#if UE_BUILD_SHIPPING
+		return false;
+#else
+		return FOpenGL::SupportsDisjointTimeQueries();
+#endif
+	}
+
+	void InitResources();
+	void ReleaseResources();
+
+
+private:
+	bool	bIsResultValid;
+	GLuint	DisjointQuery;
+	uint64	Context;
+
+	FOpenGLDynamicRHI* OpenGLRHI;
+};
+
+/** A single perf event node, which tracks information about a appBeginDrawEvent/appEndDrawEvent range. */
+class FOpenGLEventNode : public FGPUProfilerEventNode
+{
+public:
+
+	FOpenGLEventNode(const TCHAR* InName, FGPUProfilerEventNode* InParent, class FOpenGLDynamicRHI* InRHI)
+	:	FGPUProfilerEventNode(InName, InParent)
+	,	Timing(InRHI, 1)
+	{
+		// Initialize Buffered timestamp queries 
+		Timing.InitResources();
+	}
+
+	virtual ~FOpenGLEventNode()
+	{
+		Timing.ReleaseResources();
+	}
+
+	/** 
+	 * Returns the time in ms that the GPU spent in this draw event.  
+	 * This blocks the CPU if necessary, so can cause hitching.
+	 */
+	float GetTiming() override;
+
+	virtual void StartTiming() override
+	{
+		Timing.StartTiming();
+	}
+
+	virtual void StopTiming() override
+	{
+		Timing.EndTiming();
+	}
+
+	FOpenGLBufferedGPUTiming Timing;
+};
+
+/** An entire frame of perf event nodes, including ancillary timers. */
+class FOpenGLEventNodeFrame : public FGPUProfilerEventNodeFrame
+{
+public:
+	FOpenGLEventNodeFrame(class FOpenGLDynamicRHI* InRHI) :
+		FGPUProfilerEventNodeFrame(),
+		RootEventTiming(InRHI, 1),
+		DisjointQuery(InRHI)
+	{
+	  RootEventTiming.InitResources();
+	  DisjointQuery.InitResources();
+	}
+
+	~FOpenGLEventNodeFrame()
+	{
+
+		RootEventTiming.ReleaseResources();
+		DisjointQuery.ReleaseResources();
+	}
+
+	/** Start this frame of per tracking */
+	void StartFrame() override;
+
+	/** End this frame of per tracking, but do not block yet */
+	void EndFrame() override;
+
+	/** Calculates root timing base frequency (if needed by this RHI) */
+	virtual float GetRootTimingResults() override;
+
+	virtual void LogDisjointQuery() override;
+
+	/** Timer tracking inclusive time spent in the root nodes. */
+	FOpenGLBufferedGPUTiming RootEventTiming;
+
+	/** Disjoint query tracking whether the times reported by DumpEventTree are reliable. */
+	FOpenGLDisjointTimeStampQuery DisjointQuery;
+};
+
+/** 
+ * Encapsulates GPU profiling logic and data. 
+ * There's only one global instance of this struct so it should only contain global data, nothing specific to a frame.
+ */
+struct FOpenGLGPUProfiler : public FGPUProfiler
+{
+	/** Used to measure GPU time per frame. */
+	FOpenGLBufferedGPUTiming FrameTiming;
+
+	/** Measuring GPU frame time with a disjoint query. */
+	static const int MAX_GPUFRAMEQUERIES = 4;
+	FOpenGLDisjointTimeStampQuery DisjointGPUFrameTimeQuery[MAX_GPUFRAMEQUERIES];
+	int CurrentGPUFrameQueryIndex;
+
+	class FOpenGLDynamicRHI* OpenGLRHI;
+	// count the number of beginframe calls without matching endframe calls.
+	int32 NestedFrameCount;
+
+	uint32 ExternalGPUTime;
+
+	/** GPU hitch profile histories */
+	TIndirectArray<FOpenGLEventNodeFrame> GPUHitchEventNodeFrames;
+
+	FOpenGLGPUProfiler(class FOpenGLDynamicRHI* InOpenGLRHI)
+	:	FGPUProfiler()
+	,	FrameTiming(InOpenGLRHI, 4)
+	,	CurrentGPUFrameQueryIndex(0)
+	,	OpenGLRHI(InOpenGLRHI)
+	,	NestedFrameCount(0)
+	,	ExternalGPUTime(0)
+	{
+		FrameTiming.InitResources();
+		for (int32 Index = 0; Index < MAX_GPUFRAMEQUERIES; ++Index)
+		{
+			DisjointGPUFrameTimeQuery[Index].Init(OpenGLRHI);
+		}
+	}
+
+	virtual FGPUProfilerEventNode* CreateEventNode(const TCHAR* InName, FGPUProfilerEventNode* InParent) override
+	{
+		FOpenGLEventNode* EventNode = new FOpenGLEventNode(InName, InParent, OpenGLRHI);
+		return EventNode;
+	}
+
+	void Cleanup();
+
+	virtual void PushEvent(const TCHAR* Name, FColor Color) override;
+	virtual void PopEvent() override;
+
+	void BeginFrame(class FOpenGLDynamicRHI* InRHI);
+	void EndFrame();
+};
+
+PRAGMA_DISABLE_DEPRECATION_WARNINGS
+
+/** The interface which is implemented by the dynamically bound RHI. */
+class OPENGLDRV_API FOpenGLDynamicRHI final : public IOpenGLDynamicRHI, public IRHICommandContextPSOFallback
+{
+	static inline FOpenGLDynamicRHI* Singleton = nullptr;
+
+public:
+	static inline FOpenGLDynamicRHI& Get() { return *Singleton; }
+
+	friend class FOpenGLViewport;
+
+	/** Initialization constructor. */
+	FOpenGLDynamicRHI();
+
+	/** Destructor */
+	~FOpenGLDynamicRHI() {}
+
+	// IOpenGLDynamicRHI interface.
+	virtual int32 RHIGetGLMajorVersion() const final override;
+	virtual int32 RHIGetGLMinorVersion() const final override;
+	virtual bool RHISupportsFramebufferSRGBEnable() const final override;
+	virtual FTexture2DRHIRef RHICreateTexture2DFromResource(EPixelFormat Format, uint32 SizeX, uint32 SizeY, uint32 NumMips, uint32 NumSamples, uint32 NumSamplesTileMem, const FClearValueBinding& ClearValueBinding, GLuint Resource, ETextureCreateFlags Flags) final override;
+	virtual FTexture2DRHIRef RHICreateTexture2DArrayFromResource(EPixelFormat Format, uint32 SizeX, uint32 SizeY, uint32 ArraySize, uint32 NumMips, uint32 NumSamples, uint32 NumSamplesTileMem, const FClearValueBinding& ClearValueBinding, GLuint Resource, ETextureCreateFlags Flags) final override;
+	virtual FTextureCubeRHIRef RHICreateTextureCubeFromResource(EPixelFormat Format, uint32 Size, bool bArray, uint32 ArraySize, uint32 NumMips, uint32 NumSamples, uint32 NumSamplesTileMem, const FClearValueBinding& ClearValueBinding, GLuint Resource, ETextureCreateFlags Flags) final override;
+	virtual GLuint RHIGetResource(FRHITexture* InTexture) const final override;
+	virtual bool RHIIsValidTexture(GLuint InTexture) const final override;
+	virtual void RHISetExternalGPUTime(uint32 InExternalGPUTime) final override;
+
+#if PLATFORM_ANDROID
+	virtual EGLDisplay RHIGetEGLDisplay() const final override;
+	virtual EGLSurface RHIGetEGLSurface() const final override;
+	virtual EGLConfig  RHIGetEGLConfig() const final override;
+	virtual EGLContext RHIGetEGLContext() const final override;
+	virtual ANativeWindow* RHIGetEGLNativeWindow() const final override;
+	virtual bool RHIEGLSupportsNoErrorContext() const final override;
+	virtual void RHIInitEGLInstanceGLES2() final override;
+	virtual void RHIInitEGLBackBuffer() final override;
+	virtual void RHIEGLSetCurrentRenderingContext() final override;
+	virtual void RHIEGLTerminateContext() final override;
+#endif
+
+	// FDynamicRHI interface.
+	virtual void Init() override;
+
+	virtual void Shutdown() override;
+	virtual const TCHAR* GetName() override { return TEXT("OpenGL"); }
+
+	template<typename TRHIType>
+	static auto* ResourceCast(TRHIType* Resource)
+	{
+		return static_cast<typename TOpenGLResourceTraits<TRHIType>::TConcreteType*>(Resource);
+	}
+
+	static FOpenGLTexture* ResourceCast(FRHITexture* TextureRHI)
+	{
+		if (!TextureRHI)
+		{
+			return nullptr;
+		}
+		else
+		{
+			return static_cast<FOpenGLTexture*>(TextureRHI->GetTextureBaseRHI());
+		}
+	}
+
+	void BindUniformBuffer(EShaderFrequency ShaderFrequency, uint32 BufferIndex, FRHIUniformBuffer* BufferRHI);
+
+	void SetShaderParametersCommon(EShaderFrequency ShaderFrequency, TConstArrayView<uint8> InParametersData, TConstArrayView<FRHIShaderParameter> InParameters, TConstArrayView<FRHIShaderParameterResource> InResourceParameters);
+	void SetShaderUnbindsCommon(EShaderFrequency ShaderFrequency, TConstArrayView<FRHIShaderParameterUnbind> InUnbinds);
+
+	virtual FRasterizerStateRHIRef RHICreateRasterizerState(const FRasterizerStateInitializerRHI& Initializer) final override;
+	virtual FDepthStencilStateRHIRef RHICreateDepthStencilState(const FDepthStencilStateInitializerRHI& Initializer) final override;
+	virtual FBlendStateRHIRef RHICreateBlendState(const FBlendStateInitializerRHI& Initializer) final override;
+	virtual FVertexDeclarationRHIRef RHICreateVertexDeclaration(const FVertexDeclarationElementList& Elements) final override;
+
+	virtual FPixelShaderRHIRef RHICreatePixelShader(TArrayView<const uint8> Code, const FSHAHash& Hash) final override;
+	virtual FVertexShaderRHIRef RHICreateVertexShader(TArrayView<const uint8> Code, const FSHAHash& Hash) final override;
+	virtual FGeometryShaderRHIRef RHICreateGeometryShader(TArrayView<const uint8> Code, const FSHAHash& Hash) final override;
+	virtual FComputeShaderRHIRef RHICreateComputeShader(TArrayView<const uint8> Code, const FSHAHash& Hash) final override;
+
+	// SRV / UAV creation functions
+	virtual FShaderResourceViewRHIRef  RHICreateShaderResourceView (class FRHICommandListBase& RHICmdList, FRHIViewableResource* Resource, FRHIViewDesc const& ViewDesc) final override;
+	virtual FUnorderedAccessViewRHIRef RHICreateUnorderedAccessView(class FRHICommandListBase& RHICmdList, FRHIViewableResource* Resource, FRHIViewDesc const& ViewDesc) final override;
+
+	virtual FUniformBufferRHIRef RHICreateUniformBuffer(const void* Contents, const FRHIUniformBufferLayout* Layout, EUniformBufferUsage Usage, EUniformBufferValidation Validation) final override;
+	virtual void RHIUpdateUniformBuffer(FRHICommandListBase& RHICmdList, FRHIUniformBuffer* UniformBufferRHI, const void* Contents) final override;
+	virtual FBufferRHIRef RHICreateBuffer(FRHICommandListBase& RHICmdList, FRHIBufferDesc const& Desc, ERHIAccess ResourceState, FRHIResourceCreateInfo& CreateInfo) final override;
+	virtual void RHICopyBuffer(FRHIBuffer* SourceBuffer, FRHIBuffer* DestBuffer) final override;
+	virtual void RHITransferBufferUnderlyingResource(FRHICommandListBase& RHICmdList, FRHIBuffer* DestBuffer, FRHIBuffer* SrcBuffer) final override;
+	virtual void* LockBuffer_BottomOfPipe(FRHICommandListBase& RHICmdList, FRHIBuffer* Buffer, uint32 Offset, uint32 Size, EResourceLockMode LockMode) final override;
+	virtual void UnlockBuffer_BottomOfPipe(FRHICommandListBase& RHICmdList, FRHIBuffer* Buffer) final override;
+	virtual FRHICalcTextureSizeResult RHICalcTexturePlatformSize(FRHITextureDesc const& Desc, uint32 FirstMipIndex) final override;
+	virtual void RHIGetTextureMemoryStats(FTextureMemoryStats& OutStats) final override;
+	virtual bool RHIGetTextureMemoryVisualizeData(FColor* TextureData,int32 SizeX,int32 SizeY,int32 Pitch,int32 PixelSize) final override;
+	virtual FTextureRHIRef RHICreateTexture(FRHICommandListBase& RHICmdList, const FRHITextureCreateDesc& CreateDesc) final override;
+	virtual FTextureRHIRef RHIAsyncCreateTexture2D(uint32 SizeX, uint32 SizeY, uint8 Format, uint32 NumMips, ETextureCreateFlags Flags, ERHIAccess InResourceState, void** InitialMipData, uint32 NumInitialMips, const TCHAR* DebugName, FGraphEventRef& OutCompletionEvent) final override;
+	virtual void RHIGenerateMips(FRHITexture* Texture) final override;
+	virtual uint32 RHIComputeMemorySize(FRHITexture* TextureRHI) final override;
+	virtual FTexture2DRHIRef RHIAsyncReallocateTexture2D(FRHITexture2D* Texture2D, int32 NewMipCount, int32 NewSizeX, int32 NewSizeY, FThreadSafeCounter* RequestStatus) final override;
+	virtual ETextureReallocationStatus RHIFinalizeAsyncReallocateTexture2D(FRHITexture2D* Texture2D, bool bBlockUntilCompleted) final override;
+	virtual ETextureReallocationStatus RHICancelAsyncReallocateTexture2D(FRHITexture2D* Texture2D, bool bBlockUntilCompleted) final override;
+	virtual void* RHILockTexture2D(FRHITexture2D* Texture, uint32 MipIndex, EResourceLockMode LockMode, uint32& DestStride, bool bLockWithinMiptail, uint64* OutLockedByteCount = nullptr) final override;
+	virtual void RHIUnlockTexture2D(FRHITexture2D* Texture, uint32 MipIndex, bool bLockWithinMiptail) final override;
+	virtual void* RHILockTexture2DArray(FRHITexture2DArray* Texture, uint32 TextureIndex, uint32 MipIndex, EResourceLockMode LockMode, uint32& DestStride, bool bLockWithinMiptail) final override;
+	virtual void RHIUnlockTexture2DArray(FRHITexture2DArray* Texture, uint32 TextureIndex, uint32 MipIndex, bool bLockWithinMiptail) final override;
+	virtual void RHIUpdateTexture2D(FRHICommandListBase& RHICmdList, FRHITexture2D* Texture, uint32 MipIndex, const struct FUpdateTextureRegion2D& UpdateRegion, uint32 SourcePitch, const uint8* SourceData) final override;
+	virtual void RHIUpdateTexture3D(FRHICommandListBase& RHICmdList, FRHITexture3D* Texture, uint32 MipIndex, const struct FUpdateTextureRegion3D& UpdateRegion, uint32 SourceRowPitch, uint32 SourceDepthPitch, const uint8* SourceData) final override;
+	virtual void* RHILockTextureCubeFace(FRHITextureCube* Texture, uint32 FaceIndex, uint32 ArrayIndex, uint32 MipIndex, EResourceLockMode LockMode, uint32& DestStride, bool bLockWithinMiptail) final override;
+	virtual void RHIUnlockTextureCubeFace(FRHITextureCube* Texture, uint32 FaceIndex, uint32 ArrayIndex, uint32 MipIndex, bool bLockWithinMiptail) final override;
+	virtual void RHIBindDebugLabelName(FRHICommandListBase& RHICmdList, FRHITexture* Texture, const TCHAR* Name) final override;
+	virtual void RHIReadSurfaceData(FRHITexture* Texture,FIntRect Rect,TArray<FColor>& OutData,FReadSurfaceDataFlags InFlags) final override;
+	virtual void RHIReadSurfaceData(FRHITexture* Texture, FIntRect Rect, TArray<FLinearColor>& OutData, FReadSurfaceDataFlags InFlags) final override;
+	virtual void RHIMapStagingSurface(FRHITexture* Texture, FRHIGPUFence* Fence, void*& OutData, int32& OutWidth, int32& OutHeight, uint32 GPUIndex = 0) final override;
+	virtual void RHIUnmapStagingSurface(FRHITexture* Texture, uint32 GPUIndex = 0) final override;
+	virtual void RHIReadSurfaceFloatData(FRHITexture* Texture,FIntRect Rect,TArray<FFloat16Color>& OutData,ECubeFace CubeFace,int32 ArrayIndex,int32 MipIndex) final override;
+	virtual void RHIRead3DSurfaceFloatData(FRHITexture* Texture,FIntRect Rect,FIntPoint ZMinMax,TArray<FFloat16Color>& OutData) final override;
+	virtual FRenderQueryRHIRef RHICreateRenderQuery(ERenderQueryType QueryType) final override;
+	virtual bool RHIGetRenderQueryResult(FRHIRenderQuery* RenderQuery, uint64& OutResult, bool bWait, uint32 GPUIndex = INDEX_NONE) final override;
+	virtual FTexture2DRHIRef RHIGetViewportBackBuffer(FRHIViewport* Viewport) final override;
+	virtual void RHIAliasTextureResources(FTextureRHIRef& DestTextureRHI, FTextureRHIRef& SrcTextureRHI) final override;
+	virtual FTextureRHIRef RHICreateAliasedTexture(FTextureRHIRef& SourceTexture) final override;
+	virtual void RHIAdvanceFrameForGetViewportBackBuffer(FRHIViewport* Viewport) final override;
+	virtual void RHIAcquireThreadOwnership() final override;
+	virtual void RHIReleaseThreadOwnership() final override;
+	virtual void RHIFlushResources() final override;
+	virtual uint32 RHIGetGPUFrameCycles(uint32 GPUIndex = 0) final override;
+	virtual FViewportRHIRef RHICreateViewport(void* WindowHandle, uint32 SizeX, uint32 SizeY, bool bIsFullscreen, EPixelFormat PreferredPixelFormat) final override;
+	virtual void RHIResizeViewport(FRHIViewport* Viewport, uint32 SizeX, uint32 SizeY, bool bIsFullscreen) final override;
+	virtual EPixelFormat RHIPreferredPixelFormatHint(EPixelFormat PreferredPixelFormat) final override;
+	virtual void RHITick(float DeltaTime) final override;
+	virtual void RHIBlockUntilGPUIdle() final override;
+	virtual void RHIPollOcclusionQueries() final override;
+	virtual bool RHIGetAvailableResolutions(FScreenResolutionArray& Resolutions, bool bIgnoreRefreshRate) final override;
+	virtual void RHIGetSupportedResolution(uint32& Width, uint32& Height) final override;
+	virtual void* RHIGetNativeDevice() final override;
+	virtual void* RHIGetNativeInstance() final override;
+	virtual class IRHICommandContext* RHIGetDefaultContext() final override;
+	virtual IRHIComputeContext* RHIGetCommandContext(ERHIPipeline Pipeline, FRHIGPUMask GPUMask) final override;
+	virtual IRHIPlatformCommandList* RHIFinalizeContext(FRHIFinalizeContextArgs&& Args) final override;
+	virtual void RHISubmitCommandLists(FRHISubmitCommandListsArgs&& Args) final override;
+
+	virtual void RHIBeginRenderPass(const FRHIRenderPassInfo& InInfo, const TCHAR* InName) final override;
+	virtual void RHIEndRenderPass() final override;
+	virtual void RHINextSubpass() final override;
+	
+	virtual void RHIBeginTransitions(TArrayView<const FRHITransition*> Transitions) override final;
+	virtual void RHIEndTransitions(TArrayView<const FRHITransition*> Transitions) override final;
+
+	virtual void RHISetComputeShader(FRHIComputeShader* ComputeShader) final override;
+	virtual void RHIDispatchComputeShader(uint32 ThreadGroupCountX, uint32 ThreadGroupCountY, uint32 ThreadGroupCountZ) final override;
+	virtual void RHIDispatchIndirectComputeShader(FRHIBuffer* ArgumentBuffer, uint32 ArgumentOffset) final override;
+	virtual void RHISetMultipleViewports(uint32 Count, const FViewportBounds* Data) final override;
+	virtual void RHIClearUAVFloat(FRHIUnorderedAccessView* UnorderedAccessViewRHI, const FVector4f& Values) final override;
+	virtual void RHIClearUAVUint(FRHIUnorderedAccessView* UnorderedAccessViewRHI, const FUintVector4& Values) final override;
+	virtual void RHIBeginRenderQuery(FRHIRenderQuery* RenderQuery) final override;
+	virtual void RHIEndRenderQuery(FRHIRenderQuery* RenderQuery) final override;
+	virtual void RHIBeginDrawingViewport(FRHIViewport* Viewport, FRHITexture* RenderTargetRHI) final override;
+	virtual void RHIEndDrawingViewport(FRHIViewport* Viewport, bool bPresent, bool bLockToVsync) final override;
+	using FDynamicRHI::RHIBeginFrame;
+	virtual void RHIBeginFrame() final override;
+	virtual void RHIEndFrame() final override;
+	virtual void RHIBeginScene() final override;
+	virtual void RHIEndScene() final override;
+	virtual void RHISetStreamSource(uint32 StreamIndex, FRHIBuffer* VertexBuffer, uint32 Offset) final override;
+	virtual void RHISetRasterizerState(FRHIRasterizerState* NewState) final override;
+	virtual void RHISetViewport(float MinX, float MinY, float MinZ, float MaxX, float MaxY, float MaxZ) final override;
+	virtual void RHISetScissorRect(bool bEnable, uint32 MinX, uint32 MinY, uint32 MaxX, uint32 MaxY) final override;
+	virtual void RHISetBoundShaderState(FRHIBoundShaderState* BoundShaderState) final override;
+	virtual void RHISetStaticUniformBuffers(const FUniformBufferStaticBindings& InUniformBuffers) final override;
+	virtual void RHISetStaticUniformBuffer(FUniformBufferStaticSlot Slot, FRHIUniformBuffer* Buffer) final override;
+	virtual void RHISetUniformBufferDynamicOffset(FUniformBufferStaticSlot Slot, uint32 Offset) final override;
+	virtual void RHISetShaderParameters(FRHIGraphicsShader* Shader, TConstArrayView<uint8> InParametersData, TConstArrayView<FRHIShaderParameter> InParameters, TConstArrayView<FRHIShaderParameterResource> InResourceParameters, TConstArrayView<FRHIShaderParameterResource> InBindlessParameters) final override;
+	virtual void RHISetShaderParameters(FRHIComputeShader* Shader, TConstArrayView<uint8> InParametersData, TConstArrayView<FRHIShaderParameter> InParameters, TConstArrayView<FRHIShaderParameterResource> InResourceParameters, TConstArrayView<FRHIShaderParameterResource> InBindlessParameters) final override;
+	virtual void RHISetShaderUnbinds(FRHIComputeShader* Shader, TConstArrayView<FRHIShaderParameterUnbind> InUnbinds) final override;
+	virtual void RHISetShaderUnbinds(FRHIGraphicsShader* Shader, TConstArrayView<FRHIShaderParameterUnbind> InUnbinds) final override;
+	virtual void RHISetDepthStencilState(FRHIDepthStencilState* NewState, uint32 StencilRef) final override;
+	virtual void RHISetStencilRef(uint32 StencilRef) final override;
+	virtual void RHISetBlendState(FRHIBlendState* NewState, const FLinearColor& BlendFactor) final override;
+	virtual void RHISetBlendFactor(const FLinearColor& BlendFactor) final override
+	{
+		// Currently ignored as well as on RHISetBlendState()...
+	}
+
+	void SetRenderTargets(uint32 NumSimultaneousRenderTargets, const FRHIRenderTargetView* NewRenderTargets, const FRHIDepthRenderTargetView* NewDepthStencilTarget);
+	void SetRenderTargetsAndClear(const FRHISetRenderTargetsInfo& RenderTargetsInfo);
+
+	virtual void RHIDrawPrimitive(uint32 BaseVertexIndex, uint32 NumPrimitives, uint32 NumInstances) final override;
+	virtual void RHIDrawPrimitiveIndirect(FRHIBuffer* ArgumentBuffer, uint32 ArgumentOffset) final override;
+	virtual void RHIDrawIndexedIndirect(FRHIBuffer* IndexBufferRHI, FRHIBuffer* ArgumentsBufferRHI, int32 DrawArgumentsIndex, uint32 NumInstances) final override;
+	virtual void RHIDrawIndexedPrimitive(FRHIBuffer* IndexBuffer, int32 BaseVertexIndex, uint32 FirstInstance, uint32 NumVertices, uint32 StartIndex, uint32 NumPrimitives, uint32 NumInstances) final override;
+	virtual void RHIDrawIndexedPrimitiveIndirect(FRHIBuffer* IndexBuffer, FRHIBuffer* ArgumentBuffer, uint32 ArgumentOffset) final override;
+	virtual void RHIEnableDepthBoundsTest(bool bEnable) final override;
+	virtual void RHISetDepthBounds(float MinDepth, float MaxDepth) final override;
+#if WITH_RHI_BREADCRUMBS
+	virtual void RHIBeginBreadcrumbGPU(FRHIBreadcrumbNode* Breadcrumb) final override;
+	virtual void RHIEndBreadcrumbGPU  (FRHIBreadcrumbNode* Breadcrumb) final override;
+#endif
+	virtual void RHIInvalidateCachedState() final override;
+	virtual void RHIDiscardRenderTargets(bool Depth,bool Stencil,uint32 ColorBitMask) final override;
+
+	// FIXME: Broken on Android for cubemaps
+	virtual void RHICopyTexture(FRHITexture* SourceTexture, FRHITexture* DestTexture, const FRHICopyTextureInfo& CopyInfo) final override;
+
+	virtual void RHICopyBufferRegion(FRHIBuffer* DestBuffer, uint64 DstOffset, FRHIBuffer* SourceBuffer, uint64 SrcOffset, uint64 NumBytes) final override;
+
+	// Inline copy
+	virtual void RHICopyToStagingBuffer(FRHIBuffer* SourceBufferRHI, FRHIStagingBuffer* DestinationStagingBufferRHI, uint32 InOffset, uint32 InNumBytes) final override;
+	virtual void RHIWriteGPUFence(FRHIGPUFence* FenceRHI) final override;
+	virtual void* RHILockStagingBuffer(FRHIStagingBuffer* StagingBuffer, FRHIGPUFence* Fence, uint32 Offset, uint32 SizeRHI) final override;
+	virtual void RHIUnlockStagingBuffer(FRHIStagingBuffer* StagingBuffer) final override;
+	virtual void* LockStagingBuffer_RenderThread(class FRHICommandListImmediate& RHICmdList, FRHIStagingBuffer* StagingBuffer, FRHIGPUFence* Fence, uint32 Offset, uint32 SizeRHI) final override;
+	virtual void UnlockStagingBuffer_RenderThread(class FRHICommandListImmediate& RHICmdList, FRHIStagingBuffer* StagingBuffer) final override;
+	virtual void RHIMapStagingSurface_RenderThread(class FRHICommandListImmediate& RHICmdList, FRHITexture* Texture, uint32 GPUIndex, FRHIGPUFence* Fence, void*& OutData, int32& OutWidth, int32& OutHeight) final override;
+	virtual void RHIUnmapStagingSurface_RenderThread(class FRHICommandListImmediate& RHICmdList, FRHITexture* Texture, uint32 GPUIndex) final override;
+
+
+	virtual FRenderQueryPoolRHIRef RHICreateRenderQueryPool(ERenderQueryType QueryType, uint32 NumQueries = UINT32_MAX) final override;
+	virtual FStagingBufferRHIRef RHICreateStagingBuffer() final override;
+	virtual FGPUFenceRHIRef RHICreateGPUFence(const FName &Name) final override;
+
+	virtual bool RHIRequiresComputeGenerateMips() const override;
+
+	virtual uint64 RHIGetMinimumAlignmentForBufferBackedSRV(EPixelFormat Format) final override;
+
+	// Compute the hash of the state components of the PSO initializer for PSO Precaching (only hash data relevant for the RHI specific PSO)
+	virtual uint64 RHIComputeStatePrecachePSOHash(const FGraphicsPipelineStateInitializer& Initializer) final override;
+
+	// Compute the hash of the PSO initializer for PSO Precaching (only hash data relevant for the RHI specific PSO)
+	virtual uint64 RHIComputePrecachePSOHash(const FGraphicsPipelineStateInitializer& Initializer) final override;
+
+	// Check if PSO Initializers are the same used during PSO Precaching (only compare data relevant for the RHI specific PSO)
+	virtual bool RHIMatchPrecachePSOInitializers(const FGraphicsPipelineStateInitializer& LHS, const FGraphicsPipelineStateInitializer& RHS) final override;
+
+	void Cleanup();
+
+	void PurgeFramebufferFromCaches(GLuint Framebuffer);
+	void OnBufferDeletion(GLuint VertexBufferResource);
+	void OnPixelBufferDeletion(GLuint PixelBufferResource);
+	void OnUniformBufferDeletion(GLuint UniformBufferResource,uint32 AllocatedSize,bool bStreamDraw);
+	void OnProgramDeletion(GLint ProgramResource);
+	void InvalidateTextureResourceInCache(GLuint Resource);
+	void InvalidateUAVResourceInCache(GLuint Resource);
+	/** Set a resource on texture target of a specific real OpenGL stage. Goes through cache to eliminate redundant calls. */
+	FORCEINLINE void CachedSetupTextureStage(FOpenGLContextState& ContextState, GLint TextureIndex, GLenum Target, GLuint Resource, GLint BaseMip, GLint NumMips)
+	{
+		FTextureStage& TextureState = ContextState.Textures[TextureIndex];
+		const bool bSameTarget = (TextureState.Target == Target);
+		const bool bSameResource = (TextureState.Resource == Resource);
+
+		if (bSameTarget && bSameResource)
+		{
+			// Nothing changed, no need to update
+			return;
+		}
+		CachedSetupTextureStageInner(ContextState, TextureIndex, Target, Resource, BaseMip, NumMips);
+	}
+
+	void CachedSetupTextureStageInner(FOpenGLContextState& ContextState, GLint TextureIndex, GLenum Target, GLuint Resource, GLint BaseMip, GLint NumMips);
+	void CachedSetupUAVStage(FOpenGLContextState& ContextState, GLint UAVIndex, GLenum Format, GLuint Resource, bool bLayered, GLint Layer, GLenum Access);
+	void UpdateSRV(FOpenGLShaderResourceView* SRV);
+	FOpenGLContextState& GetContextStateForCurrentContext(bool bAssertIfInvalid = true);
+
+	FORCEINLINE void CachedBindArrayBuffer( FOpenGLContextState& ContextState, GLuint Buffer )
+	{
+		VERIFY_GL_SCOPE();
+		if( ContextState.ArrayBufferBound != Buffer )
+		{
+			glBindBuffer( GL_ARRAY_BUFFER, Buffer );
+			ContextState.ArrayBufferBound = Buffer;
+		}
+	}
+
+	void CachedBindElementArrayBuffer( FOpenGLContextState& ContextState, GLuint Buffer )
+	{
+		VERIFY_GL_SCOPE();
+		if( ContextState.ElementArrayBufferBound != Buffer )
+		{
+			glBindBuffer( GL_ELEMENT_ARRAY_BUFFER, Buffer );
+			ContextState.ElementArrayBufferBound = Buffer;
+		}
+	}
+	
+	FORCEINLINE void CachedBindStorageBuffer( FOpenGLContextState& ContextState, GLuint Buffer )
+	{
+		VERIFY_GL_SCOPE();
+		if( ContextState.StorageBufferBound != Buffer )
+		{
+			glBindBuffer( GL_SHADER_STORAGE_BUFFER, Buffer );
+			ContextState.StorageBufferBound = Buffer;
+		}
+	}
+
+	void CachedBindPixelUnpackBuffer( FOpenGLContextState& ContextState, GLuint Buffer )
+	{
+		VERIFY_GL_SCOPE();
+
+		if( ContextState.PixelUnpackBufferBound != Buffer )
+		{
+			glBindBuffer( GL_PIXEL_UNPACK_BUFFER, Buffer );
+			ContextState.PixelUnpackBufferBound = Buffer;
+		}
+	}
+
+	void CachedBindUniformBuffer( FOpenGLContextState& ContextState, GLuint Buffer )
+	{
+		VERIFY_GL_SCOPE();
+		check(IsInRenderingThread()||IsInRHIThread());
+		if( ContextState.UniformBufferBound != Buffer )
+		{
+			glBindBuffer( GL_UNIFORM_BUFFER, Buffer );
+			ContextState.UniformBufferBound = Buffer;
+		}
+	}
+
+	bool IsUniformBufferBound( FOpenGLContextState& ContextState, GLuint Buffer ) const
+	{
+		return ( ContextState.UniformBufferBound == Buffer );
+	}
+
+	/** Add query to Queries list upon its creation. */
+	void RegisterQuery( FOpenGLRenderQuery* Query );
+
+	/** Remove query from Queries list upon its deletion. */
+	void UnregisterQuery( FOpenGLRenderQuery* Query );
+
+	/** Inform all queries about the need to recreate themselves after OpenGL context they're in gets deleted. */
+	void InvalidateQueries();
+
+	void BeginRenderQuery_OnThisThread(FOpenGLRenderQuery* Query);
+	void EndRenderQuery_OnThisThread(FOpenGLRenderQuery* Query);
+	void GetRenderQueryResult_OnThisThread(FOpenGLRenderQuery* Query, bool bWait);
+
+	FOpenGLSamplerState* GetPointSamplerState() const { return (FOpenGLSamplerState*)PointSamplerState.GetReference(); }
+
+	void InitializeGLTexture(FOpenGLTexture* Texture, const void* BulkDataPtr, uint64 BulkDataSize);
+
+	void* GetOpenGLCurrentContextHandle();
+
+	void SetCustomPresent(class FRHICustomPresent* InCustomPresent);
+
+#define RHITHREAD_GLTRACE 1
+#if RHITHREAD_GLTRACE 
+	#define RHITHREAD_GLTRACE_BLOCKING QUICK_SCOPE_CYCLE_COUNTER(STAT_OGLRHIThread_Flush);
+//#define RHITHREAD_GLTRACE_BLOCKING FPlatformMisc::LowLevelOutputDebugStringf(TEXT("FLUSHING %s!\n"), ANSI_TO_TCHAR(__FUNCTION__))
+//#define RHITHREAD_GLTRACE_BLOCKING UE_LOG(LogRHI, Warning,TEXT("FLUSHING %s!\n"), ANSI_TO_TCHAR(__FUNCTION__));
+#else
+	#define RHITHREAD_GLTRACE_BLOCKING 
+#endif
+#define RHITHREAD_GLCOMMAND_PROLOGUE() auto GLCommand= [&]() {
+
+#define RHITHREAD_GLCOMMAND_EPILOGUE_RETURN(x) };\
+		if (RHICmdList.Bypass() ||  !IsRunningRHIInSeparateThread() || IsInRHIThread())\
+		{\
+			return GLCommand();\
+		}\
+		else\
+		{\
+			x ReturnValue = (x)0;\
+			ALLOC_COMMAND_CL(RHICmdList, FRHICommandGLCommand)([&ReturnValue, GLCommand = MoveTemp(GLCommand)]() { ReturnValue = GLCommand(); }); \
+			RHITHREAD_GLTRACE_BLOCKING;\
+			RHICmdList.GetAsImmediate().ImmediateFlush(EImmediateFlushType::FlushRHIThread);\
+			return ReturnValue;\
+		}\
+
+#define RHITHREAD_GLCOMMAND_EPILOGUE_GET_RETURN(x) };\
+		x ReturnValue = (x)0;\
+		if (RHICmdList.Bypass() ||  !IsRunningRHIInSeparateThread() || IsInRHIThread() )\
+		{\
+			ReturnValue = GLCommand();\
+		}\
+		else\
+		{\
+			ALLOC_COMMAND_CL(RHICmdList, FRHICommandGLCommand)([&ReturnValue, GLCommand = MoveTemp(GLCommand)]() { ReturnValue = GLCommand(); }); \
+			RHITHREAD_GLTRACE_BLOCKING;\
+			RHICmdList.GetAsImmediate().ImmediateFlush(EImmediateFlushType::FlushRHIThread);\
+		}\
+
+
+#define RHITHREAD_GLCOMMAND_EPILOGUE() };\
+		if (ShouldRunGLRenderContextOpOnThisThread(RHICmdList))\
+		{\
+			return GLCommand();\
+		}\
+		else\
+		{\
+			ALLOC_COMMAND_CL(RHICmdList, FRHICommandGLCommand)( MoveTemp(GLCommand) ); \
+			RHITHREAD_GLTRACE_BLOCKING;\
+			RHICmdList.GetAsImmediate().ImmediateFlush(EImmediateFlushType::FlushRHIThread);\
+		}\
+
+#define RHITHREAD_GLCOMMAND_EPILOGUE_NORETURN() };\
+		if (ShouldRunGLRenderContextOpOnThisThread(RHICmdList))\
+		{\
+			GLCommand();\
+		}\
+		else\
+		{\
+			ALLOC_COMMAND_CL(RHICmdList, FRHICommandGLCommand)(  MoveTemp(GLCommand) ); \
+			RHITHREAD_GLTRACE_BLOCKING;\
+			RHICmdList.GetAsImmediate().ImmediateFlush(EImmediateFlushType::FlushRHIThread);\
+		}\
+
+	struct FTextureLockTracker
+	{
+		struct FLockParams
+		{
+			void* RHIBuffer;
+			void* Buffer;
+			uint32 MipIndex;
+			uint32 ArrayIndex;
+			uint32 BufferSize;
+			uint32 Stride;
+			EResourceLockMode LockMode;
+
+			FORCEINLINE_DEBUGGABLE FLockParams(void* InRHIBuffer, void* InBuffer, uint32 InArrayIndex, uint32 InMipIndex, uint32 InStride, uint32 InBufferSize, EResourceLockMode InLockMode)
+				: RHIBuffer(InRHIBuffer)
+				, Buffer(InBuffer)
+				, MipIndex(InMipIndex)
+				, ArrayIndex(InArrayIndex)
+				, BufferSize(InBufferSize)
+				, Stride(InStride)
+				, LockMode(InLockMode)
+			{
+			}
+		};
+		TArray<FLockParams, TInlineAllocator<16> > OutstandingLocks;
+		uint32 TotalMemoryOutstanding;
+
+		FTextureLockTracker()
+		{
+			TotalMemoryOutstanding = 0;
+		}
+
+		FORCEINLINE_DEBUGGABLE void Lock(void* RHIBuffer, void* Buffer, uint32 ArrayIndex, uint32 MipIndex, uint32 Stride, uint32 SizeRHI, EResourceLockMode LockMode)
+		{
+//#if DO_CHECK
+			for (auto& Parms : OutstandingLocks)
+			{
+				check(Parms.RHIBuffer != RHIBuffer || Parms.MipIndex != MipIndex || Parms.ArrayIndex != ArrayIndex);
+			}
+//#endif
+			OutstandingLocks.Add(FLockParams(RHIBuffer, Buffer, ArrayIndex, MipIndex, Stride, SizeRHI, LockMode));
+			TotalMemoryOutstanding += SizeRHI;
+		}
+
+		FORCEINLINE_DEBUGGABLE FLockParams Unlock(void* RHIBuffer, uint32 ArrayIndex, uint32 MipIndex)
+		{
+			for (int32 Index = 0; Index < OutstandingLocks.Num(); Index++)
+			{
+				FLockParams& CurrentLock = OutstandingLocks[Index];
+				if (CurrentLock.RHIBuffer == RHIBuffer && CurrentLock.MipIndex == MipIndex && CurrentLock.ArrayIndex == ArrayIndex)
+				{
+					FLockParams Result = OutstandingLocks[Index];
+					OutstandingLocks.RemoveAtSwap(Index, 1, EAllowShrinking::No);
+					TotalMemoryOutstanding -= Result.BufferSize;
+					return Result;
+				}
+			}
+			check(!"Mismatched RHI buffer locks.");
+			return FLockParams(nullptr, nullptr, 0, 0, 0, 0, RLM_WriteOnly);
+		}
+	};
+
+	virtual void* LockTexture2D_RenderThread(class FRHICommandListImmediate& RHICmdList, FRHITexture2D* Texture, uint32 MipIndex, EResourceLockMode LockMode, uint32& DestStride, bool bLockWithinMiptail, bool bNeedsDefaultRHIFlush, uint64* OutLockedByteCount = nullptr) final override;
+	virtual void UnlockTexture2D_RenderThread(class FRHICommandListImmediate& RHICmdList, FRHITexture2D* Texture, uint32 MipIndex, bool bLockWithinMiptail, bool bNeedsDefaultRHIFlush) final override;
+	virtual void* RHILockTextureCubeFace_RenderThread(class FRHICommandListImmediate& RHICmdList, FRHITextureCube* Texture, uint32 FaceIndex, uint32 ArrayIndex, uint32 MipIndex, EResourceLockMode LockMode, uint32& DestStride, bool bLockWithinMiptail) final override;
+	virtual void RHIUnlockTextureCubeFace_RenderThread(class FRHICommandListImmediate& RHICmdList, FRHITextureCube* Texture, uint32 FaceIndex, uint32 ArrayIndex, uint32 MipIndex, bool bLockWithinMiptail) final override;
+	virtual void* LockTexture2DArray_RenderThread(class FRHICommandListImmediate& RHICmdList, FRHITexture2DArray* Texture, uint32 ArrayIndex, uint32 MipIndex, EResourceLockMode LockMode, uint32& DestStride, bool bLockWithinMiptail) final override;
+	virtual void UnlockTexture2DArray_RenderThread(class FRHICommandListImmediate& RHICmdList, FRHITexture2DArray* Texture, uint32 ArrayIndex, uint32 MipIndex, bool bLockWithinMiptail) final override;
+
+	virtual FTexture2DRHIRef AsyncReallocateTexture2D_RenderThread(class FRHICommandListImmediate& RHICmdList, FRHITexture2D* Texture2D, int32 NewMipCount, int32 NewSizeX, int32 NewSizeY, FThreadSafeCounter* RequestStatus) final override;
+
+	virtual ETextureReallocationStatus FinalizeAsyncReallocateTexture2D_RenderThread(class FRHICommandListImmediate& RHICmdList, FRHITexture2D* Texture2D, bool bBlockUntilCompleted) final override
+	{
+		return this->RHIFinalizeAsyncReallocateTexture2D(Texture2D, bBlockUntilCompleted);
+	}
+
+	virtual ETextureReallocationStatus CancelAsyncReallocateTexture2D_RenderThread(class FRHICommandListImmediate& RHICmdList, FRHITexture2D* Texture2D, bool bBlockUntilCompleted) final override
+	{
+		return this->RHICancelAsyncReallocateTexture2D(Texture2D, bBlockUntilCompleted);
+	}
+
+	virtual FSamplerStateRHIRef RHICreateSamplerState(const FSamplerStateInitializerRHI& Initializer) final override;
+
+	virtual FGraphicsPipelineStateRHIRef RHICreateGraphicsPipelineState(const FGraphicsPipelineStateInitializer& Initializer) override
+	{
+		PrepareGFXBoundShaderState(Initializer);
+
+		return new FRHIGraphicsPipelineStateFallBack(Initializer);
+	}
+
+	virtual void RHISetGraphicsPipelineState(FRHIGraphicsPipelineState* GraphicsState, uint32 StencilRef, bool bApplyAdditionalState) final override;
+
+	virtual FBoundShaderStateRHIRef RHICreateBoundShaderState(
+		FRHIVertexDeclaration* VertexDeclarationRHI,
+		FRHIVertexShader* VertexShaderRHI,
+		FRHIPixelShader* PixelShaderRHI,
+		FRHIGeometryShader* GeometryShaderRHI
+	) final override
+	{
+		return RHICreateBoundShaderState_internal(
+			VertexDeclarationRHI,
+			VertexShaderRHI,
+			PixelShaderRHI,
+			GeometryShaderRHI,
+			false);
+	}
+
+	bool LinkComputeShader(FRHIComputeShader* ComputeShaderRHI, FOpenGLComputeShader* ComputeShader);
+	class FOpenGLLinkedProgram* GetLinkedComputeProgram(FRHIComputeShader* ComputeShaderRHI);
+
+	FBoundShaderStateRHIRef RHICreateBoundShaderState_OnThisThread(FRHIVertexDeclaration* VertexDeclaration, FRHIVertexShader* VertexShader, FRHIPixelShader* PixelShader, FRHIGeometryShader* GeometryShader, bool FromPSOFileCache);
+	virtual void RHIAdvanceFrameFence() final override;
+
+	virtual void RHIPostExternalCommandsReset() final override;
+
+	FOpenGLGPUProfiler& GetGPUProfilingData() {
+		return GPUProfilingData;
+	}
+
+	GLuint GetOpenGLFramebuffer(uint32 NumSimultaneousRenderTargets, FOpenGLTexture** RenderTargets, const uint32* ArrayIndices, const uint32* MipmapLevels, FOpenGLTexture* DepthStencilTarget);
+	GLuint GetOpenGLFramebuffer(uint32 NumSimultaneousRenderTargets, FOpenGLTexture** RenderTargets, const uint32* ArrayIndices, const uint32* MipmapLevels, FOpenGLTexture* DepthStencilTarget, int32 NumRenderingSamples);
+	
+private:
+
+	FBoundShaderStateRHIRef RHICreateBoundShaderState_internal(
+		FRHIVertexDeclaration* VertexDeclarationRHI,
+		FRHIVertexShader* VertexShaderRHI,
+		FRHIPixelShader* PixelShaderRHI,
+		FRHIGeometryShader* GeometryShaderRHI,
+		bool FromPSOFileCache
+	)
+	{
+		FRHICommandListImmediate& RHICmdList = FRHICommandListExecutor::GetImmediateCommandList();
+
+		RHITHREAD_GLCOMMAND_PROLOGUE()
+			return RHICreateBoundShaderState_OnThisThread(VertexDeclarationRHI,
+				VertexShaderRHI,
+				PixelShaderRHI,
+				GeometryShaderRHI,
+				FromPSOFileCache);
+		RHITHREAD_GLCOMMAND_EPILOGUE_RETURN(FBoundShaderStateRHIRef);
+	}
+
+	void PrepareGFXBoundShaderState(const FGraphicsPipelineStateInitializer& Initializer);
+
+	FOpenGLLinkedProgram* LinkProgram(const class FOpenGLLinkedProgramConfiguration& Config);
+
+	/** called once per frame, used for resource processing */
+	void EndFrameTick();
+
+	/** Counter incremented each time RHIBeginScene is called. */
+	uint32 SceneFrameCounter;
+
+	/** Value used to detect when resource tables need to be recached. INDEX_NONE means always recache. */
+	uint32 ResourceTableFrameCounter;
+
+	/** RHI device state, independent of underlying OpenGL context used */
+	FOpenGLRHIState						PendingState;
+	FSamplerStateRHIRef					PointSamplerState;
+
+	/** A list of all viewport RHIs that have been created. */
+	TArray<FOpenGLViewport*> Viewports;
+	TRefCountPtr<FOpenGLViewport>		DrawingViewport;
+	bool								bRevertToSharedContextAfterDrawingViewport;
+
+	bool								bIsRenderingContextAcquired;
+
+	EPrimitiveType						PrimitiveType = PT_Num;
+
+	/** A history of the most recently used bound shader states, used to keep transient bound shader states from being recreated for each use. */
+	TGlobalResource< TBoundShaderStateHistory<10000> > BoundShaderStateHistory;
+
+	/** Per-context state caching */
+	FOpenGLContextState InvalidContextState;
+	FOpenGLContextState	SharedContextState;
+	FOpenGLContextState	RenderingContextState;
+	// Cached context type on BeginScene
+	int32 BeginSceneContextType;
+
+	template <typename TRHIShader>
+	void ApplyStaticUniformBuffers(TRHIShader* Shader);
+
+	TArray<FRHIUniformBuffer*> GlobalUniformBuffers;
+
+	/** Cached mip-limits for textures when ARB_texture_view is unavailable */
+	TMap<GLuint, TPair<GLenum, GLenum>> TextureMipLimits;
+
+	/** Underlying platform-specific data */
+	struct FPlatformOpenGLDevice* PlatformDevice;
+
+	/** Query list. This is used to inform queries they're no longer valid when OpenGL context they're in gets released from another thread. */
+	TArray<FOpenGLRenderQuery*> Queries;
+
+	/** A critical section to protect modifications and iteration over Queries list */
+	FCriticalSection QueriesListCriticalSection;
+
+	FOpenGLGPUProfiler GPUProfilingData;
+	friend FOpenGLGPUProfiler;
+
+	FCriticalSection CustomPresentSection;
+	TRefCountPtr<class FRHICustomPresent> CustomPresent;
+
+	void InitializeStateResources();
+
+	void SetupVertexArrays(FOpenGLContextState& ContextCache, uint32 BaseVertexIndex, FOpenGLStream* Streams, uint32 NumStreams, uint32 MaxVertices);
+
+	void SetupBindlessTextures( FOpenGLContextState& ContextState, const TArray<FOpenGLBindlessSamplerInfo> &Samplers );
+
+	/** needs to be called before each draw call */
+	void BindPendingFramebuffer( FOpenGLContextState& ContextState );
+	void BindPendingShaderState( FOpenGLContextState& ContextState );
+	void BindPendingComputeShaderState( FOpenGLContextState& ContextState, FOpenGLComputeShader* ComputeShader );
+	void UpdateRasterizerStateInOpenGLContext( FOpenGLContextState& ContextState );
+	void UpdateDepthStencilStateInOpenGLContext( FOpenGLContextState& ContextState );
+	void UpdateScissorRectInOpenGLContext( FOpenGLContextState& ContextState );
+	void UpdateViewportInOpenGLContext( FOpenGLContextState& ContextState );
+	
+	template <class ShaderType> void SetResourcesFromTables(ShaderType* Shader);
+	FORCEINLINE void CommitGraphicsResourceTables()
+	{
+		if (PendingState.bAnyDirtyGraphicsUniformBuffers)
+		{
+			CommitGraphicsResourceTablesInner();
+		}
+	}
+	void CommitGraphicsResourceTablesInner();
+	void CommitComputeResourceTables(FOpenGLComputeShader* ComputeShader);
+	void CommitNonComputeShaderConstants();
+	void CommitComputeShaderConstants(FOpenGLComputeShader* ComputeShader);
+	void SetPendingBlendStateForActiveRenderTargets( FOpenGLContextState& ContextState );
+	
+	void SetupTexturesForDraw( FOpenGLContextState& ContextState);
+	template <typename StateType>
+	void SetupTexturesForDraw( FOpenGLContextState& ContextState, const StateType& ShaderState, int32 MaxTexturesNeeded);
+
+	void SetupUAVsForDraw(FOpenGLContextState& ContextState);
+	void SetupUAVsForCompute(FOpenGLContextState& ContextState, const FOpenGLComputeShader* ComputeShader);
+	void SetupUAVsForProgram(FOpenGLContextState& ContextState, const TBitArray<>& NeededBits, int32 MaxUAVUnitUsed);
+
+	void RHIClearMRT(const bool* bClearColorArray, int32 NumClearColors, const FLinearColor* ColorArray, bool bClearDepth, float Depth, bool bClearStencil, uint32 Stencil);
+
+public:
+	/** Remember what RHI user wants set on a specific OpenGL texture stage, translating from Stage and TextureIndex for stage pair. */
+	void InternalSetShaderTexture(FOpenGLTexture* Texture, FOpenGLShaderResourceView* SRV, GLint TextureIndex, GLenum Target, GLuint Resource, int NumMips, int LimitMip);
+	void InternalSetShaderImageUAV(GLint UAVIndex, GLenum Format, GLuint Resource, bool bLayered, GLint Layer, GLenum Access);
+	void InternalSetShaderBufferUAV(GLint UAVIndex, GLuint Resource);
+	void InternalSetSamplerStates(GLint TextureIndex, FOpenGLSamplerState* SamplerState);
+	void InitializeGLTextureInternal(FOpenGLTexture* Texture, void const* BulkDataPtr, uint64 BulkDataSize);
+private:
+
+	void ApplyTextureStage(FOpenGLContextState& ContextState, GLint TextureIndex, const FTextureStage& TextureStage, FOpenGLSamplerState* SamplerState);
+
+	void ReadSurfaceDataRaw(FOpenGLContextState& ContextState, FRHITexture* TextureRHI,FIntRect Rect,TArray<uint8>& OutData, FReadSurfaceDataFlags InFlags);
+
+	void BindUniformBufferBase(FOpenGLContextState& ContextState, int32 NumUniformBuffers, FRHIUniformBuffer** BoundUniformBuffers, uint32* DynamicOffsets, uint32 FirstUniformBuffer, bool ForceUpdate);
+
+	void ClearCurrentFramebufferWithCurrentScissor(FOpenGLContextState& ContextState, int8 ClearType, int32 NumClearColors, const bool* bClearColorArray, const FLinearColor* ClearColorArray, float Depth, uint32 Stencil);
+
+	FTextureLockTracker GLLockTracker;
+
+};
+
+PRAGMA_ENABLE_DEPRECATION_WARNINGS
+
+/** Implements the OpenGLDrv module as a dynamic RHI providing module. */
+class FOpenGLDynamicRHIModule : public IDynamicRHIModule
+{
+public:
+	
+	// IModuleInterface
+	virtual bool SupportsDynamicReloading() override { return false; }
+
+	// IDynamicRHIModule
+	virtual bool IsSupported() override;
+
+	virtual FDynamicRHI* CreateRHI(ERHIFeatureLevel::Type RequestedFeatureLevel = ERHIFeatureLevel::Num) override;
+};
+
+extern ERHIFeatureLevel::Type GRequestedFeatureLevel;
diff --git a/Engine/Source/Runtime/OpenGLDrv/Private/OpenGLES.h b/Engine/Source/Runtime/OpenGLDrv/Private/OpenGLES.h
new file mode 100644
--- /dev/null
+++ b/Engine/Source/Runtime/OpenGLDrv/Private/OpenGLES.h
@@ -0,0 +1,904 @@
+// Copyright Epic Games, Inc. All Rights Reserved.
+
+/*=============================================================================
+	OpenGLES2.h: Public OpenGL ES 2.0 definitions for non-common functionality
+=============================================================================*/
+
+#pragma once
+
+#include "HAL/Platform.h"
+
+#if !PLATFORM_DESKTOP // need this to fix compile issues with Win configuration.
+
+#define OPENGL_ES	1
+
+typedef GLfloat GLdouble;
+
+#include "OpenGL.h"
+#include "OpenGLUtil.h"		// for VERIFY_GL
+
+#ifdef GL_AMD_debug_output
+	#undef GL_AMD_debug_output
+#endif
+
+
+
+/** Unreal tokens that maps to different OpenGL tokens by platform. */
+#undef UGL_ABGR8
+#define UGL_ABGR8				GL_UNSIGNED_BYTE
+#undef UGL_ANY_SAMPLES_PASSED
+#define UGL_ANY_SAMPLES_PASSED	GL_ANY_SAMPLES_PASSED
+#undef UGL_CLAMP_TO_BORDER
+#define UGL_CLAMP_TO_BORDER		GL_CLAMP_TO_EDGE
+#undef UGL_TIME_ELAPSED
+#define UGL_TIME_ELAPSED		GL_TIME_ELAPSED_EXT
+
+/** Official OpenGL definitions */
+#ifndef GL_FILL
+#define GL_FILL 0x1B02
+#endif
+#ifndef GL_GEOMETRY_SHADER
+#define GL_GEOMETRY_SHADER 0x8DD9
+#endif
+#ifndef GL_SAMPLER_1D
+#define GL_SAMPLER_1D 0x8B5D
+#endif
+#ifndef GL_SAMPLER_1D_SHADOW
+#define GL_SAMPLER_1D_SHADOW 0x8B61
+#endif
+#ifndef GL_DOUBLE
+#define GL_DOUBLE 0x140A
+#endif
+#ifndef GL_BGRA
+#define GL_BGRA	GL_BGRA_EXT 
+#endif
+#ifndef GL_TEXTURE_BUFFER
+#define GL_TEXTURE_BUFFER 0x8C2A
+#endif
+#ifndef GL_RGBA16
+/* GL_EXT_texture_norm16 */
+#define GL_RGBA16 GL_RGBA16_EXT
+#endif
+#ifndef GL_R16
+/* GL_EXT_texture_norm16 */
+#define GL_R16 GL_R16_EXT
+#endif
+#ifndef GL_TEXTURE_BUFFER_OFFSET_ALIGNMENT
+#define GL_TEXTURE_BUFFER_OFFSET_ALIGNMENT 0x919F
+#endif
+#ifndef GL_POLYGON_OFFSET_LINE
+#define GL_POLYGON_OFFSET_LINE 0x2A02
+#endif
+#ifndef GL_POLYGON_OFFSET_POINT
+#define GL_POLYGON_OFFSET_POINT 0x2A01
+#endif
+#ifndef GL_TEXTURE_LOD_BIAS
+#define GL_TEXTURE_LOD_BIAS 0x8501
+#endif
+#ifndef GL_FRAMEBUFFER_SRGB
+#define GL_FRAMEBUFFER_SRGB 0x8DB9
+#endif
+#ifndef GL_SAMPLES_PASSED
+#define GL_SAMPLES_PASSED 0x8914
+#endif
+#ifndef GL_POINT
+#define GL_POINT 0x1B00
+#endif
+#ifndef GL_LINE
+#define GL_LINE 0x1B01
+#endif
+#ifndef GL_TEXTURE_1D
+#define GL_TEXTURE_1D			0x0DE0
+#endif
+#ifndef GL_TEXTURE_1D_ARRAY
+#define GL_TEXTURE_1D_ARRAY		0x8C18
+#endif
+#ifndef GL_TEXTURE_RECTANGLE
+#define GL_TEXTURE_RECTANGLE	0x84F5
+#endif
+#ifndef GL_MAX_TEXTURE_BUFFER_SIZE
+#define GL_MAX_TEXTURE_BUFFER_SIZE 0x8C2B
+#endif
+#ifndef GL_DEPTH_CLAMP
+#define GL_DEPTH_CLAMP 0x864F
+#endif
+
+/** For the shader stage bits that don't exist just use 0 */
+#define GL_GEOMETRY_SHADER_BIT				0x00000000
+#define GL_TESS_CONTROL_SHADER_BIT			0x00000000
+#define GL_TESS_EVALUATION_SHADER_BIT		0x00000000
+
+// Normalize debug macros due to naming differences across GL versions
+#if defined(GL_KHR_debug) && GL_KHR_debug
+#define GL_DEBUG_SOURCE_OTHER_ARB GL_DEBUG_SOURCE_OTHER_KHR
+#define GL_DEBUG_SOURCE_API_ARB GL_DEBUG_SOURCE_API_KHR
+#define GL_DEBUG_TYPE_ERROR_ARB GL_DEBUG_TYPE_ERROR_KHR
+#define GL_DEBUG_TYPE_OTHER_ARB GL_DEBUG_TYPE_OTHER_KHR
+#define GL_DEBUG_TYPE_MARKER GL_DEBUG_TYPE_MARKER_KHR
+#define GL_DEBUG_TYPE_PUSH_GROUP GL_DEBUG_TYPE_PUSH_GROUP_KHR
+#define GL_DEBUG_TYPE_POP_GROUP GL_DEBUG_TYPE_POP_GROUP_KHR
+#define GL_DEBUG_SEVERITY_HIGH_ARB GL_DEBUG_SEVERITY_HIGH_KHR
+#define GL_DEBUG_SEVERITY_LOW_ARB GL_DEBUG_SEVERITY_LOW_KHR
+#define GL_DEBUG_SEVERITY_NOTIFICATION GL_DEBUG_SEVERITY_NOTIFICATION_KHR
+#endif
+
+// FIXME: include gl32.h
+typedef void (GL_APIENTRYP PFNGLFRAMEBUFFERTEXTUREPROC) (GLenum target, GLenum attachment, GLuint texture, GLint level);
+
+namespace GLFuncPointers
+{
+	// GL_EXT_multisampled_render_to_texture
+	extern PFNGLFRAMEBUFFERTEXTURE2DMULTISAMPLEEXTPROC	glFramebufferTexture2DMultisampleEXT;
+	extern PFNGLRENDERBUFFERSTORAGEMULTISAMPLEEXTPROC	glRenderbufferStorageMultisampleEXT;
+
+	// GL_EXT_debug_marker
+	extern PFNGLPUSHGROUPMARKEREXTPROC		glPushGroupMarkerEXT;
+	extern PFNGLPOPGROUPMARKEREXTPROC 		glPopGroupMarkerEXT;
+
+	// GL_EXT_debug_label
+	extern PFNGLLABELOBJECTEXTPROC			glLabelObjectEXT;
+	extern PFNGLGETOBJECTLABELEXTPROC		glGetObjectLabelEXT;
+	//GL_EXT_buffer_storage
+	extern PFNGLBUFFERSTORAGEEXTPROC		glBufferStorageEXT;
+
+	extern PFNGLDEBUGMESSAGECONTROLKHRPROC	glDebugMessageControlKHR;
+	extern PFNGLDEBUGMESSAGEINSERTKHRPROC	glDebugMessageInsertKHR;
+	extern PFNGLDEBUGMESSAGECALLBACKKHRPROC	glDebugMessageCallbackKHR;
+	extern PFNGLGETDEBUGMESSAGELOGKHRPROC	glDebugMessageLogKHR;
+	extern PFNGLGETPOINTERVKHRPROC			glGetPointervKHR;
+	extern PFNGLPUSHDEBUGGROUPKHRPROC		glPushDebugGroupKHR;
+	extern PFNGLPOPDEBUGGROUPKHRPROC		glPopDebugGroupKHR;
+	extern PFNGLOBJECTLABELKHRPROC			glObjectLabelKHR;
+	extern PFNGLGETOBJECTLABELKHRPROC		glGetObjectLabelKHR;
+	extern PFNGLOBJECTPTRLABELKHRPROC		glObjectPtrLabelKHR;
+	extern PFNGLGETOBJECTPTRLABELKHRPROC	glGetObjectPtrLabelKHR;
+
+	// ES 3.2
+	extern PFNGLTEXBUFFEREXTPROC				glTexBufferEXT;
+	extern PFNGLTEXBUFFERRANGEEXTPROC			glTexBufferRangeEXT;
+	extern PFNGLCOPYIMAGESUBDATAEXTPROC			glCopyImageSubData;
+	extern PFNGLENABLEIEXTPROC					glEnableiEXT;
+	extern PFNGLDISABLEIEXTPROC					glDisableiEXT;
+	extern PFNGLBLENDEQUATIONIEXTPROC			glBlendEquationiEXT;
+	extern PFNGLBLENDEQUATIONSEPARATEIEXTPROC	glBlendEquationSeparateiEXT;
+	extern PFNGLBLENDFUNCIEXTPROC				glBlendFunciEXT;
+	extern PFNGLBLENDFUNCSEPARATEIEXTPROC		glBlendFuncSeparateiEXT;
+	extern PFNGLCOLORMASKIEXTPROC				glColorMaskiEXT;
+	extern PFNGLFRAMEBUFFERTEXTUREPROC			glFramebufferTexture;
+
+	// Mobile multi-view
+	extern PFNGLFRAMEBUFFERTEXTUREMULTIVIEWOVRPROC glFramebufferTextureMultiviewOVR;
+	extern PFNGLFRAMEBUFFERTEXTUREMULTISAMPLEMULTIVIEWOVRPROC glFramebufferTextureMultisampleMultiviewOVR;
+};
+
+using namespace GLFuncPointers;
+
+struct FOpenGLES : public FOpenGLBase
+{
+	static FORCEINLINE bool IsES31Usable()
+	{
+		check(CurrentFeatureLevelSupport != EFeatureLevelSupport::Invalid);
+		return CurrentFeatureLevelSupport >= EFeatureLevelSupport::ES31;
+	}
+
+	static FORCEINLINE bool IsES32Usable()
+	{
+		check(CurrentFeatureLevelSupport != EFeatureLevelSupport::Invalid);
+		return CurrentFeatureLevelSupport == EFeatureLevelSupport::ES32;
+	}
+
+	static void		ProcessQueryGLInt();
+	static void		ProcessExtensions(const FString& ExtensionsString);
+
+	static FORCEINLINE bool SupportsUniformBuffers() { return true; }
+	static FORCEINLINE bool SupportsStructuredBuffers() { return true; }
+	static FORCEINLINE bool SupportsExactOcclusionQueries() { return false; }
+	// MLCHANGES BEGIN -- changed to use bSupportsDisjointTimeQueries
+	static FORCEINLINE bool SupportsTimestampQueries() { return bSupportsDisjointTimeQueries; }
+	// MLCHANGES END
+	static bool SupportsDisjointTimeQueries();
+	static FORCEINLINE bool SupportsDepthStencilRead() { return false; }
+	static FORCEINLINE bool SupportsFloatReadSurface() { return SupportsColorBufferHalfFloat(); }
+	static FORCEINLINE bool SupportsWideMRT() { return true; }
+	static FORCEINLINE bool SupportsPolygonMode() { return false; }
+	static FORCEINLINE bool SupportsTexture3D() { return true; }
+	static FORCEINLINE bool SupportsMobileMultiView() { return bSupportsMobileMultiView; }
+	static FORCEINLINE bool SupportsImageExternal() { return false; }
+	static FORCEINLINE bool SupportsTextureLODBias() { return false; }
+	static FORCEINLINE bool SupportsTextureCompare() { return false; }
+	static FORCEINLINE bool SupportsDrawIndexOffset() { return false; }
+	static FORCEINLINE bool SupportsDiscardFrameBuffer() { return true; }
+	static FORCEINLINE bool SupportsIndexedExtensions() { return false; }
+	static FORCEINLINE bool SupportsColorBufferFloat() { return bSupportsColorBufferFloat; }
+	static FORCEINLINE bool SupportsColorBufferHalfFloat() { return bSupportsColorBufferHalfFloat; }
+	static FORCEINLINE bool SupportsShaderFramebufferFetch() { return bSupportsShaderFramebufferFetch; }
+	static FORCEINLINE bool SupportsShaderMRTFramebufferFetch() { return bSupportsShaderMRTFramebufferFetch; }
+	static FORCEINLINE bool SupportsShaderDepthStencilFetch() { return bSupportsShaderDepthStencilFetch; }
+	static FORCEINLINE bool SupportsPixelLocalStorage() { return bSupportsPixelLocalStorage; }
+	static FORCEINLINE bool SupportsMultisampledRenderToTexture() { return bSupportsMultisampledRenderToTexture; }
+	static FORCEINLINE bool SupportsVertexArrayBGRA() { return false; }
+	static FORCEINLINE bool SupportsBGRA8888() { return bSupportsBGRA8888; }
+	static FORCEINLINE bool SupportsDXT() { return bSupportsDXT; }
+	static FORCEINLINE bool SupportsETC2() { return bSupportsETC2; }
+	static FORCEINLINE GLenum GetDepthFormat() { return GL_DEPTH_COMPONENT24; }
+	static FORCEINLINE GLenum GetShadowDepthFormat() { return GL_DEPTH_COMPONENT16; }
+	static FORCEINLINE bool SupportsFramebufferSRGBEnable() { return false; }
+	static FORCEINLINE bool SupportsRGB10A2() { return bSupportsRGB10A2; }
+	static FORCEINLINE bool SupportsDrawIndirect() { return true; }
+	static FORCEINLINE bool SupportsBufferStorage() { return bSupportsBufferStorage; }
+	static FORCEINLINE bool SupportsDepthClamp() { return bSupportsDepthClamp; }
+	
+
+	static FORCEINLINE bool HasBinaryProgramRetrievalFailed() { return bBinaryProgramRetrievalFailed; }
+	static FORCEINLINE bool RequiresDisabledEarlyFragmentTests() { return bRequiresDisabledEarlyFragmentTests; }
+	static FORCEINLINE bool RequiresReadOnlyBuffersWorkaround() { return bRequiresReadOnlyBuffersWorkaround; }
+	static FORCEINLINE bool RequiresARMShaderFramebufferFetchDepthStencilUndef() { return bRequiresARMShaderFramebufferFetchDepthStencilUndef; }
+
+	// Adreno doesn't support HALF_FLOAT
+	static FORCEINLINE int32 GetReadHalfFloatPixelsEnum() { return GL_FLOAT; }
+	static FORCEINLINE GLint GetMaxMSAASamplesTileMem() { return MaxMSAASamplesTileMem; }
+
+	// On iOS both glMapBufferOES() and glBufferSubData() for immediate vertex and index data
+	// is the slow path (they both hit GPU sync and data cache flush in driver according to profiling in driver symbols).
+	// Turning this to false reverts back to not using vertex and index buffers
+	// for glDrawArrays() and glDrawElements() on dynamic data.
+	static FORCEINLINE bool SupportsFastBufferData() { return false; }
+
+	static FORCEINLINE bool SupportsASTCDecodeMode() { return bSupportsASTCDecodeMode; }
+
+	// Optional
+	static FORCEINLINE void BeginQuery(GLenum QueryType, GLuint QueryId)
+	{
+		glBeginQuery(QueryType, QueryId);
+	}
+
+	static FORCEINLINE void EndQuery(GLenum QueryType)
+	{
+		glEndQuery(QueryType);
+	}
+
+	static FORCEINLINE void GenQueries(GLsizei NumQueries, GLuint* QueryIDs)
+	{
+		glGenQueries(NumQueries, QueryIDs);
+	}
+
+	static FORCEINLINE void DeleteQueries(GLsizei NumQueries, const GLuint* QueryIDs)
+	{
+		glDeleteQueries(NumQueries, QueryIDs);
+	}
+
+	static FORCEINLINE void GetQueryObject(GLuint QueryId, EQueryMode QueryMode, GLuint* OutResult)
+	{
+		GLenum QueryName = (QueryMode == QM_Result) ? GL_QUERY_RESULT : GL_QUERY_RESULT_AVAILABLE;
+		glGetQueryObjectuiv(QueryId, QueryName, OutResult);
+	}
+		
+	static FORCEINLINE void LabelObject(GLenum Type, GLuint Object, const ANSICHAR* Name)
+	{
+		if (glLabelObjectEXT != nullptr)
+		{
+			glLabelObjectEXT(Type, Object, 0, Name);
+		}
+	}
+
+	static FORCEINLINE GLsizei GetLabelObject(GLenum Type, GLuint Object, GLsizei BufferSize, ANSICHAR* OutName)
+	{
+		GLsizei Length = 0;
+		if (glGetObjectLabelEXT != nullptr)
+		{
+			glGetObjectLabelEXT(Type, Object, BufferSize, &Length, OutName);
+		}
+		return Length;
+	}
+
+	static FORCEINLINE void PushGroupMarker(const ANSICHAR* Name)
+	{
+		if (glPushGroupMarkerEXT != nullptr)
+		{
+			glPushGroupMarkerEXT(0, Name);
+		}
+	}
+
+	static FORCEINLINE void PopGroupMarker()
+	{
+		if (glPopGroupMarkerEXT != nullptr)
+		{
+			glPopGroupMarkerEXT();
+		}
+	}
+
+	static FORCEINLINE void* MapBufferRange(GLenum Type, uint32 InOffset, uint32 InSize, EResourceLockMode LockMode)
+	{
+		GLenum Access;
+		switch (LockMode)
+		{
+		case EResourceLockMode::RLM_ReadOnly:
+			Access = GL_MAP_READ_BIT;
+			break;
+		case EResourceLockMode::RLM_ReadOnlyPersistent:
+			Access = (GL_MAP_READ_BIT | GL_MAP_PERSISTENT_BIT | GL_MAP_COHERENT_BIT);
+			break;
+		case EResourceLockMode::RLM_WriteOnly:
+			Access = (GL_MAP_INVALIDATE_BUFFER_BIT | GL_MAP_WRITE_BIT | GL_MAP_UNSYNCHRONIZED_BIT);
+			break;
+		case EResourceLockMode::RLM_WriteOnlyUnsynchronized:
+			Access = (GL_MAP_WRITE_BIT | GL_MAP_UNSYNCHRONIZED_BIT);
+			break;
+		case EResourceLockMode::RLM_WriteOnlyPersistent:
+			Access = (GL_MAP_WRITE_BIT | GL_MAP_PERSISTENT_BIT | GL_MAP_COHERENT_BIT);
+			break;
+		case EResourceLockMode::RLM_ReadWrite:
+		default:
+			Access = (GL_MAP_READ_BIT | GL_MAP_WRITE_BIT);
+	}
+		return glMapBufferRange(Type, InOffset, InSize, Access);
+	}
+
+	static FORCEINLINE void UnmapBuffer(GLenum Type)
+	{
+		glUnmapBuffer(Type);
+	}
+
+	static FORCEINLINE void UnmapBufferRange(GLenum Type, uint32 InOffset, uint32 InSize)
+	{
+		UnmapBuffer(Type);
+	}
+
+	static FORCEINLINE void TexImage3D(GLenum Target, GLint Level, GLint InternalFormat, GLsizei Width, GLsizei Height, GLsizei Depth, GLint Border, GLenum Format, GLenum Type, const GLvoid* PixelData)
+	{
+		glTexImage3D(Target, Level, InternalFormat, Width, Height, Depth, Border, Format, Type, PixelData);
+	}
+
+	static FORCEINLINE void CompressedTexImage3D(GLenum Target, GLint Level, GLenum InternalFormat, GLsizei Width, GLsizei Height, GLsizei Depth, GLint Border, GLsizei ImageSize, const GLvoid* PixelData)
+	{
+		glCompressedTexImage3D(Target, Level, InternalFormat, Width, Height, Depth, Border, ImageSize, PixelData);
+	}
+
+	static FORCEINLINE void CompressedTexSubImage3D(GLenum Target, GLint Level, GLint XOffset, GLint YOffset, GLint ZOffset, GLsizei Width, GLsizei Height, GLsizei Depth, GLenum Format, GLsizei ImageSize, const GLvoid* PixelData)
+	{
+		glCompressedTexSubImage3D( Target, Level, XOffset, YOffset, ZOffset, Width, Height, Depth, Format, ImageSize, PixelData);
+	}
+
+	static FORCEINLINE void TexSubImage3D(GLenum Target, GLint Level, GLint XOffset, GLint YOffset, GLint ZOffset, GLsizei Width, GLsizei Height, GLsizei Depth, GLenum Format, GLenum Type, const GLvoid* PixelData)
+	{
+		glTexSubImage3D(Target, Level, XOffset, YOffset, ZOffset, Width, Height, Depth, Format, Type, PixelData);
+	}
+
+	static FORCEINLINE void	CopyTexSubImage1D(GLenum Target, GLint Level, GLint XOffset, GLint X, GLint Y, GLsizei Width)
+	{
+	}
+
+	static FORCEINLINE void	CopyTexSubImage2D(GLenum Target, GLint Level, GLint XOffset, GLint YOffset, GLint X, GLint Y, GLsizei Width, GLsizei Height)
+	{
+		glCopyTexSubImage2D(Target, Level, XOffset, YOffset, X, Y, Width, Height);
+	}
+
+	static FORCEINLINE void	CopyTexSubImage3D(GLenum Target, GLint Level, GLint XOffset, GLint YOffset, GLint ZOffset, GLint X, GLint Y, GLsizei Width, GLsizei Height)
+	{
+		glCopyTexSubImage3D(Target, Level, XOffset, YOffset, ZOffset, X, Y, Width, Height);
+	}
+
+	static FORCEINLINE void CopyImageSubData(GLuint SrcName, GLenum SrcTarget, GLint SrcLevel, GLint SrcX, GLint SrcY, GLint SrcZ, GLuint DstName, GLenum DstTarget, GLint DstLevel, GLint DstX, GLint DstY, GLint DstZ, GLsizei Width, GLsizei Height, GLsizei Depth)
+	{
+		glCopyImageSubData(SrcName, SrcTarget, SrcLevel, SrcX, SrcY, SrcZ, DstName, DstTarget, DstLevel, DstX, DstY, DstZ, Width, Height, Depth);
+	}
+
+	static FORCEINLINE bool TexStorage2DMultisample(GLenum Target, GLsizei Samples, GLint InternalFormat, GLsizei Width, GLsizei Height, GLboolean FixedSampleLocations)
+	{
+		glTexStorage2DMultisample(Target, Samples, InternalFormat, Width, Height, FixedSampleLocations);
+		return true;
+	}
+
+	static FORCEINLINE void RenderbufferStorageMultisample(GLenum Target, GLsizei Samples, GLint InternalFormat, GLsizei Width, GLsizei Height)
+	{
+		check(glRenderbufferStorageMultisampleEXT);
+		glRenderbufferStorageMultisampleEXT(Target, Samples, InternalFormat, Width, Height);
+	}
+
+	static FORCEINLINE void ClearBufferfv(GLenum Buffer, GLint DrawBufferIndex, const GLfloat* Value)
+	{
+		glClearBufferfv(Buffer, DrawBufferIndex, Value);
+	}
+
+	static FORCEINLINE void ClearBufferfi(GLenum Buffer, GLint DrawBufferIndex, GLfloat Depth, GLint Stencil)
+	{
+		glClearBufferfi(Buffer, DrawBufferIndex, Depth, Stencil);
+	}
+
+	static FORCEINLINE void ClearBufferiv(GLenum Buffer, GLint DrawBufferIndex, const GLint* Value)
+	{
+		glClearBufferiv(Buffer, DrawBufferIndex, Value);
+	}
+
+	static FORCEINLINE void DrawBuffers(GLsizei NumBuffers, const GLenum* Buffers)
+	{
+		glDrawBuffers(NumBuffers, Buffers);
+	}
+
+	static FORCEINLINE void ReadBuffer(GLenum Mode)
+	{
+		glReadBuffer(Mode);
+	}
+
+	static FORCEINLINE void DrawBuffer(GLenum Mode)
+	{
+		DrawBuffers(1, &Mode);
+	}
+
+	static FORCEINLINE void EnableIndexed(GLenum Parameter, GLuint Index)
+	{
+		if (bSupportsDrawBuffersBlend)
+		{
+			// ES 3.2 or extension
+			glEnableiEXT(Parameter, Index);
+		}
+		else
+		{
+			check(Index == 0);
+			glEnable(Parameter);
+		}
+	}
+	static FORCEINLINE void DisableIndexed(GLenum Parameter, GLuint Index)
+	{
+		if (bSupportsDrawBuffersBlend)
+		{
+			// ES 3.2 or extension
+			glDisableiEXT(Parameter, Index);
+		}
+		else
+		{
+			check(Index == 0);
+			glDisable(Parameter);
+		}
+	}
+	static FORCEINLINE void ColorMaskIndexed(GLuint Index, GLboolean Red, GLboolean Green, GLboolean Blue, GLboolean Alpha)
+	{
+		if (bSupportsDrawBuffersBlend)
+		{
+			// ES 3.2 or extension
+			glColorMaskiEXT(Index, Red, Green, Blue, Alpha);
+		}
+		else
+		{
+			check(Index == 0);
+			glColorMask(Red, Green, Blue, Alpha);
+		}
+	}
+	static FORCEINLINE void BlendEquationi(GLuint Buf, GLenum Mode)
+	{
+		// ES 3.2 or extension
+		glBlendEquationiEXT(Buf, Mode);
+	}
+	static FORCEINLINE void BlendEquationSeparatei(GLuint Buf, GLenum ModeRGB, GLenum ModeAlpha)
+	{
+		// ES 3.2 or extension
+		glBlendEquationSeparateiEXT(Buf, ModeRGB, ModeAlpha);
+	}
+	static FORCEINLINE void BlendFunci(GLuint Buf, GLenum Src, GLenum Dst)
+	{
+		// ES 3.2 or extension
+		glBlendFunciEXT(Buf, Src, Dst);
+	}
+	static FORCEINLINE void BlendFuncSeparatei(GLuint Buf, GLenum SrcRGB, GLenum DstRGB, GLenum SrcAlpha, GLenum DstAlpha)
+	{
+		// ES 3.2 or extension
+		glBlendFuncSeparateiEXT(Buf, SrcRGB, DstRGB, SrcAlpha, DstAlpha);
+	}
+
+	static FORCEINLINE void TexBuffer(GLenum Target, GLenum InternalFormat, GLuint Buffer)
+	{
+		// ES 3.2 or extension
+		glTexBufferEXT(Target, InternalFormat, Buffer);
+	}
+
+	static FORCEINLINE void TexBufferRange(GLenum Target, GLenum InternalFormat, GLuint Buffer, GLintptr Offset, GLsizeiptr Size)
+	{
+		// ES 3.2 or extension
+		glTexBufferRangeEXT(Target, InternalFormat, Buffer, Offset, Size);
+	}
+
+	static FORCEINLINE void ProgramUniform4uiv(GLuint Program, GLint Location, GLsizei Count, const GLuint* Value)
+	{
+		glUniform4uiv(Location, Count, Value);
+	}
+
+	static FORCEINLINE bool SupportsProgramBinary() { return bSupportsProgramBinary; }
+
+	static FORCEINLINE void GetProgramBinary(GLuint Program, GLsizei BufSize, GLsizei* Length, GLenum* BinaryFormat, void* Binary)
+	{
+		glGetProgramBinary(Program, BufSize, Length, BinaryFormat, Binary);
+	}
+
+	static FORCEINLINE void ProgramBinary(GLuint Program, GLenum BinaryFormat, const void* Binary, GLsizei Length)
+	{
+		glProgramBinary(Program, BinaryFormat, Binary, Length);
+	}
+
+	static FORCEINLINE void ProgramParameter(GLuint Program, GLenum PName, GLint Value)
+	{
+		glProgramParameteri(Program, PName, Value);
+	}
+
+	static FORCEINLINE void BindBufferBase(GLenum Target, GLuint Index, GLuint Buffer)
+	{
+		glBindBufferBase(Target, Index, Buffer);
+	}
+
+	static FORCEINLINE void BindBufferRange(GLenum Target, GLuint Index, GLuint Buffer, GLintptr Offset, GLsizeiptr Size)
+	{
+		glBindBufferRange(Target, Index, Buffer, Offset, Size);
+	}
+
+	static FORCEINLINE GLuint GetUniformBlockIndex(GLuint Program, const GLchar* UniformBlockName)
+	{
+		return glGetUniformBlockIndex(Program, UniformBlockName);
+	}
+
+	static FORCEINLINE void UniformBlockBinding(GLuint Program, GLuint UniformBlockIndex, GLuint UniformBlockBinding)
+	{
+		glUniformBlockBinding(Program, UniformBlockIndex, UniformBlockBinding);
+	}
+
+	static FORCEINLINE void BufferSubData(GLenum Target, GLintptr Offset, GLsizeiptr Size, const GLvoid* Data)
+	{
+		glBufferSubData(Target, Offset, Size, Data);
+	}
+
+	static FORCEINLINE void VertexAttribIPointer(GLuint Index, GLint Size, GLenum Type, GLsizei Stride, const GLvoid* Pointer)
+	{
+		glVertexAttribIPointer(Index, Size, Type, Stride, Pointer);
+	}
+
+	static FORCEINLINE void GenSamplers(GLsizei Count, GLuint* Samplers)
+	{
+		glGenSamplers(Count, Samplers);
+	}
+
+	static FORCEINLINE void DeleteSamplers(GLsizei Count, GLuint* Samplers)
+	{
+		glDeleteSamplers(Count, Samplers);
+	}
+
+	static FORCEINLINE void SetSamplerParameter(GLuint Sampler, GLenum Parameter, GLint Value)
+	{
+		glSamplerParameteri(Sampler, Parameter, Value);
+	}
+
+	static FORCEINLINE void BindSampler(GLuint Unit, GLuint Sampler)
+	{
+		glBindSampler(Unit, Sampler);
+	}
+
+	static FORCEINLINE void MemoryBarrier(GLbitfield Barriers)
+	{
+		glMemoryBarrier(Barriers);
+	}
+
+	static FORCEINLINE void DispatchCompute(GLuint NumGroupsX, GLuint NumGroupsY, GLuint NumGroupsZ)
+	{
+		glDispatchCompute(NumGroupsX, NumGroupsY, NumGroupsZ);
+	}
+
+	static FORCEINLINE void DispatchComputeIndirect(GLintptr Offset)
+	{
+		glDispatchComputeIndirect(Offset);
+	}
+
+	static FORCEINLINE void BindImageTexture(GLuint Unit, GLuint Texture, GLint Level, GLboolean Layered, GLint Layer, GLenum Access, GLenum Format)
+	{
+		glBindImageTexture(Unit, Texture, Level, Layered, Layer, Access, Format);
+	}
+	
+	static FORCEINLINE void DepthRange(GLdouble Near, GLdouble Far)
+	{
+		glDepthRangef(Near, Far);
+	}
+
+	static FORCEINLINE void VertexAttribPointer(GLuint Index, GLint Size, GLenum Type, GLboolean Normalized, GLsizei Stride, const GLvoid* Pointer)
+	{
+		Size = (Size == GL_BGRA) ? 4 : Size;
+		glVertexAttribPointer(Index, Size, Type, Normalized, Stride, Pointer);
+	}
+
+	static FORCEINLINE void ClearDepth(GLdouble Depth)
+	{
+		glClearDepthf(Depth);
+	}
+
+	static FORCEINLINE void GenerateMipmap( GLenum Target )
+	{
+		glGenerateMipmap( Target);
+	}
+	
+	static FORCEINLINE bool SupportsGenerateMipmap() { return true; }
+
+	static FORCEINLINE GLuint GetMajorVersion()
+	{
+		return 3;
+	}
+
+	static FORCEINLINE GLuint GetMinorVersion()
+	{
+		return 1;
+	}
+
+	static FORCEINLINE EShaderPlatform GetShaderPlatform()
+	{
+		return SP_OPENGL_ES3_1_ANDROID;
+	}
+
+	static FORCEINLINE ERHIFeatureLevel::Type GetFeatureLevel()
+	{
+		return ERHIFeatureLevel::ES3_1;
+	}
+
+	static FORCEINLINE FString GetAdapterName()
+	{
+		return (TCHAR*)ANSI_TO_TCHAR((const ANSICHAR*)glGetString(GL_RENDERER));
+	}
+
+	static FORCEINLINE void TexParameter(GLenum Target, GLenum Parameter, GLint Value)
+	{
+		glTexParameteri(Target, Parameter, Value);
+	}
+
+	static FORCEINLINE void FramebufferTexture(GLenum Target, GLenum Attachment, GLuint Texture, GLint Level)
+	{
+		// ES 3.2
+		glFramebufferTexture(Target, Attachment, Texture, Level);
+	}
+
+	static FORCEINLINE void FramebufferTexture3D(GLenum Target, GLenum Attachment, GLenum TexTarget, GLuint Texture, GLint Level, GLint ZOffset)
+	{
+		// glFramebufferTexture3D is not supported on GLES
+		glFramebufferTextureLayer(Target, Attachment, Texture, Level, ZOffset);
+	}
+
+	static FORCEINLINE void FramebufferTextureLayer(GLenum Target, GLenum Attachment, GLuint Texture, GLint Level, GLint Layer)
+	{
+		glFramebufferTextureLayer(Target, Attachment, Texture, Level, Layer);
+	}
+
+	static FORCEINLINE void FramebufferTexture2D(GLenum Target, GLenum Attachment, GLenum TexTarget, GLuint Texture, GLint Level)
+	{
+		check(Attachment == GL_COLOR_ATTACHMENT0 ||
+				Attachment == GL_DEPTH_ATTACHMENT || 
+				Attachment == GL_STENCIL_ATTACHMENT ||
+				Attachment == GL_DEPTH_STENCIL_ATTACHMENT ||
+				(Attachment >= GL_COLOR_ATTACHMENT0 && Attachment <= GL_COLOR_ATTACHMENT7));
+
+		glFramebufferTexture2D(Target, Attachment, TexTarget, Texture, Level);
+		VERIFY_GL(FramebufferTexture_2D);
+	}
+
+	static FORCEINLINE void FramebufferTexture2DMultisample(GLenum Target, GLenum Attachment, GLenum TexTarget, GLuint Texture, GLint Level, GLint NumSamples)
+	{
+		check(glFramebufferTexture2DMultisampleEXT != nullptr);
+		glFramebufferTexture2DMultisampleEXT(Target, Attachment, TexTarget, Texture, Level, NumSamples);
+	}
+
+	static FORCEINLINE void FramebufferTextureMultiviewOVR(GLenum Target, GLenum Attachment, GLuint Texture, GLint Level, GLint BaseViewIndex, GLsizei NumViews)
+	{
+		check(glFramebufferTextureMultiviewOVR);
+		glFramebufferTextureMultiviewOVR(Target, Attachment, Texture, Level, BaseViewIndex, NumViews);
+	}
+	
+	static FORCEINLINE void FramebufferTextureMultisampleMultiviewOVR(GLenum Target, GLenum Attachment, GLuint Texture, GLint Level, GLsizei NumSamples, GLint BaseViewIndex, GLsizei NumViews)
+	{
+		check(glFramebufferTextureMultisampleMultiviewOVR);
+		glFramebufferTextureMultisampleMultiviewOVR(Target, Attachment, Texture, Level, NumSamples, BaseViewIndex, NumViews);
+	}
+
+	static FORCEINLINE void BlitFramebuffer(GLint SrcX0, GLint SrcY0, GLint SrcX1, GLint SrcY1, GLint DstX0, GLint DstY0, GLint DstX1, GLint DstY1, GLbitfield Mask, GLenum Filter)
+	{
+		glBlitFramebuffer(SrcX0, SrcY0, SrcX1, SrcY1, DstX0, DstY0, DstX1, DstY1, Mask, Filter);
+	}
+
+	static FORCEINLINE bool TexStorage2D(GLenum Target, GLint Levels, GLint InternalFormat, GLsizei Width, GLsizei Height, GLenum Format, GLenum Type, ETextureCreateFlags Flags)
+	{
+		glTexStorage2D(Target, Levels, InternalFormat, Width, Height);
+		VERIFY_GL(glTexStorage2D);
+		return true;
+	}
+
+	static FORCEINLINE void DrawArraysInstanced(GLenum Mode, GLint First, GLsizei Count, GLsizei InstanceCount)
+	{
+		glDrawArraysInstanced(Mode, First, Count, InstanceCount);
+	}
+
+	static FORCEINLINE void DrawElementsInstanced(GLenum Mode, GLsizei Count, GLenum Type, const GLvoid* Indices, GLsizei InstanceCount)
+	{
+		glDrawElementsInstanced(Mode, Count, Type, Indices, InstanceCount);
+	}
+
+	static FORCEINLINE void CopyBufferSubData(GLenum ReadTarget, GLenum WriteTarget, GLintptr ReadOffset, GLintptr WriteOffset, GLsizeiptr Size)
+	{
+		glCopyBufferSubData(ReadTarget, WriteTarget, ReadOffset, WriteOffset, Size);
+	}
+
+	static FORCEINLINE void DrawArraysIndirect(GLenum Mode, const void* Offset)
+	{
+		glDrawArraysIndirect(Mode, Offset);
+	}
+
+	static FORCEINLINE void DrawElementsIndirect(GLenum Mode, GLenum Type, const void* Offset)
+	{
+		glDrawElementsIndirect(Mode, Type, Offset);
+	}
+
+	static FORCEINLINE void VertexAttribDivisor(GLuint Index, GLuint Divisor)
+	{
+		glVertexAttribDivisor(Index, Divisor);
+	}
+
+	static FORCEINLINE void TexStorage3D(GLenum Target, GLint Levels, GLint InternalFormat, GLsizei Width, GLsizei Height, GLsizei Depth, GLenum Format, GLenum Type)
+	{
+		glTexStorage3D(Target, Levels, InternalFormat, Width, Height, Depth);
+	}
+
+	static FPlatformOpenGLDevice* CreateDevice()	UGL_REQUIRED(NULL)
+	static FPlatformOpenGLContext* CreateContext(FPlatformOpenGLDevice* Device, void* WindowHandle)	UGL_REQUIRED(NULL)
+
+	static FORCEINLINE void InvalidateFramebuffer(GLenum Target, GLsizei NumAttachments, const GLenum* Attachments)
+	{
+		glInvalidateFramebuffer(Target, NumAttachments, Attachments);
+	}
+
+	static FORCEINLINE void GenBuffers(GLsizei n, GLuint* buffers)
+	{
+		glGenBuffers(n, buffers);
+	}
+
+	static FORCEINLINE void GenTextures(GLsizei n, GLuint* textures)
+	{
+		glGenTextures(n, textures);
+	}
+
+	static FORCEINLINE bool TimerQueryDisjoint()
+	{
+		bool Disjoint = false;
+
+		if (bTimerQueryCanBeDisjoint)
+		{
+			GLint WasDisjoint = 0;
+			glGetIntegerv(GL_GPU_DISJOINT_EXT, &WasDisjoint);
+			Disjoint = (WasDisjoint != 0);
+		}
+
+		return Disjoint;
+	}
+
+	static FORCEINLINE void BindVertexBuffer(GLuint BindingIndex, GLuint Buffer, GLintptr Offset, GLsizei Stride)
+	{
+		glBindVertexBuffer(BindingIndex, Buffer, Offset, Stride);
+	}
+	static FORCEINLINE void VertexAttribFormat(GLuint AttribIndex, GLint Size, GLenum Type, GLboolean Normalized, GLuint RelativeOffset)
+	{
+		glVertexAttribFormat(AttribIndex, Size, Type, Normalized, RelativeOffset);
+	}
+	static FORCEINLINE void VertexAttribIFormat(GLuint AttribIndex, GLint Size, GLenum Type, GLuint RelativeOffset)
+	{
+		glVertexAttribIFormat(AttribIndex, Size, Type, RelativeOffset);
+	}
+	static FORCEINLINE void VertexAttribBinding(GLuint AttribIndex, GLuint BindingIndex)
+	{
+		glVertexAttribBinding(AttribIndex, BindingIndex);
+	}
+	static FORCEINLINE void VertexBindingDivisor(GLuint BindingIndex, GLuint Divisor)
+	{
+		glVertexBindingDivisor(BindingIndex, Divisor);
+	}
+
+	static FORCEINLINE void BufferStorage(GLenum Target, GLsizeiptr Size, const void* Data, GLbitfield Flags)
+	{
+		glBufferStorageEXT(Target, Size, Data, Flags);
+	}
+protected:
+	/** GL_EXT_disjoint_timer_query */
+	static bool bSupportsDisjointTimeQueries;
+
+	/** Some timer query implementations are never disjoint */
+	static bool bTimerQueryCanBeDisjoint;
+
+	/** GL_APPLE_texture_format_BGRA8888 */
+	static bool bSupportsBGRA8888;
+
+	/** GL_NV_texture_compression_s3tc, GL_EXT_texture_compression_s3tc */
+	static bool bSupportsDXT;
+
+	/** OpenGL ES 3.0 profile */
+	static bool bSupportsETC2;
+
+	/** GL_EXT_color_buffer_float */
+	static bool bSupportsColorBufferFloat;
+
+	/** GL_EXT_color_buffer_half_float */
+	static bool bSupportsColorBufferHalfFloat;
+
+	/** GL_EXT_shader_framebuffer_fetch */
+	static bool bSupportsShaderFramebufferFetch;
+
+	/** GL_EXT_shader_framebuffer_fetch (MRT's) */
+	static bool bSupportsShaderMRTFramebufferFetch;
+
+
+	/** GL_ARM_shader_framebuffer_fetch_depth_stencil */
+	static bool bSupportsShaderDepthStencilFetch;
+
+	/** GL_EXT_MULTISAMPLED_RENDER_TO_TEXTURE */
+	static bool bSupportsMultisampledRenderToTexture;
+
+	/** workaround for GL_EXT_shader_pixel_local_storage */
+	static bool bSupportsPixelLocalStorage;
+
+	/** GL_FRAGMENT_SHADER, GL_LOW_FLOAT */
+	static int ShaderLowPrecision;
+
+	/** GL_FRAGMENT_SHADER, GL_MEDIUM_FLOAT */
+	static int ShaderMediumPrecision;
+
+	/** GL_FRAGMENT_SHADER, GL_HIGH_FLOAT */
+	static int ShaderHighPrecision;
+
+	/** GL_NV_framebuffer_blit */
+	static bool bSupportsNVFrameBufferBlit;
+
+	/** GL_OES_standard_derivations */
+	static bool bSupportsStandardDerivativesExtension;
+
+	/** Maximum number of MSAA samples supported on chip in tile memory, or 1 if not available */
+	static GLint MaxMSAASamplesTileMem;
+
+	/** GL_EXT_texture_compression_astc_decode_mode */
+	static bool bSupportsASTCDecodeMode;
+
+public:
+	/* This indicates failure when attempting to retrieve driver's binary representation of the hack program  */
+	static bool bBinaryProgramRetrievalFailed;
+
+	/* Some Mali devices do not work correctly with early_fragment_test enabled */
+	static bool bRequiresDisabledEarlyFragmentTests;
+		
+	/* This is a workaround for a Mali bug where read-only buffers do not work when passed to functions*/
+	static bool bRequiresReadOnlyBuffersWorkaround;
+
+	/* This is to avoid a bug in Adreno drivers that define GL_ARM_shader_framebuffer_fetch_depth_stencil even when device does not support this extension  */
+	static bool bRequiresARMShaderFramebufferFetchDepthStencilUndef;
+
+	/** GL_OES_vertex_type_10_10_10_2 */
+	static bool bSupportsRGB10A2;
+
+	/** GL_OES_get_program_binary */
+	static bool bSupportsProgramBinary;
+
+	/** GL_EXT_buffer_storage */
+	static bool bSupportsBufferStorage;
+
+	/** GL_EXT_depth_clamp */
+	static bool bSupportsDepthClamp;
+
+	enum class EFeatureLevelSupport : uint8
+	{
+		Invalid,	// no feature level has yet been determined
+		ES2,
+		ES31,
+		ES32
+	};
+
+	/** Describes which feature level is currently being supported */
+	static EFeatureLevelSupport CurrentFeatureLevelSupport;
+
+	/** Whether device supports Hidden Surface Removal */
+	static bool bHasHardwareHiddenSurfaceRemoval;
+
+	/** Whether device supports mobile multi-view */
+	static bool bSupportsMobileMultiView;
+
+	static GLint MaxComputeUniformComponents;
+
+	static GLint MaxCombinedUAVUnits;
+	static GLint MaxComputeUAVUnits;
+	static GLint MaxPixelUAVUnits;
+};
+
+#endif //desktop
diff --git a/Engine/Source/Runtime/OpenGLDrv/Private/OpenGLResources.h b/Engine/Source/Runtime/OpenGLDrv/Private/OpenGLResources.h
new file mode 100644
--- /dev/null
+++ b/Engine/Source/Runtime/OpenGLDrv/Private/OpenGLResources.h
@@ -0,0 +1,1430 @@
+// Copyright Epic Games, Inc. All Rights Reserved.
+
+/*=============================================================================
+	OpenGLResources.h: OpenGL resource RHI definitions.
+=============================================================================*/
+
+#pragma once
+
+// HEADER_UNIT_SKIP - Not included directly
+
+#include "CoreTypes.h"
+#include "Misc/AssertionMacros.h"
+#include "HAL/UnrealMemory.h"
+#include "HAL/LowLevelMemTracker.h"
+#include "Containers/ContainerAllocationPolicies.h"
+#include "Containers/Array.h"
+#include "Math/UnrealMathUtility.h"
+#include "Logging/LogMacros.h"
+#include "Containers/BitArray.h"
+#include "Math/IntPoint.h"
+#include "Misc/CommandLine.h"
+#include "Templates/RefCounting.h"
+#include "Stats/Stats.h"
+#include "RHI.h"
+#include "BoundShaderStateCache.h"
+#include "RenderResource.h"
+#include "OpenGLShaderResources.h"
+#include "PsoLruCache.h"
+
+class FOpenGLDynamicRHI;
+class FOpenGLLinkedProgram;
+class FOpenGLTexture;
+typedef TArray<ANSICHAR> FAnsiCharArray;
+
+
+extern void OnBufferDeletion( GLuint BufferResource );
+extern void OnPixelBufferDeletion( GLuint PixelBufferResource );
+extern void OnUniformBufferDeletion( GLuint UniformBufferResource, uint32 AllocatedSize, bool bStreamDraw, uint32 Offset, uint8* Pointer );
+extern void OnProgramDeletion( GLint ProgramResource );
+
+extern void CachedBindBuffer( GLenum Type, GLuint Buffer );
+extern void CachedBindPixelUnpackBuffer( GLenum Type, GLuint Buffer );
+extern void CachedBindUniformBuffer( GLuint Buffer );
+extern bool IsUniformBufferBound( GLuint Buffer );
+
+namespace OpenGLConsoleVariables
+{
+	extern int32 bUseMapBuffer;
+	extern int32 MaxSubDataSize;
+	extern int32 bUseStagingBuffer;
+	extern int32 bBindlessTexture;
+	extern int32 bUseBufferDiscard;
+};
+
+#if PLATFORM_WINDOWS
+#define RESTRICT_SUBDATA_SIZE 1
+#else
+#define RESTRICT_SUBDATA_SIZE 0
+#endif
+
+namespace OpenGLBufferStats
+{
+	void UpdateUniformBufferStats(int64 BufferSize, bool bAllocating);
+	void UpdateBufferStats(const FRHIBufferDesc& BufferDesc, bool bAllocating);
+}
+
+// Extra stats for finer-grained timing
+// They shouldn't always be on, as they may impact overall performance
+#define OPENGLRHI_DETAILED_STATS 0
+
+
+#if OPENGLRHI_DETAILED_STATS
+	DECLARE_CYCLE_STAT_EXTERN(TEXT("MapBuffer time"),STAT_OpenGLMapBufferTime,STATGROUP_OpenGLRHI, );
+	DECLARE_CYCLE_STAT_EXTERN(TEXT("UnmapBuffer time"),STAT_OpenGLUnmapBufferTime,STATGROUP_OpenGLRHI, );
+	#define SCOPE_CYCLE_COUNTER_DETAILED(Stat)	SCOPE_CYCLE_COUNTER(Stat)
+	#define DETAILED_QUICK_SCOPE_CYCLE_COUNTER(x) QUICK_SCOPE_CYCLE_COUNTER(x)
+#else
+	#define SCOPE_CYCLE_COUNTER_DETAILED(Stat)
+	#define DETAILED_QUICK_SCOPE_CYCLE_COUNTER(x)
+#endif
+
+#define GLDEBUG_LABELS_ENABLED (!UE_BUILD_SHIPPING)
+
+class FOpenGLViewableResource
+{
+public:
+	~FOpenGLViewableResource()
+	{
+		checkf(!HasLinkedViews(), TEXT("All linked views must have been removed before the underlying resource can be deleted."));
+	}
+
+	bool HasLinkedViews() const
+	{
+		return LinkedViews != nullptr;
+	}
+
+	void UpdateLinkedViews();
+
+private:
+	friend class FOpenGLShaderResourceView;
+	friend class FOpenGLUnorderedAccessView;
+	class FOpenGLView* LinkedViews = nullptr;
+};
+
+typedef void (*BufferBindFunction)( GLenum Type, GLuint Buffer );
+
+template <typename BaseType, BufferBindFunction BufBind>
+class TOpenGLBuffer : public BaseType
+{
+	void LoadData(uint32 InOffset, uint32 InSize, const void* InData)
+	{
+		VERIFY_GL_SCOPE();
+		const uint8* Data = (const uint8*)InData;
+		const uint32 BlockSize = OpenGLConsoleVariables::MaxSubDataSize;
+
+		if (BlockSize > 0)
+		{
+			while ( InSize > 0)
+			{
+				const uint32 BufferSize = FMath::Min<uint32>( BlockSize, InSize);
+
+				FOpenGL::BufferSubData( Type, InOffset, BufferSize, Data);
+
+				InOffset += BufferSize;
+				InSize -= BufferSize;
+				Data += BufferSize;
+			}
+		}
+		else
+		{
+			FOpenGL::BufferSubData( Type, InOffset, InSize, InData);
+		}
+	}
+
+	GLenum GetAccess()
+	{
+		// Previously there was special-case logic to always use GL_STATIC_DRAW for vertex buffers allocated from staging buffer.
+		// However it seems to be incorrect as NVidia drivers complain (via debug output callback) about VIDEO->HOST copying for buffers with such hints
+		return IsDynamic() ? GL_DYNAMIC_DRAW : GL_STATIC_DRAW;
+	}
+
+public:
+	TOpenGLBuffer(FRHICommandListBase* RHICmdList, GLenum InType, FRHIBufferDesc const& BufferDesc, const void *InData)
+		: BaseType(BufferDesc)
+		, Resource(0)
+		, Type(InType)
+		, bIsLocked(false)
+		, bIsLockReadOnly(false)
+		, bLockBufferWasAllocated(false)
+		, LockSize(0)
+		, LockOffset(0)
+		, LockBuffer(nullptr)
+	{
+		if (RHICmdList && RHICmdList->IsTopOfPipe() && InData)
+		{
+			void* LocalCopy = RHICmdList->Alloc(BaseType::GetSize(), 16);
+			FMemory::Memcpy(LocalCopy, InData, BaseType::GetSize());
+			InData = LocalCopy;
+		}
+
+		auto InitLambda = [this, InData]()
+		{
+			VERIFY_GL_SCOPE();
+
+			FOpenGL::GenBuffers(1, &Resource);
+			check(Type != GL_UNIFORM_BUFFER || !IsUniformBufferBound(Resource));
+
+			Bind();
+
+#if !RESTRICT_SUBDATA_SIZE
+			glBufferData(Type, BaseType::GetSize(), InData, GetAccess());
+#else
+			glBufferData(Type, BaseType::GetSize(), nullptr, GetAccess());
+			if (InData != nullptr)
+			{
+				LoadData(0, BaseType::GetSize(), InData);
+			}
+#endif
+			OpenGLBufferStats::UpdateBufferStats(BaseType::GetDesc(), true);
+		};
+
+		if (!BufferDesc.IsNull())
+		{
+			if (RHICmdList)
+			{
+				RHICmdList->EnqueueLambda([Lambda = MoveTemp(InitLambda)](FRHICommandListBase&) { Lambda(); });
+			}
+			else
+			{
+				InitLambda();
+			}
+		}
+	}
+
+	virtual ~TOpenGLBuffer()
+	{
+		ReleaseOwnership();
+	}
+
+	void Bind()
+	{
+		VERIFY_GL_SCOPE();
+		BufBind(Type, Resource);
+	}
+
+	uint8 *Lock(uint32 InOffset, uint32 InSize, bool bReadOnly, bool bDiscard)
+	{
+		SCOPE_CYCLE_COUNTER_DETAILED(STAT_OpenGLMapBufferTime);
+		check(InOffset + InSize <= BaseType::GetSize());
+		//check( LockBuffer == NULL );	// Only one outstanding lock is allowed at a time!
+		check( !bIsLocked );	// Only one outstanding lock is allowed at a time!
+		VERIFY_GL_SCOPE();
+
+		Bind();
+
+		bIsLocked = true;
+		bIsLockReadOnly = bReadOnly;
+		uint8 *Data = NULL;
+
+		// Discard if the input size is the same as the backing store size, regardless of the input argument, as orphaning the backing store will typically be faster.
+		bDiscard = (bDiscard || (!bReadOnly && InSize == BaseType::GetSize())) && FOpenGL::DiscardFrameBufferToResize();
+
+		// Map buffer is faster in some circumstances and slower in others, decide when to use it carefully.
+		bool const bUseMapBuffer = (bReadOnly || OpenGLConsoleVariables::bUseMapBuffer);
+
+		// If we're able to discard the current data, do so right away
+		// If we can then we should orphan the buffer name & reallocate the backing store only once as calls to glBufferData may do so even when the size is the same.
+		uint32 DiscardSize = (bDiscard && !bUseMapBuffer && InSize == BaseType::GetSize() && !RESTRICT_SUBDATA_SIZE) ? 0 : BaseType::GetSize();
+
+		// Don't call BufferData if Bindless is on, as bindless texture buffers make buffers immutable
+		if (bDiscard && !OpenGLConsoleVariables::bBindlessTexture && OpenGLConsoleVariables::bUseBufferDiscard)
+		{
+			// @todo Lumin hack:
+			// When not hinted with GL_STATIC_DRAW, glBufferData() would introduce long uploading times
+			// that would show up in TGD. Without the workaround of hinting glBufferData() with the static buffer usage, 
+			// the buffer mapping / unmapping has an unexpected cost(~5 - 10ms) that manifests itself in light grid computation 
+			// and vertex buffer mapping for bone matrices. We believe this issue originates from the driver as the OpenGL spec 
+			// specifies the following on the usage hint parameter of glBufferData() :
+			//
+			// > usage is a hint to the GL implementation as to how a buffer object's data store will be accessed. 
+			// > This enables the GL implementation to make more intelligent decisions that may significantly impact buffer object performance. 
+			// > It does not, however, constrain the actual usage of the data store.
+			//
+			// As the alternative approach of using uniform buffers for bone matrix uploading (isntead of buffer mapping/unmapping)
+			// limits the number of bone matrices to 75 in the current engine architecture and that is not desirable, 
+			// we can stick with the STATIC_DRAW hint workaround for glBufferData().
+			//
+			// We haven't seen the buffer mapping/unmapping issue show up elsewhere in the pipeline in our test scenes. 
+			// However, depending on the UnrealEditor features that are used, this issue might pop up elsewhere that we're yet to see.
+			// As there are concerns for maximum number of bone matrices, going for the GL_STATIC_DRAW hint should be safer, 
+			// given the fact that it won't constrain the actual usage of the data store as per the OpenGL4 spec.
+#if 0
+			glBufferData(Type, DiscardSize, NULL, GL_STATIC_DRAW);
+#else
+			glBufferData(Type, DiscardSize, NULL, GetAccess());
+#endif			
+		}
+
+		if ( bUseMapBuffer)
+		{
+			FOpenGL::EResourceLockMode LockMode = bReadOnly ? FOpenGL::EResourceLockMode::RLM_ReadOnly : FOpenGL::EResourceLockMode::RLM_WriteOnly;
+			Data = static_cast<uint8*>( FOpenGL::MapBufferRange( Type, InOffset, InSize, LockMode ) );
+//			checkf(Data != NULL, TEXT("FOpenGL::MapBufferRange Failed, glError %d (0x%x)"), glGetError(), glGetError());
+
+			LockOffset = InOffset;
+			LockSize = InSize;
+			LockBuffer = Data;
+			bLockBufferWasAllocated = false;
+		}
+		else
+		{
+			// Allocate a temp buffer to write into
+			LockOffset = InOffset;
+			LockSize = InSize;
+			if (CachedBuffer && InSize <= CachedBufferSize)
+			{
+				LockBuffer = CachedBuffer;
+				CachedBuffer = nullptr;
+				// Keep CachedBufferSize to keep the actual size allocated.
+			}
+			else
+			{
+				ReleaseCachedBuffer();
+				LockBuffer = FMemory::Malloc( InSize );
+				CachedBufferSize = InSize; // Safegard
+			}
+			Data = static_cast<uint8*>( LockBuffer );
+			bLockBufferWasAllocated = true;
+		}
+
+		if (Data == nullptr)
+		{
+			UE_LOG(LogRHI, Fatal, TEXT("Failed to lock buffer: Resource %u, Size %u, Offset %u, bReadOnly %d, bUseMapBuffer %d, glError (0x%x)"), 
+				(uint32)Resource, 
+				InSize, 
+				InOffset, 
+				bReadOnly ? 1:0, 
+				bUseMapBuffer ? 1:0, 
+				glGetError());
+		}
+		
+		return Data;
+	}
+
+	uint8 *LockWriteOnlyUnsynchronized(uint32 InOffset, uint32 InSize, bool bDiscard)
+	{
+		//SCOPE_CYCLE_COUNTER_DETAILED(STAT_OpenGLMapBufferTime);
+		check(InOffset + InSize <= BaseType::GetSize());
+		//check( LockBuffer == NULL );	// Only one outstanding lock is allowed at a time!
+		check( !bIsLocked );	// Only one outstanding lock is allowed at a time!
+		VERIFY_GL_SCOPE();
+
+		Bind();
+
+		bIsLocked = true;
+		bIsLockReadOnly = false;
+		uint8 *Data = NULL;
+
+		// Discard if the input size is the same as the backing store size, regardless of the input argument, as orphaning the backing store will typically be faster.
+		bDiscard = (bDiscard || InSize == BaseType::GetSize()) && FOpenGL::DiscardFrameBufferToResize();
+
+		// Map buffer is faster in some circumstances and slower in others, decide when to use it carefully.
+		bool const bUseMapBuffer = OpenGLConsoleVariables::bUseMapBuffer;
+
+		// If we're able to discard the current data, do so right away
+		// If we can then we should orphan the buffer name & reallocate the backing store only once as calls to glBufferData may do so even when the size is the same.
+		uint32 DiscardSize = (bDiscard && !bUseMapBuffer && InSize == BaseType::GetSize() && !RESTRICT_SUBDATA_SIZE) ? 0 : BaseType::GetSize();
+
+		// Don't call BufferData if Bindless is on, as bindless texture buffers make buffers immutable
+		if ( bDiscard && !OpenGLConsoleVariables::bBindlessTexture && OpenGLConsoleVariables::bUseBufferDiscard)
+		{
+			glBufferData( Type, DiscardSize, NULL, GetAccess());
+		}
+
+		if ( bUseMapBuffer)
+		{
+			FOpenGL::EResourceLockMode LockMode = bDiscard ? FOpenGL::EResourceLockMode::RLM_WriteOnly : FOpenGL::EResourceLockMode::RLM_WriteOnlyUnsynchronized;
+			Data = static_cast<uint8*>( FOpenGL::MapBufferRange( Type, InOffset, InSize, LockMode ) );
+			LockOffset = InOffset;
+			LockSize = InSize;
+			LockBuffer = Data;
+			bLockBufferWasAllocated = false;
+		}
+		else
+		{
+			// Allocate a temp buffer to write into
+			LockOffset = InOffset;
+			LockSize = InSize;
+			if (CachedBuffer && InSize <= CachedBufferSize)
+			{
+				LockBuffer = CachedBuffer;
+				CachedBuffer = nullptr;
+				// Keep CachedBufferSize to keep the actual size allocated.
+			}
+			else
+			{
+				ReleaseCachedBuffer();
+				LockBuffer = FMemory::Malloc( InSize );
+				CachedBufferSize = InSize; // Safegard
+			}
+			Data = static_cast<uint8*>( LockBuffer );
+			bLockBufferWasAllocated = true;
+		}
+
+		if (Data == nullptr)
+		{
+			UE_LOG(LogRHI, Fatal, TEXT("Failed to lock buffer (write only): Resource %u, Size %u, Offset %u, bUseMapBuffer %d, glError (0x%x)"), 
+				(uint32)Resource, 
+				InSize, 
+				InOffset, 
+				bUseMapBuffer ? 1:0, 
+				glGetError());
+		}
+
+		return Data;
+	}
+
+	void Unlock()
+	{
+		//SCOPE_CYCLE_COUNTER_DETAILED(STAT_OpenGLUnmapBufferTime);
+		VERIFY_GL_SCOPE();
+		if (bIsLocked)
+		{
+			Bind();
+
+			if (OpenGLConsoleVariables::bUseMapBuffer || bIsLockReadOnly)
+			{
+				check(!bLockBufferWasAllocated);
+				FOpenGL::UnmapBufferRange(Type, LockOffset, LockSize);
+				LockBuffer = NULL;
+			}
+			else
+			{
+#if !RESTRICT_SUBDATA_SIZE
+				// Check for the typical, optimized case
+				if(LockSize == BaseType::GetSize())
+				{
+					if (FOpenGL::DiscardFrameBufferToResize())
+					{
+						glBufferData(Type, BaseType::GetSize(), LockBuffer, GetAccess());
+					}
+					else
+					{
+						FOpenGL::BufferSubData(Type, 0, LockSize, LockBuffer);
+					}
+					check( LockBuffer != NULL );
+				}
+				else
+				{
+					// Only updating a subset of the data
+					FOpenGL::BufferSubData(Type, LockOffset, LockSize, LockBuffer);
+					check( LockBuffer != NULL );
+				}
+#else
+				LoadData(LockOffset, LockSize, LockBuffer);
+				check(LockBuffer != NULL);
+#endif
+
+				check(bLockBufferWasAllocated);
+
+				if (EnumHasAnyFlags(this->GetUsage(), BUF_Volatile))
+				{
+					ReleaseCachedBuffer(); // Safegard
+
+					CachedBuffer = LockBuffer;
+					// Possibly > LockSize when reusing cached allocation.
+					CachedBufferSize = FMath::Max<GLuint>(CachedBufferSize, LockSize);
+				}
+				else
+				{
+					FMemory::Free(LockBuffer);
+				}
+				LockBuffer = NULL;
+				bLockBufferWasAllocated = false;
+				LockSize = 0;
+			}
+			bIsLocked = false;
+		}
+	}
+
+	void Update(void *InData, uint32 InOffset, uint32 InSize, bool bDiscard)
+	{
+		check(InOffset + InSize <= BaseType::GetSize());
+		VERIFY_GL_SCOPE();
+		Bind();
+#if !RESTRICT_SUBDATA_SIZE
+		FOpenGL::BufferSubData(Type, InOffset, InSize, InData);
+#else
+		LoadData(InOffset, InSize, InData);
+#endif
+	}
+
+	bool IsDynamic() const { return EnumHasAnyFlags(this->GetUsage(), BUF_AnyDynamic); }
+	bool IsLocked() const { return bIsLocked; }
+	bool IsLockReadOnly() const { return bIsLockReadOnly; }
+	void* GetLockedBuffer() const { return LockBuffer; }
+
+	void ReleaseCachedBuffer()
+	{
+		if (CachedBuffer)
+		{
+			FMemory::Free(CachedBuffer);
+			CachedBuffer = nullptr;
+			CachedBufferSize = 0;
+		}
+		// Don't reset CachedBufferSize if !CachedBuffer since it could be the locked buffer allocation size.
+	}
+
+	void TakeOwnership(TOpenGLBuffer& Other)
+	{
+		VERIFY_GL_SCOPE();
+		check(!bIsLocked && !Other.bIsLocked);
+
+		ReleaseOwnership();
+
+		BaseType::TakeOwnership(Other);
+
+		Type             = Other.Type;
+		Resource         = Other.Resource;
+		CachedBuffer     = Other.CachedBuffer;
+		CachedBufferSize = Other.CachedBufferSize;
+
+		Other.Type             = 0;
+		Other.Resource         = 0;
+		Other.CachedBuffer     = nullptr;
+		Other.CachedBufferSize = 0;
+	}
+
+	void ReleaseOwnership()
+	{
+		VERIFY_GL_SCOPE();
+
+		BaseType::ReleaseOwnership();
+
+		if (Resource != 0)
+		{
+			if (LockBuffer != NULL)
+			{
+				if (bLockBufferWasAllocated)
+				{
+					FMemory::Free(LockBuffer);
+				}
+				else
+				{
+					UE_LOG(LogRHI, Warning, TEXT("Destroying TOpenGLBuffer without returning memory to the driver; possibly called RHIMapStagingSurface() but didn't call RHIUnmapStagingSurface()? Resource %u"), Resource);
+				}
+			}
+
+			if (BaseType::OnDelete(Resource, BaseType::GetSize(), false, 0))
+			{
+				FOpenGL::DeleteBuffers(1, &Resource);
+				Resource = 0;
+			}
+
+			LockBuffer = nullptr;
+			OpenGLBufferStats::UpdateBufferStats(BaseType::GetDesc(), false);
+
+			ReleaseCachedBuffer();
+		}
+	}
+
+	GLuint Resource;
+	GLenum Type;
+
+private:
+	uint32 bIsLocked : 1;
+	uint32 bIsLockReadOnly : 1;
+	uint32 bLockBufferWasAllocated : 1;
+
+	GLuint LockSize;
+	GLuint LockOffset;
+	void* LockBuffer;
+
+	// A cached allocation that can be reused. The same allocation can never be in CachedBuffer and LockBuffer at the same time.
+	void* CachedBuffer = nullptr;
+	// The size of the cached buffer allocation. Can be non zero even though CachedBuffer is  null, to preserve the allocation size.
+	GLuint CachedBufferSize = 0;
+};
+
+class FOpenGLBasePixelBuffer : public FRefCountedObject
+{
+public:
+	FOpenGLBasePixelBuffer(FRHIBufferDesc const& BufferDesc)
+		: Desc(BufferDesc)
+	{}
+
+	static bool OnDelete(GLuint Resource,uint32 Size,bool bStreamDraw,uint32 Offset)
+	{
+		OnPixelBufferDeletion(Resource);
+		return true;
+	}
+
+	const FRHIBufferDesc& GetDesc() const { return Desc; }
+	uint32 GetSize() const { return Desc.Size; }
+	EBufferUsageFlags GetUsage() const { return Desc.Usage; }
+
+	void TakeOwnership(FOpenGLBasePixelBuffer& Other)
+	{
+		Desc = Other.Desc;
+		Other.Desc = FRHIBufferDesc::Null();
+	}
+
+	void ReleaseOwnership()
+	{
+		Desc = FRHIBufferDesc::Null();
+	}
+
+private:
+	FRHIBufferDesc Desc;
+};
+
+class FOpenGLBaseBuffer : public FRHIBuffer, public FOpenGLViewableResource
+{
+public:
+	FOpenGLBaseBuffer(FRHIBufferDesc const& BufferDesc)
+		: FRHIBuffer(BufferDesc)
+	{
+		LLM_SCOPED_PAUSE_TRACKING_WITH_ENUM_AND_AMOUNT(ELLMTag::GraphicsPlatform, (int64)GetSize(), ELLMTracker::Platform, ELLMAllocType::None);
+		LLM_SCOPED_PAUSE_TRACKING_WITH_ENUM_AND_AMOUNT(ELLMTag::Meshes, (int64)GetSize(), ELLMTracker::Default, ELLMAllocType::None);
+	}
+
+	~FOpenGLBaseBuffer()
+	{
+		// If ReleaseOwnership was called, then Size is already 0
+		if (GetSize() != 0)
+		{
+			LLM_SCOPED_PAUSE_TRACKING_WITH_ENUM_AND_AMOUNT(ELLMTag::GraphicsPlatform, -(int64)GetSize(), ELLMTracker::Platform, ELLMAllocType::None);
+			LLM_SCOPED_PAUSE_TRACKING_WITH_ENUM_AND_AMOUNT(ELLMTag::Meshes, -(int64)GetSize(), ELLMTracker::Default, ELLMAllocType::None);
+		}
+    }
+
+	void TakeOwnership(FOpenGLBaseBuffer& Other)
+	{
+		FRHIBuffer::TakeOwnership(Other);
+		LLM_SCOPED_PAUSE_TRACKING_WITH_ENUM_AND_AMOUNT(ELLMTag::GraphicsPlatform, (int64)GetSize(), ELLMTracker::Platform, ELLMAllocType::None);
+		LLM_SCOPED_PAUSE_TRACKING_WITH_ENUM_AND_AMOUNT(ELLMTag::Meshes, (int64)GetSize(), ELLMTracker::Default, ELLMAllocType::None);
+	}
+
+	void ReleaseOwnership()
+	{
+		LLM_SCOPED_PAUSE_TRACKING_WITH_ENUM_AND_AMOUNT(ELLMTag::GraphicsPlatform, -(int64)GetSize(), ELLMTracker::Platform, ELLMAllocType::None);
+		LLM_SCOPED_PAUSE_TRACKING_WITH_ENUM_AND_AMOUNT(ELLMTag::Meshes, -(int64)GetSize(), ELLMTracker::Default, ELLMAllocType::None);
+		FRHIBuffer::ReleaseOwnership();
+	}
+
+	static bool OnDelete(GLuint Resource,uint32 Size,bool bStreamDraw,uint32 Offset)
+	{
+		OnBufferDeletion(Resource);
+		return true;
+	}
+};
+
+typedef TOpenGLBuffer<FOpenGLBasePixelBuffer, CachedBindPixelUnpackBuffer> FOpenGLPixelBuffer;
+typedef TOpenGLBuffer<FOpenGLBaseBuffer, CachedBindBuffer> FOpenGLBuffer;
+
+struct FOpenGLEUniformBufferData : public FRefCountedObject
+{
+	FOpenGLEUniformBufferData(uint32 SizeInBytes)
+	{
+		uint32 SizeInUint32s = (SizeInBytes + 3) / 4;
+		Data.Empty(SizeInUint32s);
+		Data.AddUninitialized(SizeInUint32s);
+		OpenGLBufferStats::UpdateUniformBufferStats(Data.GetAllocatedSize(), true);
+	}
+
+	~FOpenGLEUniformBufferData()
+	{
+		OpenGLBufferStats::UpdateUniformBufferStats(Data.GetAllocatedSize(), false);
+	}
+
+	TArray<uint32> Data;
+};
+typedef TRefCountPtr<FOpenGLEUniformBufferData> FOpenGLEUniformBufferDataRef;
+
+class FOpenGLUniformBuffer : public FRHIUniformBuffer
+{
+public:
+	/** The GL resource for this uniform buffer. */
+	GLuint Resource;
+
+	/** The offset of the uniform buffer's contents in the resource. */
+	uint32 Offset;
+
+	/** The data range size of uniform buffer's contents in the resource. */
+	uint32 RangeSize;
+
+	/** When using a persistently mapped buffer this is a pointer to the CPU accessible data. */
+	uint8* PersistentlyMappedBuffer;
+
+	/** Unique ID for state shadowing purposes. */
+	uint32 UniqueID;
+
+	/** Emulated uniform data for ES2. */
+	FOpenGLEUniformBufferDataRef EmulatedBufferData;
+
+	/** The size of the buffer allocated to hold the uniform buffer contents. May be larger than necessary. */
+	uint32 AllocatedSize;
+
+	/** True if the uniform buffer is not used across frames. */
+	bool bStreamDraw;
+
+	/** True if the uniform buffer is emulated */
+	bool bIsEmulatedUniformBuffer;
+
+	/** True if Resource belongs to this UniformBuffer */
+	bool bOwnsResource;
+
+	/** Initialization constructor. */
+	FOpenGLUniformBuffer(const FRHIUniformBufferLayout* InLayout);
+
+	void SetGLUniformBufferParams(GLuint InResource, uint32 InOffset, uint8* InPersistentlyMappedBuffer, uint32 InAllocatedSize, FOpenGLEUniformBufferDataRef InEmulatedBuffer, bool bInStreamDraw);
+
+	/** Destructor. */
+	~FOpenGLUniformBuffer();
+
+	// Provides public non-const access to ResourceTable.
+	// @todo refactor uniform buffers to perform updates as a member function, so this isn't necessary.
+	TArray<TRefCountPtr<FRHIResource>>& GetResourceTable() { return ResourceTable; }
+
+	void SetLayoutTable(const void* Contents, EUniformBufferValidation Validation);
+};
+
+#define MAX_STREAMED_BUFFERS_IN_ARRAY 2	// must be > 1!
+#define MIN_DRAWS_IN_SINGLE_BUFFER 16
+
+template <typename BaseType, uint32 Stride>
+class TOpenGLStreamedBufferArray
+{
+public:
+
+	TOpenGLStreamedBufferArray( void ) {}
+	virtual ~TOpenGLStreamedBufferArray( void ) {}
+
+	void Init( uint32 InitialBufferSize )
+	{
+		for( int32 BufferIndex = 0; BufferIndex < MAX_STREAMED_BUFFERS_IN_ARRAY; ++BufferIndex )
+		{
+			Buffer[BufferIndex] = new BaseType(Stride, InitialBufferSize, BUF_Volatile, NULL, true);
+		}
+		CurrentBufferIndex = 0;
+		CurrentOffset = 0;
+		LastOffset = 0;
+		MinNeededBufferSize = InitialBufferSize / MIN_DRAWS_IN_SINGLE_BUFFER;
+	}
+
+	void Cleanup( void )
+	{
+		for( int32 BufferIndex = 0; BufferIndex < MAX_STREAMED_BUFFERS_IN_ARRAY; ++BufferIndex )
+		{
+			Buffer[BufferIndex].SafeRelease();
+		}
+	}
+
+	uint8* Lock( uint32 DataSize )
+	{
+		check(!Buffer[CurrentBufferIndex]->IsLocked());
+		DataSize = Align(DataSize, (1<<8));	// to keep the speed up, let's start data for each next draw at 256-byte aligned offset
+
+		// Keep our dynamic buffers at least MIN_DRAWS_IN_SINGLE_BUFFER times bigger than max single request size
+		uint32 NeededBufSize = Align( MIN_DRAWS_IN_SINGLE_BUFFER * DataSize, (1 << 20) );	// 1MB increments
+		if (NeededBufSize > MinNeededBufferSize)
+		{
+			MinNeededBufferSize = NeededBufSize;
+		}
+
+		// Check if we need to switch buffer, as the current draw data won't fit in current one
+		bool bDiscard = false;
+		if (Buffer[CurrentBufferIndex]->GetSize() < CurrentOffset + DataSize)
+		{
+			// We do.
+			++CurrentBufferIndex;
+			if (CurrentBufferIndex == MAX_STREAMED_BUFFERS_IN_ARRAY)
+			{
+				CurrentBufferIndex = 0;
+			}
+			CurrentOffset = 0;
+
+			// Check if we should extend the next buffer, as max request size has changed
+			if (MinNeededBufferSize > Buffer[CurrentBufferIndex]->GetSize())
+			{
+				Buffer[CurrentBufferIndex].SafeRelease();
+				Buffer[CurrentBufferIndex] = new BaseType(Stride, MinNeededBufferSize, BUF_Volatile);
+			}
+
+			bDiscard = true;
+		}
+
+		LastOffset = CurrentOffset;
+		CurrentOffset += DataSize;
+
+		return Buffer[CurrentBufferIndex]->LockWriteOnlyUnsynchronized(LastOffset, DataSize, bDiscard);
+	}
+
+	void Unlock( void )
+	{
+		check(Buffer[CurrentBufferIndex]->IsLocked());
+		Buffer[CurrentBufferIndex]->Unlock();
+	}
+
+	BaseType* GetPendingBuffer( void ) { return Buffer[CurrentBufferIndex]; }
+	uint32 GetPendingOffset( void ) { return LastOffset; }
+
+private:
+	TRefCountPtr<BaseType> Buffer[MAX_STREAMED_BUFFERS_IN_ARRAY];
+	uint32 CurrentBufferIndex;
+	uint32 CurrentOffset;
+	uint32 LastOffset;
+	uint32 MinNeededBufferSize;
+};
+
+struct FOpenGLVertexElement
+{
+	GLenum Type;
+	GLuint StreamIndex;
+	GLuint Offset;
+	GLuint Size;
+	GLuint Divisor;
+	GLuint HashStride;
+	uint8 bNormalized;
+	uint8 AttributeIndex;
+	uint8 bShouldConvertToFloat;
+	uint8 Padding;
+
+	FOpenGLVertexElement()
+		: Padding(0)
+	{
+	}
+};
+
+/** Convenience typedef: preallocated array of OpenGL input element descriptions. */
+typedef TArray<FOpenGLVertexElement,TFixedAllocator<MaxVertexElementCount> > FOpenGLVertexElements;
+
+/** This represents a vertex declaration that hasn't been combined with a specific shader to create a bound shader. */
+class FOpenGLVertexDeclaration : public FRHIVertexDeclaration
+{
+public:
+	/** Elements of the vertex declaration. */
+	FOpenGLVertexElements VertexElements;
+
+	uint16 StreamStrides[MaxVertexElementCount];
+
+	/** Initialization constructor. */
+	FOpenGLVertexDeclaration(const FOpenGLVertexElements& InElements, const uint16* InStrides)
+		: VertexElements(InElements)
+	{
+		FMemory::Memcpy(StreamStrides, InStrides, sizeof(StreamStrides));
+	}
+	
+	virtual bool GetInitializer(FVertexDeclarationElementList& Init) override final;
+};
+
+
+/**
+ * Combined shader state and vertex definition for rendering geometry.
+ * Each unique instance consists of a vertex decl, vertex shader, and pixel shader.
+ */
+class FOpenGLBoundShaderState : public FRHIBoundShaderState
+{
+public:
+
+	FCachedBoundShaderStateLink CacheLink;
+
+	uint16 StreamStrides[MaxVertexElementCount];
+
+	FOpenGLLinkedProgram* LinkedProgram;
+	TRefCountPtr<FOpenGLVertexDeclaration> VertexDeclaration;
+	TRefCountPtr<FOpenGLVertexShader> VertexShader;
+	TRefCountPtr<FOpenGLPixelShader> PixelShader;
+	TRefCountPtr<FOpenGLGeometryShader> GeometryShader;
+
+	/** Initialization constructor. */
+	FOpenGLBoundShaderState(
+		FOpenGLLinkedProgram* InLinkedProgram,
+		FRHIVertexDeclaration* InVertexDeclarationRHI,
+		FRHIVertexShader* InVertexShaderRHI,
+		FRHIPixelShader* InPixelShaderRHI,
+		FRHIGeometryShader* InGeometryShaderRHI
+		);
+
+	const TBitArray<>& GetTextureNeeds(int32& OutMaxTextureStageUsed);
+	const TBitArray<>& GetUAVNeeds(int32& OutMaxUAVUnitUsed) const;
+	void GetNumUniformBuffers(int32 NumVertexUniformBuffers[SF_NumGraphicsFrequencies]);
+
+	bool NeedsTextureStage(int32 TextureStageIndex);
+	int32 MaxTextureStageUsed();
+	bool RequiresDriverInstantiation();
+
+	FOpenGLVertexShader* GetVertexShader()
+	{
+		check(IsValidRef(VertexShader));
+		return VertexShader;
+	}
+
+	FOpenGLPixelShader* GetPixelShader()
+	{
+		check(IsValidRef(PixelShader));
+		return PixelShader;
+	}
+
+	FOpenGLGeometryShader* GetGeometryShader()
+	{
+		return GeometryShader;
+	}
+
+	virtual ~FOpenGLBoundShaderState();
+};
+
+class FTextureEvictionLRU
+{
+private:
+	typedef TPsoLruCache<FOpenGLTexture*, FOpenGLTexture*> FOpenGLTextureLRUContainer;
+	FCriticalSection TextureLRULock;
+
+	static FORCEINLINE_DEBUGGABLE FOpenGLTextureLRUContainer& GetLRUContainer()
+	{
+		const int32 MaxNumLRUs = 10000;
+		static FOpenGLTextureLRUContainer TextureLRU(MaxNumLRUs);
+		return TextureLRU;
+	}
+
+public:
+
+	static FORCEINLINE_DEBUGGABLE FTextureEvictionLRU& Get()
+	{
+		static FTextureEvictionLRU Lru;
+		return Lru;
+	}
+	uint32 Num() const { return GetLRUContainer().Num(); }
+
+	void Remove(FOpenGLTexture* TextureBase);
+	bool Add(FOpenGLTexture* TextureBase);
+	void Touch(FOpenGLTexture* TextureBase);
+	void TickEviction();
+	FOpenGLTexture* GetLeastRecent();
+};
+class FTextureEvictionParams
+{
+public:
+	FTextureEvictionParams(uint32 NumMips);
+	~FTextureEvictionParams();
+	TArray<TArray<uint8>> MipImageData;
+
+ 	uint32 bHasRestored : 1;	
+	FSetElementId LRUNode;
+	uint32 FrameLastRendered;
+
+#if GLDEBUG_LABELS_ENABLED
+	FAnsiCharArray TextureDebugName;
+	void SetDebugLabelName(const FAnsiCharArray& TextureDebugNameIn) { TextureDebugName = TextureDebugNameIn; }
+	void SetDebugLabelName(const ANSICHAR * TextureDebugNameIn) { TextureDebugName.Append(TextureDebugNameIn, FCStringAnsi::Strlen(TextureDebugNameIn) + 1); }
+	FAnsiCharArray& GetDebugLabelName() { return TextureDebugName; }
+#else
+	void SetDebugLabelName(FAnsiCharArray TextureDebugNameIn) { checkNoEntry(); }
+	FAnsiCharArray& GetDebugLabelName() { checkNoEntry(); static FAnsiCharArray Dummy;  return Dummy; }
+#endif
+
+	void SetMipData(uint32 MipIndex, const void* Data, uint32 Bytes);
+	void ReleaseMipData(uint32 RetainMips);
+
+	void CloneMipData(const FTextureEvictionParams& Src, uint32 NumMips, int32 SrcOffset, int DstOffset);
+
+	uint32 GetTotalAllocated() const
+	{
+		uint32 TotalAllocated = 0;
+		for (const auto& MipData : MipImageData)
+		{
+			TotalAllocated += MipData.Num();
+		}
+		return TotalAllocated;
+	}
+
+	bool AreAllMipsPresent() const
+	{
+		bool bRet = MipImageData.Num() > 0;
+		for (const auto& MipData : MipImageData)
+		{
+			bRet = bRet && MipData.Num() > 0;
+		}
+		return bRet;
+	}
+};
+
+class FOpenGLTextureDesc
+{
+public:
+	FOpenGLTextureDesc(FRHITextureDesc const& InDesc);
+
+	GLenum Target = GL_NONE;
+	GLenum Attachment = GL_NONE;
+
+	uint32 MemorySize = 0;
+
+	uint8 bCubemap            : 1;
+	uint8 bArrayTexture       : 1;
+	uint8 bStreamable         : 1;
+	uint8 bDepthStencil       : 1;
+	uint8 bCanCreateAsEvicted : 1;
+	uint8 bIsPowerOfTwo       : 1;
+	uint8 bMultisampleRenderbuffer : 1;
+
+private:
+	static bool CanDeferTextureCreation();
+};
+
+class FOpenGLTextureCreateDesc : public FRHITextureCreateDesc, public FOpenGLTextureDesc
+{
+public:
+	FOpenGLTextureCreateDesc(FRHITextureCreateDesc const& CreateDesc)
+		: FRHITextureCreateDesc(CreateDesc)
+		, FOpenGLTextureDesc(CreateDesc)
+	{
+	}
+};
+
+class OPENGLDRV_API FOpenGLTexture : public FRHITexture, public FOpenGLViewableResource
+{
+	// Prevent copying
+	FOpenGLTexture(FOpenGLTexture const&) = delete;
+	FOpenGLTexture& operator = (FOpenGLTexture const&) = delete;
+
+public:
+	// Standard constructor.
+	explicit FOpenGLTexture(FRHICommandListBase& RHICmdList, FOpenGLTextureCreateDesc const& CreateDesc);
+
+	// Constructor for external resources (RHICreateTexture2DFromResource etc).
+	explicit FOpenGLTexture(FOpenGLTextureCreateDesc const& CreateDesc, GLuint Resource);
+
+	// Constructor for RHICreateAliasedTexture
+	enum EAliasConstructorParam { AliasResource };
+	explicit FOpenGLTexture(FOpenGLTexture& Texture, const FString& Name, EAliasConstructorParam);
+	void AliasResources(FOpenGLTexture& Texture);
+
+	virtual ~FOpenGLTexture();
+
+	GLuint GetResource()
+	{
+		TryRestoreGLResource();
+		return Resource;
+	}
+
+	GLuint& GetResourceRef() 
+	{ 
+		VERIFY_GL_SCOPE();
+		TryRestoreGLResource();
+		return Resource;
+	}
+
+	// GetRawResourceName - A const accessor to the resource name, this could potentially be an evicted resource.
+	// It will not trigger the GL resource's creation.
+	GLuint GetRawResourceName() const
+	{
+		return Resource;
+	}
+
+	// GetRawResourceNameRef - A const accessor to the resource name, this could potentially be an evicted resource.
+	// It will not trigger the GL resource's creation.
+	const GLuint& GetRawResourceNameRef() const
+	{
+		return Resource;
+	}
+
+	void SetResource(GLuint InResource)
+	{
+		VERIFY_GL_SCOPE();
+		Resource = InResource;
+	}
+
+	bool IsEvicted() const { VERIFY_GL_SCOPE(); return EvictionParamsPtr.IsValid() && !EvictionParamsPtr->bHasRestored; }
+
+	virtual void* GetTextureBaseRHI() override final
+	{
+		return this;
+	}
+
+	/**
+	 * Locks one of the texture's mip-maps.
+	 * @return A pointer to the specified texture data.
+	 */
+	void* Lock(uint32 MipIndex, uint32 ArrayIndex, EResourceLockMode LockMode, uint32& DestStride);
+
+	/**
+	* Returns the size of the memory block that is returned from Lock, threadsafe
+	*/
+	uint32 GetLockSize(uint32 MipIndex, uint32 ArrayIndex, EResourceLockMode LockMode, uint32& DestStride);
+
+	/** Unlocks a previously locked mip-map. */
+	void Unlock(uint32 MipIndex, uint32 ArrayIndex);
+
+	/** FRHITexture override.  See FRHITexture::GetNativeResource() */
+	virtual void* GetNativeResource() const override
+	{
+		// this must become a full GL resource here, calling the non-const GetResourceRef ensures this.
+		return const_cast<void*>(reinterpret_cast<const void*>(&const_cast<FOpenGLTexture*>(this)->GetResourceRef()));
+	}
+
+	/**
+	 * Accessors to mark whether or not we have allocated storage for each mip/face.
+	 * For non-cubemaps FaceIndex should always be zero.
+	 */
+	bool GetAllocatedStorageForMip(uint32 MipIndex, uint32 FaceIndex) const
+	{
+		return bAllocatedStorage[MipIndex * (bCubemap ? 6 : 1) + FaceIndex];
+	}
+	void SetAllocatedStorageForMip(uint32 MipIndex, uint32 FaceIndex)
+	{
+		bAllocatedStorage[MipIndex * (bCubemap ? 6 : 1) + FaceIndex] = true;
+	}
+
+	// Set allocated storage state for all mip/faces
+	void SetAllocatedStorage(bool bInAllocatedStorage)
+	{
+		bAllocatedStorage.Init(bInAllocatedStorage, this->GetNumMips() * (bCubemap ? 6 : 1));
+	}
+
+	/**
+	 * Clone texture from a source using CopyImageSubData
+	 */
+	void CloneViaCopyImage(FOpenGLTexture* Src, uint32 InNumMips, int32 SrcOffset, int32 DstOffset);
+
+	/**
+	 * Resolved the specified face for a read Lock, for non-renderable, CPU readable surfaces this eliminates the readback inside Lock itself.
+	 */
+	void Resolve(uint32 MipIndex, uint32 ArrayIndex);
+
+	// Texture eviction
+	void RestoreEvictedGLResource(bool bAttemptToRetainMips);
+	bool CanBeEvicted();
+	void TryEvictGLResource();
+
+private:
+	static void UpdateTextureStats(FOpenGLTexture* Texture, bool bAllocating);
+
+	void TryRestoreGLResource()
+	{
+		if (EvictionParamsPtr.IsValid() && !EvictionParamsPtr->bHasRestored)
+		{
+			VERIFY_GL_SCOPE();
+			if (!EvictionParamsPtr->bHasRestored)
+			{
+				RestoreEvictedGLResource(true);
+			}
+			else
+			{
+				check(CanBeEvicted());
+				FTextureEvictionLRU::Get().Touch(this);
+			}
+		}
+	}
+
+	void DeleteGLResource();
+
+	void Fill2DGLTextureImage(const FOpenGLTextureFormat& GLFormat, const bool bSRGB, uint32 MipIndex, const void* LockedBuffer, uint32 LockedSize, uint32 ArrayIndex);
+
+	uint32 GetEffectiveSizeZ() const
+	{
+		FRHITextureDesc const& Desc = GetDesc();
+
+		return Desc.IsTexture3D()
+			? Desc.Depth
+			: Desc.ArraySize;
+	}
+
+private:
+	/** The OpenGL texture resource. */
+	GLuint Resource = GL_NONE;
+
+public:
+	/** The OpenGL texture target. */
+	GLenum const Target = 0;
+
+	/** The OpenGL attachment point. This should always be GL_COLOR_ATTACHMENT0 in case of color buffer, but the actual texture may be attached on other color attachments. */
+	GLenum const Attachment = 0;
+
+	TUniquePtr<FTextureEvictionParams> EvictionParamsPtr;
+
+	// Pointer to current sampler state in this unit
+	class FOpenGLSamplerState* SamplerState = nullptr;
+
+private:
+	TArray<TRefCountPtr<FOpenGLPixelBuffer>> PixelBuffers;
+
+	/** Bitfields marking whether we have allocated storage for each mip */
+	TBitArray<TInlineAllocator<1> > bAllocatedStorage;
+
+public:
+	uint32 const MemorySize;
+
+public:
+	uint8 const bIsPowerOfTwo       : 1;
+	uint8 const bCanCreateAsEvicted : 1;
+	uint8 const bStreamable         : 1;
+	uint8 const bCubemap            : 1;
+	uint8 const bArrayTexture       : 1;
+	uint8 const bDepthStencil       : 1;
+	uint8 const bAlias              : 1;
+	uint8 const bMultisampleRenderbuffer : 1;
+};
+
+class FOpenGLRenderQuery : public FRHIRenderQuery
+{
+public:
+
+	/** The query resource. */
+	GLuint Resource;
+
+	/** Identifier of the OpenGL context the query is a part of. */
+	uint64 ResourceContext;
+
+	/** The cached query result. */
+	GLuint64 Result;
+
+	FThreadSafeCounter TotalBegins;
+	FThreadSafeCounter TotalResults;
+
+	/** true if the context the query is in was released from another thread */
+	bool bResultWasSuccess;
+
+	/** true if the context the query is in was released from another thread */
+	bool bInvalidResource;
+
+	// todo: memory optimize
+	ERenderQueryType QueryType;
+
+	FOpenGLRenderQuery(ERenderQueryType InQueryType);
+	virtual ~FOpenGLRenderQuery();
+
+	void AcquireResource();
+	static void ReleaseResource(GLuint Resource, uint64 ResourceContext);
+};
+
+class FOpenGLView : public TIntrusiveLinkedList<FOpenGLView>
+{
+public:
+	virtual void UpdateView() = 0;
+};
+
+class FOpenGLUnorderedAccessView final : public FRHIUnorderedAccessView, public FOpenGLView
+{
+public:
+	FOpenGLUnorderedAccessView(FRHICommandListBase& RHICmdList, FRHIViewableResource* Resource, FRHIViewDesc const& ViewDesc);
+	virtual ~FOpenGLUnorderedAccessView();
+
+	FOpenGLViewableResource* GetBaseResource() const;
+
+	void UpdateView() override;
+
+	GLuint Resource = 0;
+	GLuint BufferResource = 0;
+	GLenum Format = 0;
+	uint8  UnrealFormat = 0;
+
+	uint32 GetBufferSize() const
+	{
+		return IsBuffer() ? GetBuffer()->GetSize() : 0;
+	}
+
+	bool IsLayered() const
+	{
+		return IsTexture() && GetTexture()->GetDesc().Dimension == ETextureDimension::Texture3D;
+	}
+
+	GLint GetLayer() const
+	{
+		return 0;
+	}
+
+private:
+	void Invalidate();
+	bool OwnsResource = false;
+};
+
+class FOpenGLShaderResourceView final : public FRHIShaderResourceView, public FOpenGLView
+{
+public:
+	FOpenGLShaderResourceView(FRHICommandListBase& RHICmdList, FRHIViewableResource* Resource, FRHIViewDesc const& ViewDesc);
+	virtual ~FOpenGLShaderResourceView();
+
+	FOpenGLViewableResource* GetBaseResource() const;
+
+	void UpdateView() override;
+
+	/** OpenGL texture the buffer is bound with */
+	GLuint Resource = GL_NONE;
+	GLenum Target = GL_TEXTURE_BUFFER;
+
+	int32 LimitMip = -1;
+
+private:
+	void Invalidate();
+	bool OwnsResource = false;
+};
+
+void OPENGLDRV_API ReleaseOpenGLFramebuffers(FRHITexture* TextureRHI);
+
+/** A OpenGL event query resource. */
+class FOpenGLEventQuery
+{
+public:
+	/** Initialization constructor. */
+	FOpenGLEventQuery();
+	~FOpenGLEventQuery();
+
+	/** Issues an event for the query to poll. */
+	void IssueEvent();
+
+	/** Waits for the event query to finish. */
+	void WaitForCompletion();
+
+private:
+	UGLsync Sync = {};
+};
+
+class FOpenGLViewport : public FRHIViewport
+{
+public:
+
+	FOpenGLViewport(class FOpenGLDynamicRHI* InOpenGLRHI,void* InWindowHandle,uint32 InSizeX,uint32 InSizeY,bool bInIsFullscreen,EPixelFormat PreferredPixelFormat);
+	~FOpenGLViewport();
+
+	void Resize(uint32 InSizeX,uint32 InSizeY,bool bInIsFullscreen);
+
+	// Accessors.
+	FIntPoint GetSizeXY() const { return FIntPoint(SizeX, SizeY); }
+	FOpenGLTexture* GetBackBuffer() const { return BackBuffer; }
+	bool IsFullscreen( void ) const { return bIsFullscreen; }
+
+	virtual void WaitForFrameEventCompletion() override;
+	virtual void IssueFrameEvent() override;
+
+	virtual void* GetNativeWindow(void** AddParam) const override;
+
+	struct FPlatformOpenGLContext* GetGLContext() const { return OpenGLContext; }
+	FOpenGLDynamicRHI* GetOpenGLRHI() const { return OpenGLRHI; }
+
+	virtual void SetCustomPresent(FRHICustomPresent* InCustomPresent) override
+	{
+		CustomPresent = InCustomPresent;
+	}
+	FRHICustomPresent* GetCustomPresent() const { return CustomPresent.GetReference(); }
+private:
+
+	friend class FOpenGLDynamicRHI;
+
+	FOpenGLDynamicRHI* OpenGLRHI;
+	struct FPlatformOpenGLContext* OpenGLContext;
+	uint32 SizeX;
+	uint32 SizeY;
+	bool bIsFullscreen;
+	EPixelFormat PixelFormat;
+	bool bIsValid;
+	TRefCountPtr<FOpenGLTexture> BackBuffer;
+	TUniquePtr<FOpenGLEventQuery> FrameSyncEvent;
+	FCustomPresentRHIRef CustomPresent;
+};
+
+class FOpenGLGPUFence final : public FRHIGPUFence
+{
+public:
+	FOpenGLGPUFence(FName InName);
+	~FOpenGLGPUFence() override;
+
+	void Clear() override;
+	bool Poll() const override;
+	
+	void WriteInternal();
+private:
+	struct FOpenGLGPUFenceProxy* Proxy;
+};
+
+class FOpenGLStagingBuffer final : public FRHIStagingBuffer
+{
+	friend class FOpenGLDynamicRHI;
+public:
+	FOpenGLStagingBuffer() : FRHIStagingBuffer()
+	{
+		Initialize();
+	}
+
+	~FOpenGLStagingBuffer() override;
+
+	// Locks the shadow of VertexBuffer for read. This will stall the RHI thread.
+	void *Lock(uint32 Offset, uint32 NumBytes) override;
+
+	// Unlocks the shadow. This is an error if it was not locked previously.
+	void Unlock() override;
+
+	uint64 GetGPUSizeBytes() const override { return ShadowSize; }
+private:
+	void Initialize();
+
+	GLuint ShadowBuffer;
+	uint32 ShadowSize;
+	void* Mapping;
+};
+
+template<class T>
+struct TOpenGLResourceTraits
+{
+};
+template<>
+struct TOpenGLResourceTraits<FRHIGPUFence>
+{
+	typedef FOpenGLGPUFence TConcreteType;
+};
+template<>
+struct TOpenGLResourceTraits<FRHIStagingBuffer>
+{
+	typedef FOpenGLStagingBuffer TConcreteType;
+};
+template<>
+struct TOpenGLResourceTraits<FRHIVertexDeclaration>
+{
+	typedef FOpenGLVertexDeclaration TConcreteType;
+};
+template<>
+struct TOpenGLResourceTraits<FRHIVertexShader>
+{
+	typedef FOpenGLVertexShader TConcreteType;
+};
+template<>
+struct TOpenGLResourceTraits<FRHIGeometryShader>
+{
+	typedef FOpenGLGeometryShader TConcreteType;
+};
+template<>
+struct TOpenGLResourceTraits<FRHIPixelShader>
+{
+	typedef FOpenGLPixelShader TConcreteType;
+};
+template<>
+struct TOpenGLResourceTraits<FRHIComputeShader>
+{
+	typedef FOpenGLComputeShader TConcreteType;
+};
+template<>
+struct TOpenGLResourceTraits<FRHIBoundShaderState>
+{
+	typedef FOpenGLBoundShaderState TConcreteType;
+};
+template<>
+struct TOpenGLResourceTraits<FRHIRenderQuery>
+{
+	typedef FOpenGLRenderQuery TConcreteType;
+};
+template<>
+struct TOpenGLResourceTraits<FRHIUniformBuffer>
+{
+	typedef FOpenGLUniformBuffer TConcreteType;
+};
+template<>
+struct TOpenGLResourceTraits<FRHIBuffer>
+{
+	typedef FOpenGLBuffer TConcreteType;
+};
+template<>
+struct TOpenGLResourceTraits<FRHIShaderResourceView>
+{
+	typedef FOpenGLShaderResourceView TConcreteType;
+};
+template<>
+struct TOpenGLResourceTraits<FRHIUnorderedAccessView>
+{
+	typedef FOpenGLUnorderedAccessView TConcreteType;
+};
+
+template<>
+struct TOpenGLResourceTraits<FRHIViewport>
+{
+	typedef FOpenGLViewport TConcreteType;
+};
diff --git a/Engine/Source/Runtime/OpenGLDrv/Private/OpenGLShaders.cpp b/Engine/Source/Runtime/OpenGLDrv/Private/OpenGLShaders.cpp
--- a/Engine/Source/Runtime/OpenGLDrv/Private/OpenGLShaders.cpp
+++ b/Engine/Source/Runtime/OpenGLDrv/Private/OpenGLShaders.cpp
@@ -1061,6 +1061,54 @@ FOpenGLShader::FOpenGLShader(TArrayView<const uint8> Code, const FSHAHash& Hash,
 	// The shader is compiled when we link program
 }
 
+FOpenGLVertexShader::FOpenGLVertexShader(TArrayView<const uint8> Code, const FSHAHash& Hash)
+	: FOpenGLShader(Code, Hash, GL_VERTEX_SHADER)
+{}
+
+void FOpenGLVertexShader::ConditionalyCompile()
+{
+	if (Resource == 0)
+	{
+		Compile(GL_VERTEX_SHADER);
+	}
+}
+
+FOpenGLPixelShader::FOpenGLPixelShader(TArrayView<const uint8> Code, const FSHAHash& Hash)
+	: FOpenGLShader(Code, Hash, GL_FRAGMENT_SHADER)
+{}
+
+void FOpenGLPixelShader::ConditionalyCompile()
+{
+	if (Resource == 0)
+	{
+		Compile(GL_FRAGMENT_SHADER);
+	}
+}
+
+FOpenGLGeometryShader::FOpenGLGeometryShader(TArrayView<const uint8> Code, const FSHAHash& Hash)
+	: FOpenGLShader(Code, Hash, GL_GEOMETRY_SHADER)
+{}
+
+void FOpenGLGeometryShader::ConditionalyCompile()
+{
+	if (Resource == 0)
+	{
+		Compile(GL_GEOMETRY_SHADER);
+	}
+}
+
+FOpenGLComputeShader::FOpenGLComputeShader(TArrayView<const uint8> Code, const FSHAHash& Hash)
+	: FOpenGLShader(Code, Hash, GL_COMPUTE_SHADER)
+{}
+
+void FOpenGLComputeShader::ConditionalyCompile()
+{
+	if (Resource == 0)
+	{
+		Compile(GL_COMPUTE_SHADER);
+	}
+}
+
 FVertexShaderRHIRef FOpenGLDynamicRHI::RHICreateVertexShader(TArrayView<const uint8> Code, const FSHAHash& Hash)
 {
 	return new FOpenGLVertexShader(Code, Hash);
diff --git a/Engine/Source/Runtime/OpenGLDrv/Private/OpenGLShaders.h b/Engine/Source/Runtime/OpenGLDrv/Private/OpenGLShaders.h
new file mode 100644
--- /dev/null
+++ b/Engine/Source/Runtime/OpenGLDrv/Private/OpenGLShaders.h
@@ -0,0 +1,59 @@
+// Copyright Epic Games, Inc. All Rights Reserved.
+
+/*=============================================================================
+	OpenGLShaders.h: OpenGL shader RHI declaration.
+=============================================================================*/
+
+#pragma once
+
+#include "CoreMinimal.h"
+#include "OpenGLDrv.h"
+
+enum EOpenGLShaderTargetPlatform
+{
+	OGLSTP_Unknown,
+	OGLSTP_Desktop,
+	OGLSTP_Android,
+	OGLSTP_iOS,
+};
+
+/**
+  * GL device capabilities for generating GLSL compilable on platform with described capabilities
+  */
+struct FOpenGLShaderDeviceCapabilities
+{
+	EOpenGLShaderTargetPlatform TargetPlatform;
+	EShaderPlatform MaxRHIShaderPlatform;
+	bool bSupportsShaderFramebufferFetch;
+	bool bRequiresARMShaderFramebufferFetchDepthStencilUndef;
+	GLint MaxVaryingVectors;
+	bool bRequiresDisabledEarlyFragmentTests;
+	bool bRequiresReadOnlyBuffersWorkaround;
+};
+
+/**
+  * Gets the GL device capabilities for the current device.
+  *
+  * @param Capabilities [out] The current platform's capabilities on device for shader compiling
+  */
+void OPENGLDRV_API GetCurrentOpenGLShaderDeviceCapabilities(FOpenGLShaderDeviceCapabilities& Capabilities);
+
+/**
+  * Processes the GLSL output of the shader cross compiler to get GLSL that can be compiled on a
+  * platform with the specified capabilities. Works around inconsistencies between OpenGL
+  * implementations, including lack of support for certain extensions and drivers with
+  * non-conformant behavior.
+  *
+  * @param GlslCodeOriginal - [in,out] GLSL output from shader cross compiler to be modified.  Process is destructive; pass in a copy if still need original!
+  * @param ShaderName - [in] Shader name
+  * @param TypeEnum - [in] Type of shader (GL_[VERTEX, FRAGMENT, GEOMETRY, TESS_CONTROL, TESS_EVALUATION]_SHADER)
+  * @param Capabilities - [in] GL Device capabilities
+  * @param GlslCode - [out] Compilable GLSL
+  */
+void OPENGLDRV_API GLSLToDeviceCompatibleGLSL(FAnsiCharArray& GlslCodeOriginal, const FString& ShaderName, GLenum TypeEnum, const FOpenGLShaderDeviceCapabilities& Capabilities, FAnsiCharArray& GlslCode);
+
+
+// make some anon ns functions available to platform extensions
+extern "C" void PE_AppendCString(TArray<ANSICHAR> & Dest, const ANSICHAR * Source);
+extern "C" void PE_ReplaceCString(TArray<ANSICHAR> & Dest, const ANSICHAR * Source, const ANSICHAR * Replacement);
+
diff --git a/Engine/Source/Runtime/OpenGLDrv/Private/OpenGLState.h b/Engine/Source/Runtime/OpenGLDrv/Private/OpenGLState.h
new file mode 100644
--- /dev/null
+++ b/Engine/Source/Runtime/OpenGLDrv/Private/OpenGLState.h
@@ -0,0 +1,560 @@
+// Copyright Epic Games, Inc. All Rights Reserved.
+
+/*=============================================================================
+	OpenGLState.h: OpenGL state definitions.
+=============================================================================*/
+
+#pragma once
+
+#include "CoreTypes.h"
+#include "Misc/AssertionMacros.h"
+#include "HAL/UnrealMemory.h"
+#include "Containers/Array.h"
+#include "Math/Color.h"
+#include "Math/IntRect.h"
+
+#include "RHIDefinitions.h"
+#include "Containers/StaticArray.h"
+#include "RHI.h"
+#include "OpenGLResources.h"
+
+#define ZERO_FILLED_DUMMY_UNIFORM_BUFFER_SIZE 65536
+
+class FRenderTarget;
+
+struct FOpenGLSamplerStateData
+{
+	// These enum is just used to count the number of members in this struct
+	enum EGLSamplerData
+	{
+		EGLSamplerData_WrapS,
+		EGLSamplerData_WrapT,
+		EGLSamplerData_WrapR,
+		EGLSamplerData_LODBias,
+		EGLSamplerData_MagFilter,
+		EGLSamplerData_MinFilter,
+		EGLSamplerData_MaxAniso,
+		EGLSamplerData_CompareMode,
+		EGLSamplerData_CompareFunc,
+		EGLSamplerData_Num,
+	};
+
+	GLint WrapS;
+	GLint WrapT;
+	GLint WrapR;
+	GLint LODBias;
+	GLint MagFilter;
+	GLint MinFilter;
+	GLint MaxAnisotropy;
+	GLint CompareMode;
+	GLint CompareFunc;
+
+	FOpenGLSamplerStateData()
+		: WrapS(GL_REPEAT)
+		, WrapT(GL_REPEAT)
+		, WrapR(GL_REPEAT)
+		, LODBias(0)
+		, MagFilter(GL_NEAREST)
+		, MinFilter(GL_NEAREST)
+		, MaxAnisotropy(1)
+		, CompareMode(GL_NONE)
+		, CompareFunc(GL_ALWAYS)
+	{
+	}
+};
+
+class FOpenGLSamplerState : public FRHISamplerState
+{
+public:
+	GLuint Resource;
+	FOpenGLSamplerStateData Data;
+
+	~FOpenGLSamplerState();
+};
+
+struct FOpenGLRasterizerStateData
+{
+	GLenum FillMode = GL_FILL;
+	GLenum CullMode = GL_NONE;
+	float DepthBias = 0.0f;
+	float SlopeScaleDepthBias = 0.0f;
+	ERasterizerDepthClipMode DepthClipMode = ERasterizerDepthClipMode::DepthClip;
+};
+
+class FOpenGLRasterizerState : public FRHIRasterizerState
+{
+public:
+	virtual bool GetInitializer(FRasterizerStateInitializerRHI& Init) override final;
+	
+	FOpenGLRasterizerStateData Data;
+};
+
+struct FOpenGLDepthStencilStateData
+{
+	bool bZEnable;
+	bool bZWriteEnable;
+	GLenum ZFunc;
+	
+
+	bool bStencilEnable;
+	bool bTwoSidedStencilMode;
+	GLenum StencilFunc;
+	GLenum StencilFail;
+	GLenum StencilZFail;
+	GLenum StencilPass;
+	GLenum CCWStencilFunc;
+	GLenum CCWStencilFail;
+	GLenum CCWStencilZFail;
+	GLenum CCWStencilPass;
+	uint32 StencilReadMask;
+	uint32 StencilWriteMask;
+
+	FOpenGLDepthStencilStateData()
+		: bZEnable(false)
+		, bZWriteEnable(true)
+		, ZFunc(GL_LESS)
+		, bStencilEnable(false)
+		, bTwoSidedStencilMode(false)
+		, StencilFunc(GL_ALWAYS)
+		, StencilFail(GL_KEEP)
+		, StencilZFail(GL_KEEP)
+		, StencilPass(GL_KEEP)
+		, CCWStencilFunc(GL_ALWAYS)
+		, CCWStencilFail(GL_KEEP)
+		, CCWStencilZFail(GL_KEEP)
+		, CCWStencilPass(GL_KEEP)
+		, StencilReadMask(0xFFFFFFFF)
+		, StencilWriteMask(0xFFFFFFFF)
+	{
+	}
+};
+
+class FOpenGLDepthStencilState : public FRHIDepthStencilState
+{
+public:
+	virtual bool GetInitializer(FDepthStencilStateInitializerRHI& Init) override final;
+	
+	FOpenGLDepthStencilStateData Data;
+};
+
+struct FOpenGLBlendStateData
+{
+	struct FRenderTarget
+	{
+		bool bAlphaBlendEnable;
+		GLenum ColorBlendOperation;
+		GLenum ColorSourceBlendFactor;
+		GLenum ColorDestBlendFactor;
+		bool bSeparateAlphaBlendEnable;
+		GLenum AlphaBlendOperation;
+		GLenum AlphaSourceBlendFactor;
+		GLenum AlphaDestBlendFactor;
+		uint32 ColorWriteMaskR : 1;
+		uint32 ColorWriteMaskG : 1;
+		uint32 ColorWriteMaskB : 1;
+		uint32 ColorWriteMaskA : 1;
+	};
+	
+	TStaticArray<FRenderTarget,MaxSimultaneousRenderTargets> RenderTargets;
+
+	bool bUseAlphaToCoverage;
+
+	FOpenGLBlendStateData()
+	{
+		bUseAlphaToCoverage = false;
+		for (int32 i = 0; i < MaxSimultaneousRenderTargets; ++i)
+		{
+			FRenderTarget& Target = RenderTargets[i];
+			Target.bAlphaBlendEnable = false;
+			Target.ColorBlendOperation = GL_NONE;
+			Target.ColorSourceBlendFactor = GL_NONE;
+			Target.ColorDestBlendFactor = GL_NONE;
+			Target.bSeparateAlphaBlendEnable = false;
+			Target.AlphaBlendOperation = GL_NONE;
+			Target.AlphaSourceBlendFactor = GL_NONE;
+			Target.AlphaDestBlendFactor = GL_NONE;
+			Target.ColorWriteMaskR = false;
+			Target.ColorWriteMaskG = false;
+			Target.ColorWriteMaskB = false;
+			Target.ColorWriteMaskA = false;
+		}
+	}
+};
+
+class FOpenGLBlendState : public FRHIBlendState
+{
+	FBlendStateInitializerRHI RHIInitializer;
+public:
+	FOpenGLBlendState(const FBlendStateInitializerRHI& Initializer) : RHIInitializer(Initializer) {}
+	virtual bool GetInitializer(FBlendStateInitializerRHI& Init) override final
+	{ 
+		Init = RHIInitializer; 
+		return true;
+	}
+
+	FOpenGLBlendStateData Data;
+};
+
+struct FTextureStage
+{
+	class FOpenGLTexture* Texture;
+	class FOpenGLShaderResourceView* SRV;
+	GLenum Target;
+	GLuint Resource;
+	int32 LimitMip;
+	bool bHasMips;
+	int32 NumMips;
+
+	FTextureStage()
+	:	Texture(NULL)
+	,	SRV(NULL)
+	,	Target(GL_NONE)
+	,	Resource(0)
+	,	LimitMip(-1)
+	,	bHasMips(false)
+	,	NumMips(0)
+	{
+	}
+};
+
+struct FUAVStage
+{
+	GLenum Format;
+	GLuint Resource;
+	GLenum Access;
+	GLint Layer;
+	bool bLayered;
+	
+	FUAVStage()
+	:	Format(GL_NONE)
+	,	Resource(0)
+	,	Access(GL_READ_WRITE)
+	,	Layer(0)
+	,	bLayered(false)
+	{
+	}
+};
+#define FOpenGLCachedAttr_Invalid (void*)(UPTRINT)0xFFFFFFFF
+#define FOpenGLCachedAttr_SingleVertex (void*)(UPTRINT)0xFFFFFFFE
+
+struct FOpenGLCachedAttr
+{
+	GLuint Size;
+	GLenum Type;
+	GLuint StreamOffset;
+	GLuint StreamIndex;
+	GLboolean bNormalized;
+	GLboolean bShouldConvertToFloat;
+
+	FOpenGLCachedAttr() : 
+		Size(),
+		Type(0), 
+		StreamOffset(),
+		StreamIndex(0xFFFFFFFF),
+		bNormalized(),
+		bShouldConvertToFloat()
+	{
+	}
+};
+
+struct FOpenGLStream
+{
+	GLuint VertexBufferResource;
+	uint32 Stride;
+	uint32 Offset;
+	uint32 Divisor;
+	
+	FOpenGLStream()
+		: VertexBufferResource(0)
+		, Stride(0)
+		, Offset(0)
+		, Divisor(0)
+	{}
+};
+
+#define NUM_OPENGL_VERTEX_STREAMS 16
+
+struct FOpenGLCommonState
+{
+	TArray<FTextureStage>	Textures;
+	TArray<FOpenGLSamplerState*>	SamplerStates;
+	TArray<FUAVStage>		UAVs;
+
+	FOpenGLCommonState()
+	{}
+
+	virtual ~FOpenGLCommonState()
+	{
+		FOpenGLCommonState::CleanupResources();
+	}
+
+	// NumCombinedTextures must be greater than or equal to FOpenGL::GetMaxCombinedTextureImageUnits()
+	// NumCombinedUAVUnits must be greater than or equal to FOpenGL::GetMaxCombinedUAVUnits()
+	virtual void InitializeResources(int32 NumCombinedTextures, int32 NumCombinedUAVUnits)
+	{
+		check(NumCombinedTextures >= FOpenGL::GetMaxCombinedTextureImageUnits());
+		check(NumCombinedUAVUnits >= FOpenGL::GetMaxCombinedUAVUnits());
+		check(Textures.IsEmpty() && SamplerStates.IsEmpty() && UAVs.Num() == 0);
+		Textures.SetNum(NumCombinedTextures);
+		SamplerStates.SetNumZeroed(NumCombinedTextures);
+		
+		UAVs.Reserve(NumCombinedUAVUnits);
+		UAVs.AddDefaulted(NumCombinedUAVUnits);
+	}
+
+	virtual void CleanupResources()
+	{
+		SamplerStates.Empty();
+		Textures.Empty();
+		UAVs.Empty();
+	}
+};
+
+struct FOpenGLContextState final : public FOpenGLCommonState
+{
+	FOpenGLRasterizerStateData		RasterizerState;
+	FOpenGLDepthStencilStateData	DepthStencilState;
+	uint32							StencilRef;
+	FOpenGLBlendStateData			BlendState;
+	GLuint							Framebuffer;
+	uint32							RenderTargetWidth;
+	uint32							RenderTargetHeight;
+	GLuint							OcclusionQuery;
+	GLuint							Program;
+	GLuint 							UniformBuffers[CrossCompiler::NUM_SHADER_STAGES*OGL_MAX_UNIFORM_BUFFER_BINDINGS];
+	GLuint 							UniformBufferOffsets[CrossCompiler::NUM_SHADER_STAGES*OGL_MAX_UNIFORM_BUFFER_BINDINGS];
+	TArray<FOpenGLSamplerState*>	CachedSamplerStates;
+	GLenum							ActiveTexture;
+	bool							bScissorEnabled;
+	FIntRect						Scissor;
+	FIntRect						Viewport;
+	float							DepthMinZ;
+	float							DepthMaxZ;
+	GLuint							ArrayBufferBound;
+	GLuint							ElementArrayBufferBound;
+	GLuint							StorageBufferBound;
+	GLuint							PixelUnpackBufferBound;
+	GLuint							UniformBufferBound;
+	FLinearColor					ClearColor;
+	uint16							ClearStencil;
+	float							ClearDepth;
+	int32							FirstNonzeroRenderTarget;
+	bool							bAlphaToCoverageEnabled;
+	
+	FOpenGLVertexDeclaration*		VertexDecl;
+	FOpenGLCachedAttr				VertexAttrs[NUM_OPENGL_VERTEX_STREAMS];
+	FOpenGLStream					VertexStreams[NUM_OPENGL_VERTEX_STREAMS];
+		
+	uint32							ActiveStreamMask;
+	uint32							VertexAttrs_EnabledBits;
+	FORCEINLINE bool GetVertexAttrEnabled(int32 Index) const
+	{
+		static_assert(NUM_OPENGL_VERTEX_STREAMS <= sizeof(VertexAttrs_EnabledBits) * 8, "Not enough bits in VertexAttrs_EnabledBits to store NUM_OPENGL_VERTEX_STREAMS");
+		return !!(VertexAttrs_EnabledBits & (1 << Index));
+	}
+	FORCEINLINE void SetVertexAttrEnabled(int32 Index, bool bEnable)
+	{
+		if (bEnable)
+		{
+			VertexAttrs_EnabledBits |= (1 << Index);
+		}
+		else
+		{
+			VertexAttrs_EnabledBits &= ~(1 << Index);
+		}
+	}
+
+	uint32 ActiveUAVMask;
+
+	FOpenGLContextState()
+	:	StencilRef(0)
+	,	Framebuffer(0)
+	,	Program(0)
+	,	ActiveTexture(GL_TEXTURE0)
+	,	bScissorEnabled(false)
+	,	DepthMinZ(0.0f)
+	,	DepthMaxZ(1.0f)
+	,	ArrayBufferBound(0)
+	,	ElementArrayBufferBound(0)
+	,	StorageBufferBound(0)
+	,	PixelUnpackBufferBound(0)
+	,	UniformBufferBound(0)
+	,	ClearColor(-1, -1, -1, -1)
+	,	ClearStencil(0xFFFF)
+	,	ClearDepth(-1.0f)
+	,	FirstNonzeroRenderTarget(0)
+	,	bAlphaToCoverageEnabled(false)
+	,	VertexDecl(0)
+	,   VertexAttrs()
+	,	VertexStreams()
+	,	ActiveStreamMask(0)
+	,	VertexAttrs_EnabledBits(0)
+	,	ActiveUAVMask(0)
+	{
+		Scissor.Min.X = Scissor.Min.Y = Scissor.Max.X = Scissor.Max.Y = 0;
+		Viewport.Min.X = Viewport.Min.Y = Viewport.Max.X = Viewport.Max.Y = 0;
+		FMemory::Memzero(UniformBuffers, sizeof(UniformBuffers));
+		FMemory::Memzero(UniformBufferOffsets, sizeof(UniformBufferOffsets));
+	}
+
+	virtual void InitializeResources(int32 NumCombinedTextures, int32 NumCombinedUAVUnits) override
+	{
+		FOpenGLCommonState::InitializeResources(NumCombinedTextures, NumCombinedUAVUnits);
+		CachedSamplerStates.Empty(NumCombinedTextures);
+		CachedSamplerStates.AddZeroed(NumCombinedTextures);
+
+		checkf(NumCombinedUAVUnits <= sizeof(ActiveUAVMask) * 8, TEXT("Not enough bits in ActiveUAVMask to store %d UAV units"), NumCombinedUAVUnits);
+	}
+
+	virtual void CleanupResources() override
+	{
+		CachedSamplerStates.Empty();
+		FOpenGLCommonState::CleanupResources();
+	}
+};
+
+struct FOpenGLRHIState final : public FOpenGLCommonState
+{
+	FOpenGLRasterizerStateData		RasterizerState;
+	FOpenGLDepthStencilStateData	DepthStencilState;
+	uint32							StencilRef;
+	FOpenGLBlendStateData			BlendState;
+	GLuint							Framebuffer;
+	bool							bScissorEnabled;
+	FIntRect						Scissor;
+	FIntRect						Viewport;
+	float							DepthMinZ;
+	float							DepthMaxZ;
+	GLuint							ZeroFilledDummyUniformBuffer;
+	uint32							RenderTargetWidth;
+	uint32							RenderTargetHeight;
+	GLuint							RunningOcclusionQuery;
+	bool							bAlphaToCoverageEnabled;
+
+	// Pending framebuffer setup
+	int32							NumRenderingSamples;// Only used with GL_EXT_multisampled_render_to_texture
+	int32							FirstNonzeroRenderTarget;
+	FOpenGLTexture*					RenderTargets[MaxSimultaneousRenderTargets];
+	uint32							RenderTargetMipmapLevels[MaxSimultaneousRenderTargets];
+	uint32							RenderTargetArrayIndex[MaxSimultaneousRenderTargets];
+	FOpenGLTexture*					DepthStencil;
+	ERenderTargetStoreAction		StencilStoreAction;
+	uint32							DepthTargetWidth;
+	uint32							DepthTargetHeight;
+	bool							bFramebufferSetupInvalid;
+
+	// Information about pending BeginDraw[Indexed]PrimitiveUP calls.
+	FOpenGLStream					DynamicVertexStream;
+	uint32							NumVertices;
+	uint32							PrimitiveType;
+	uint32							NumPrimitives;
+	uint32							MinVertexIndex;
+	uint32							IndexDataStride;
+
+	FOpenGLStream					Streams[NUM_OPENGL_VERTEX_STREAMS];
+
+	// we null this when the we dirty PackedGlobalUniformDirty. Thus we can skip all of CommitNonComputeShaderConstants if it matches the current program
+	FOpenGLLinkedProgram* LinkedProgramAndDirtyFlag;
+	FOpenGLShaderParameterCache*	ShaderParameters;
+
+	TRefCountPtr<FOpenGLBoundShaderState>	BoundShaderState;
+	FComputeShaderRHIRef					CurrentComputeShader;	
+
+	/** The RHI does not allow more than 14 constant buffers per shader stage due to D3D11 limits. */
+	enum { MAX_UNIFORM_BUFFERS_PER_SHADER_STAGE = 14 };
+
+	/** Track the currently bound uniform buffers. */
+	FRHIUniformBuffer* BoundUniformBuffers[SF_NumStandardFrequencies][MAX_UNIFORM_BUFFERS_PER_SHADER_STAGE];
+	uint32 BoundUniformBuffersDynamicOffset[SF_NumStandardFrequencies][MAX_UNIFORM_BUFFERS_PER_SHADER_STAGE];
+
+	/** Array to track if any real (not emulated) uniform buffers have been bound since the last draw call */
+	bool bAnyDirtyRealUniformBuffers[SF_NumStandardFrequencies];
+	/** Bit array to track which uniform buffers have changed since the last draw call. */
+	bool bAnyDirtyGraphicsUniformBuffers;
+	uint16 DirtyUniformBuffers[SF_NumStandardFrequencies];
+
+	// Used for if(!FOpenGL::SupportsFastBufferData())
+	uint32 UpVertexBufferBytes;
+	uint32 UpIndexBufferBytes;
+	uint32 UpStride;
+	void* UpVertexBuffer;
+	void* UpIndexBuffer;
+
+	FOpenGLRHIState()
+	:	StencilRef(0)
+	,	Framebuffer(0)
+	,	bScissorEnabled(false)
+	,	DepthMinZ(0.0f)
+	,	DepthMaxZ(1.0f)
+	,	ZeroFilledDummyUniformBuffer(0)
+	,	RenderTargetWidth(0)
+	,	RenderTargetHeight(0)
+	,	RunningOcclusionQuery(0)
+	,	bAlphaToCoverageEnabled(false)
+	,	NumRenderingSamples(1)
+	,	FirstNonzeroRenderTarget(-1)
+	,	DepthStencil(0)
+	,	StencilStoreAction(ERenderTargetStoreAction::ENoAction)
+	,	DepthTargetWidth(0)
+	,	DepthTargetHeight(0)
+	,	bFramebufferSetupInvalid(true)
+	,	NumVertices(0)
+	,	PrimitiveType(0)
+	,	NumPrimitives(0)
+	,	MinVertexIndex(0)
+	,	IndexDataStride(0)
+	,	LinkedProgramAndDirtyFlag(nullptr)
+	,	ShaderParameters(NULL)
+	,	BoundShaderState(NULL)
+	,	CurrentComputeShader(NULL)
+	,	UpVertexBufferBytes(0)
+	,   UpIndexBufferBytes(0)
+	,	UpVertexBuffer(0)
+	,	UpIndexBuffer(0)
+	{
+		Scissor.Min.X = Scissor.Min.Y = Scissor.Max.X = Scissor.Max.Y = 0;
+		Viewport.Min.X = Viewport.Min.Y = Viewport.Max.X = Viewport.Max.Y = 0;
+		FMemory::Memset( RenderTargets, 0, sizeof(RenderTargets) );	// setting all to 0 at start
+		FMemory::Memset( RenderTargetMipmapLevels, 0, sizeof(RenderTargetMipmapLevels) );	// setting all to 0 at start
+		FMemory::Memset( RenderTargetArrayIndex, 0, sizeof(RenderTargetArrayIndex) );	// setting all to 0 at start
+		FMemory::Memset(BoundUniformBuffers, 0, sizeof(BoundUniformBuffers));
+		FMemory::Memset(BoundUniformBuffersDynamicOffset, 0u, sizeof(BoundUniformBuffersDynamicOffset));
+	}
+
+	~FOpenGLRHIState()
+	{
+		CleanupResources();
+	}
+
+	virtual void InitializeResources(int32 NumCombinedTextures, int32 NumComputeUAVUnits) override;
+
+	virtual void CleanupResources() override
+	{
+		delete [] ShaderParameters;
+		ShaderParameters = NULL;
+		FMemory::Memset(BoundUniformBuffers, 0, sizeof(BoundUniformBuffers));
+		FMemory::Memset(BoundUniformBuffersDynamicOffset, 0u, sizeof(BoundUniformBuffersDynamicOffset));
+		FOpenGLCommonState::CleanupResources();
+	}
+};
+
+template<>
+struct TOpenGLResourceTraits<FRHISamplerState>
+{
+	typedef FOpenGLSamplerState TConcreteType;
+};
+template<>
+struct TOpenGLResourceTraits<FRHIRasterizerState>
+{
+	typedef FOpenGLRasterizerState TConcreteType;
+};
+template<>
+struct TOpenGLResourceTraits<FRHIDepthStencilState>
+{
+	typedef FOpenGLDepthStencilState TConcreteType;
+};
+template<>
+struct TOpenGLResourceTraits<FRHIBlendState>
+{
+	typedef FOpenGLBlendState TConcreteType;
+};
diff --git a/Engine/Source/Runtime/OpenGLDrv/Public/IOpenGLDynamicRHI.h b/Engine/Source/Runtime/OpenGLDrv/Public/IOpenGLDynamicRHI.h
--- a/Engine/Source/Runtime/OpenGLDrv/Public/IOpenGLDynamicRHI.h
+++ b/Engine/Source/Runtime/OpenGLDrv/Public/IOpenGLDynamicRHI.h
@@ -3,35 +3,7 @@
 
 #include "RHI.h"
 
-#if PLATFORM_WINDOWS
-	#include "Windows/AllowWindowsPlatformTypes.h"
-	THIRD_PARTY_INCLUDES_START
-		#include <GL/glcorearb.h>
-		#include <GL/glext.h>
-		#include <GL/wglext.h>
-	THIRD_PARTY_INCLUDES_END
-	#include "Windows/HideWindowsPlatformTypes.h"
-#elif PLATFORM_LINUX
-	THIRD_PARTY_INCLUDES_START
-		#include <GL/glcorearb.h>
-		#include <GL/glext.h>
-	THIRD_PARTY_INCLUDES_END
-#elif PLATFORM_ANDROID
-	#include "Android/AndroidPlatform.h"
-	THIRD_PARTY_INCLUDES_START
-		#include <EGL/egl.h>
-		#include <EGL/eglext.h>
-		#include <GLES3/gl31.h>
-	THIRD_PARTY_INCLUDES_END
-
-	#ifndef USE_ANDROID_EGL_NO_ERROR_CONTEXT
-		#if UE_BUILD_SHIPPING
-			#define USE_ANDROID_EGL_NO_ERROR_CONTEXT 1
-		#else
-			#define USE_ANDROID_EGL_NO_ERROR_CONTEXT 0
-		#endif
-	#endif
-#endif
+#include "OpenGLThirdParty.h"
 
 struct IOpenGLDynamicRHI : public FDynamicRHIPSOFallback
 {
diff --git a/Engine/Source/Runtime/OpenGLDrv/Public/OpenGL.h b/Engine/Source/Runtime/OpenGLDrv/Public/OpenGL.h
deleted file mode 100644
diff --git a/Engine/Source/Runtime/OpenGLDrv/Public/OpenGL3.h b/Engine/Source/Runtime/OpenGLDrv/Public/OpenGL3.h
deleted file mode 100644
diff --git a/Engine/Source/Runtime/OpenGLDrv/Public/OpenGL4.h b/Engine/Source/Runtime/OpenGLDrv/Public/OpenGL4.h
deleted file mode 100644
diff --git a/Engine/Source/Runtime/OpenGLDrv/Public/OpenGLDrv.h b/Engine/Source/Runtime/OpenGLDrv/Public/OpenGLDrv.h
deleted file mode 100644
diff --git a/Engine/Source/Runtime/OpenGLDrv/Public/OpenGLES.h b/Engine/Source/Runtime/OpenGLDrv/Public/OpenGLES.h
deleted file mode 100644
diff --git a/Engine/Source/Runtime/OpenGLDrv/Public/OpenGLResources.h b/Engine/Source/Runtime/OpenGLDrv/Public/OpenGLResources.h
deleted file mode 100644
diff --git a/Engine/Source/Runtime/OpenGLDrv/Public/OpenGLShaderResources.h b/Engine/Source/Runtime/OpenGLDrv/Public/OpenGLShaderResources.h
--- a/Engine/Source/Runtime/OpenGLDrv/Public/OpenGLShaderResources.h
+++ b/Engine/Source/Runtime/OpenGLDrv/Public/OpenGLShaderResources.h
@@ -8,23 +8,15 @@
 
 // HEADER_UNIT_SKIP - Not included directly
 
-#include "HAL/UnrealMemory.h"
-#include "Templates/UnrealTemplate.h"
 #include "Containers/Array.h"
-#include "Misc/Crc.h"
-#include "Containers/UnrealString.h"
-#include "Containers/Map.h"
-#include "Misc/SecureHash.h"
-#include "HAL/IConsoleManager.h"
-#include "RHI.h"
-#include "ShaderCore.h"
 #include "CrossCompilerCommon.h"
-#include "ShaderCodeLibrary.h"
-#include "Async/AsyncFileHandle.h"
-#include "ShaderPipelineCache.h"
+#include "OpenGLThirdParty.h"
 
 class FOpenGLLinkedProgram;
 
+/** Set to 1 to enable shader debugging which e.g. keeps the GLSL source as members of TOpenGLShader*/
+#define DEBUG_GL_SHADERS (UE_BUILD_DEBUG || UE_EDITOR)
+
 /**
  * Shader related constants.
  */
@@ -327,17 +319,9 @@ class FOpenGLVertexShader : public FRHIVertexShader, public FOpenGLShader
 public:
 	static constexpr EShaderFrequency Frequency = SF_Vertex;
 
-	FOpenGLVertexShader(TArrayView<const uint8> Code, const FSHAHash& Hash)
-		: FOpenGLShader(Code, Hash, GL_VERTEX_SHADER)
-	{}
+	FOpenGLVertexShader(TArrayView<const uint8> Code, const FSHAHash& Hash);
 
-	void ConditionalyCompile()
-	{
-		if (Resource == 0)
-		{
-			Compile(GL_VERTEX_SHADER);
-		}
-	}
+	void ConditionalyCompile();
 };
 
 class FOpenGLPixelShader : public FRHIPixelShader, public FOpenGLShader
@@ -345,17 +329,9 @@ class FOpenGLPixelShader : public FRHIPixelShader, public FOpenGLShader
 public:
 	static constexpr EShaderFrequency Frequency = SF_Pixel;
 
-	FOpenGLPixelShader(TArrayView<const uint8> Code, const FSHAHash& Hash)
-		: FOpenGLShader(Code, Hash, GL_FRAGMENT_SHADER)
-	{}
+	FOpenGLPixelShader(TArrayView<const uint8> Code, const FSHAHash& Hash);
 
-	void ConditionalyCompile()
-	{
-		if (Resource == 0)
-		{
-			Compile(GL_FRAGMENT_SHADER);
-		}
-	}
+	void ConditionalyCompile();
 };
 
 class FOpenGLGeometryShader : public FRHIGeometryShader, public FOpenGLShader
@@ -363,17 +339,9 @@ class FOpenGLGeometryShader : public FRHIGeometryShader, public FOpenGLShader
 public:
 	static constexpr EShaderFrequency Frequency = SF_Geometry;
 
-	FOpenGLGeometryShader(TArrayView<const uint8> Code, const FSHAHash& Hash)
-		: FOpenGLShader(Code, Hash, GL_GEOMETRY_SHADER)
-	{}
+	FOpenGLGeometryShader(TArrayView<const uint8> Code, const FSHAHash& Hash);
 
-	void ConditionalyCompile()
-	{
-		if (Resource == 0)
-		{
-			Compile(GL_GEOMETRY_SHADER);
-		}
-	}
+	void ConditionalyCompile();
 };
 
 class FOpenGLComputeShader : public FRHIComputeShader, public FOpenGLShader
@@ -381,17 +349,9 @@ class FOpenGLComputeShader : public FRHIComputeShader, public FOpenGLShader
 public:
 	static constexpr EShaderFrequency Frequency = SF_Compute;
 
-	FOpenGLComputeShader(TArrayView<const uint8> Code, const FSHAHash& Hash)
-		: FOpenGLShader(Code, Hash, GL_COMPUTE_SHADER)
-	{}
+	FOpenGLComputeShader(TArrayView<const uint8> Code, const FSHAHash& Hash);
 
-	void ConditionalyCompile()
-	{
-		if (Resource == 0)
-		{
-			Compile(GL_COMPUTE_SHADER);
-		}
-	}
+	void ConditionalyCompile();
 
 	bool NeedsTextureStage(int32 TextureStageIndex);
 	int32 MaxTextureStageUsed();
diff --git a/Engine/Source/Runtime/OpenGLDrv/Public/OpenGLShaders.h b/Engine/Source/Runtime/OpenGLDrv/Public/OpenGLShaders.h
deleted file mode 100644
diff --git a/Engine/Source/Runtime/OpenGLDrv/Public/OpenGLState.h b/Engine/Source/Runtime/OpenGLDrv/Public/OpenGLState.h
deleted file mode 100644
diff --git a/Engine/Source/Runtime/OpenGLDrv/Public/OpenGLThirdParty.h b/Engine/Source/Runtime/OpenGLDrv/Public/OpenGLThirdParty.h
new file mode 100644
--- /dev/null
+++ b/Engine/Source/Runtime/OpenGLDrv/Public/OpenGLThirdParty.h
@@ -0,0 +1,27 @@
+// Copyright Epic Games, Inc. All Rights Reserved.
+
+#pragma once
+
+#if PLATFORM_WINDOWS
+	#include "Windows/AllowWindowsPlatformTypes.h"
+	THIRD_PARTY_INCLUDES_START
+		#include <GL/glcorearb.h>
+		#include <GL/glext.h>
+		#include <GL/wglext.h>
+	THIRD_PARTY_INCLUDES_END
+	#include "Windows/HideWindowsPlatformTypes.h"
+#elif PLATFORM_LINUX
+	THIRD_PARTY_INCLUDES_START
+		#include <GL/glcorearb.h>
+		#include <GL/glext.h>
+	THIRD_PARTY_INCLUDES_END
+#elif PLATFORM_ANDROID
+	#include "Android/AndroidPlatform.h"
+	THIRD_PARTY_INCLUDES_START
+		#include <EGL/egl.h>
+		#include <EGL/eglext.h>
+		#include <GLES3/gl31.h>
+	THIRD_PARTY_INCLUDES_END
+#endif
+
+using UGLsync = GLsync;
diff --git a/Engine/Source/Runtime/OpenGLDrv/Public/OpenGLUtil.h b/Engine/Source/Runtime/OpenGLDrv/Public/OpenGLUtil.h
--- a/Engine/Source/Runtime/OpenGLDrv/Public/OpenGLUtil.h
+++ b/Engine/Source/Runtime/OpenGLDrv/Public/OpenGLUtil.h
@@ -5,7 +5,10 @@
 =============================================================================*/
 
 #pragma once
+
 #include "RHICommandList.h"
+#include "OpenGLThirdParty.h"
+#include "OpenGLShaderResources.h"
 
 /** Set to 1 to enable the VERIFY_GL macros which call glGetError */
 #define ENABLE_VERIFY_GL (0 & DO_CHECK)
@@ -31,9 +34,6 @@ int32 GetOGLDebugOutputLevel();
 /** Set to 1 to additinally dump uniform buffer layout at shader link time, this assumes ENABLE_UNIFORM_BUFFER_LAYOUT_VERIFICATION == 1 */
 #define ENABLE_UNIFORM_BUFFER_LAYOUT_DUMP 0
 
-/** Set to 1 to enable shader debugging which e.g. keeps the GLSL source as members of TOpenGLShader*/
-#define DEBUG_GL_SHADERS (UE_BUILD_DEBUG || UE_EDITOR)
-
 /** Set to 1 to enable calls to place event markers into the OpenGL stream
     this is purposefully not considered for OPENGL_PERFORMANCE_DATA_INVALID, 
 	since there is an additional cvar OpenGLConsoleVariables::bEnableARBDebug*/
