OpenGLDrv Header dependency cleanup
- Fixing includes in OpenGLDrv so that all headers can be included independently.
- Focused on moving per-platform and per-GL Version code into more specific headers.
- Using COMPILED_PLATFORM_HEADER more instead of per-platform checks.

diff --git a/Engine/Source/Runtime/OpenGLDrv/OpenGLDrv.Build.cs b/Engine/Source/Runtime/OpenGLDrv/OpenGLDrv.Build.cs
--- a/Engine/Source/Runtime/OpenGLDrv/OpenGLDrv.Build.cs
+++ b/Engine/Source/Runtime/OpenGLDrv/OpenGLDrv.Build.cs
@@ -8,8 +8,6 @@ public class OpenGLDrv : ModuleRules
 {
 	public OpenGLDrv(ReadOnlyTargetRules Target) : base(Target)
 	{
-		IWYUSupport = IWYUSupport.None;
-
 		PublicDependencyModuleNames.AddRange(new string[] {
 			"Core",
 			"RHI",
@@ -43,8 +41,8 @@ public class OpenGLDrv : ModuleRules
 
 		if (Target.Platform == UnrealTargetPlatform.Android)
         {
-            // for Swappy
-            PublicDefinitions.Add("USE_ANDROID_OPENGL_SWAPPY=1");
+			// for Swappy
+			PrivateDefinitions.Add("USE_ANDROID_OPENGL_SWAPPY=1");
 
             PrivateDependencyModuleNames.Add("GoogleGameSDK");
 			PrivateIncludePathModuleNames.Add("Launch");
@@ -58,6 +56,6 @@ public class OpenGLDrv : ModuleRules
 			PrecompileForTargets = PrecompileTargetsType.None;
 		}
 
-		PublicDefinitions.Add(Target.Platform == UnrealTargetPlatform.Android ? "USE_ANDROID_OPENGL=1" : "USE_ANDROID_OPENGL=0");
+		PrivateDefinitions.Add(Target.Platform == UnrealTargetPlatform.Android ? "USE_ANDROID_OPENGL=1" : "USE_ANDROID_OPENGL=0");
 	}
 }
diff --git a/Engine/Source/Runtime/OpenGLDrv/Private/Android/AndroidEGL.cpp b/Engine/Source/Runtime/OpenGLDrv/Private/Android/AndroidEGL.cpp
--- a/Engine/Source/Runtime/OpenGLDrv/Private/Android/AndroidEGL.cpp
+++ b/Engine/Source/Runtime/OpenGLDrv/Private/Android/AndroidEGL.cpp
@@ -1,6 +1,6 @@
 // Copyright Epic Games, Inc. All Rights Reserved.
 
-#include "Android/AndroidPlatform.h"
+#include "AndroidEGL.h"
 
 #if USE_ANDROID_OPENGL
 
diff --git a/Engine/Source/Runtime/OpenGLDrv/Private/Android/AndroidOpenGL.cpp b/Engine/Source/Runtime/OpenGLDrv/Private/Android/AndroidOpenGL.cpp
--- a/Engine/Source/Runtime/OpenGLDrv/Private/Android/AndroidOpenGL.cpp
+++ b/Engine/Source/Runtime/OpenGLDrv/Private/Android/AndroidOpenGL.cpp
@@ -1,12 +1,9 @@
 // Copyright Epic Games, Inc. All Rights Reserved.
 
-#include "CoreMinimal.h"
-#include "Android/AndroidPlatform.h"
+#include "AndroidOpenGL.h"
 
 #if USE_ANDROID_OPENGL
 
-#include "OpenGLDrvPrivate.h"
-#include "AndroidOpenGL.h"
 #include "OpenGLDrvPrivate.h"
 #include "OpenGLES.h"
 #include "Android/AndroidWindow.h"
@@ -15,19 +12,7 @@
 #include "Android/AndroidPlatformFramePacer.h"
 #include "Android/AndroidJNI.h"
 #include "GenericPlatform/GenericPlatformCrashContext.h"
-
-PFNeglPresentationTimeANDROID eglPresentationTimeANDROID_p = NULL;
-PFNeglGetNextFrameIdANDROID eglGetNextFrameIdANDROID_p = NULL;
-PFNeglGetCompositorTimingANDROID eglGetCompositorTimingANDROID_p = NULL;
-PFNeglGetFrameTimestampsANDROID eglGetFrameTimestampsANDROID_p = NULL;
-PFNeglQueryTimestampSupportedANDROID eglQueryTimestampSupportedANDROID_p = NULL;
-PFNeglQueryTimestampSupportedANDROID eglGetCompositorTimingSupportedANDROID_p = NULL;
-PFNeglQueryTimestampSupportedANDROID eglGetFrameTimestampsSupportedANDROID_p = NULL;
-
-namespace GLFuncPointers
-{
-	PFNGLFRAMEBUFFERFETCHBARRIERQCOMPROC glFramebufferFetchBarrierQCOM = NULL;
-}
+#include "Misc/ConfigCacheIni.h"
 
 int32 FAndroidOpenGL::GLMajorVerion = 0;
 int32 FAndroidOpenGL::GLMinorVersion = 0;
@@ -1071,53 +1056,8 @@ void FAndroidOpenGL::ProcessExtensions(const FString& ExtensionsString)
 	}
 }
 
-FString FAndroidMisc::GetGPUFamily()
-{
-	return FAndroidGPUInfo::Get().GetGPUFamily();
-}
-
-FString FAndroidMisc::GetGLVersion()
-{
-	return FAndroidGPUInfo::Get().GLVersion;
-}
-
-bool FAndroidMisc::SupportsFloatingPointRenderTargets()
-{
-	return FAndroidGPUInfo::Get().bSupportsFloatingPointRenderTargets;
-}
-
-bool FAndroidMisc::SupportsShaderFramebufferFetch()
-{
-	return FAndroidGPUInfo::Get().bSupportsFrameBufferFetch;
-}
-
-bool FAndroidMisc::SupportsES30()
-{
-	return true;
-}
-
-void FAndroidMisc::GetValidTargetPlatforms(TArray<FString>& TargetPlatformNames)
-{
-	TargetPlatformNames = FAndroidGPUInfo::Get().TargetPlatformNames;
-}
-
-void FAndroidAppEntry::PlatformInit()
-{
-	// Try to create an ES3.2 EGL here for gpu queries and don't have to recreate the GL context.
-	AndroidEGL::GetInstance()->Init(AndroidEGL::AV_OpenGLES, 3, 2);
-}
-
-void FAndroidAppEntry::ReleaseEGL()
-{
-	AndroidEGL* EGL = AndroidEGL::GetInstance();
-	if (EGL->IsInitialized())
-	{
-		EGL->DestroyBackBuffer();
-		EGL->Terminate();
-	}
-}
-
 static bool GRemoteCompileServicesActive = false;
+extern bool AreAndroidOpenGLRemoteCompileServicesAvailable();
 
 bool AreAndroidOpenGLRemoteCompileServicesActive()
 {
diff --git a/Engine/Source/Runtime/OpenGLDrv/Private/Android/AndroidOpenGL.h b/Engine/Source/Runtime/OpenGLDrv/Private/Android/AndroidOpenGL.h
--- a/Engine/Source/Runtime/OpenGLDrv/Private/Android/AndroidOpenGL.h
+++ b/Engine/Source/Runtime/OpenGLDrv/Private/Android/AndroidOpenGL.h
@@ -5,73 +5,12 @@
 =============================================================================*/
 #pragma once
 
-#include "CoreMinimal.h"
-#include "Misc/ConfigCacheIni.h"
-#include "RenderingThread.h"
-#include "RHI.h"
-
-#if PLATFORM_ANDROID
+#include "OpenGLThirdParty.h"
+#include "AndroidOpenGLPlatform.h"
 
 #include "AndroidEGL.h"
-
-#include <EGL/eglext.h>
-#include <EGL/eglplatform.h>
-#include <GLES2/gl2ext.h>
-	
-typedef GLsync			UGLsync;
-#define GLdouble		GLfloat
-typedef khronos_int64_t GLint64;
-typedef khronos_uint64_t GLuint64;
-#define GL_CLAMP		GL_CLAMP_TO_EDGE
-
 #include "OpenGLES.h"
 
-typedef khronos_stime_nanoseconds_t EGLnsecsANDROID;
-
-typedef GLboolean(GL_APIENTRYP PFNeglPresentationTimeANDROID) (EGLDisplay dpy, EGLSurface surface, EGLnsecsANDROID time);
-typedef GLboolean(GL_APIENTRYP PFNeglGetNextFrameIdANDROID) (EGLDisplay dpy, EGLSurface surface, EGLuint64KHR *frameId);
-typedef GLboolean(GL_APIENTRYP PFNeglGetCompositorTimingANDROID) (EGLDisplay dpy, EGLSurface surface, EGLint numTimestamps, const EGLint *names, EGLnsecsANDROID *values);
-typedef GLboolean(GL_APIENTRYP PFNeglGetFrameTimestampsANDROID) (EGLDisplay dpy, EGLSurface surface, EGLuint64KHR frameId, EGLint numTimestamps, const EGLint *timestamps, EGLnsecsANDROID *values);
-typedef GLboolean(GL_APIENTRYP PFNeglQueryTimestampSupportedANDROID) (EGLDisplay dpy, EGLSurface surface, EGLint timestamp);
-
-#define EGL_TIMESTAMPS_ANDROID 0x3430
-#define EGL_COMPOSITE_DEADLINE_ANDROID 0x3431
-#define EGL_COMPOSITE_INTERVAL_ANDROID 0x3432
-#define EGL_COMPOSITE_TO_PRESENT_LATENCY_ANDROID 0x3433
-#define EGL_REQUESTED_PRESENT_TIME_ANDROID 0x3434
-#define EGL_RENDERING_COMPLETE_TIME_ANDROID 0x3435
-#define EGL_COMPOSITION_LATCH_TIME_ANDROID 0x3436
-#define EGL_FIRST_COMPOSITION_START_TIME_ANDROID 0x3437
-#define EGL_LAST_COMPOSITION_START_TIME_ANDROID 0x3438
-#define EGL_FIRST_COMPOSITION_GPU_FINISHED_TIME_ANDROID 0x3439
-#define EGL_DISPLAY_PRESENT_TIME_ANDROID 0x343A
-#define EGL_DEQUEUE_READY_TIME_ANDROID 0x343B
-#define EGL_READS_DONE_TIME_ANDROID 0x343C
-
-
-extern "C"
-{
-	extern PFNEGLGETSYSTEMTIMENVPROC eglGetSystemTimeNV_p;
-	extern PFNEGLCREATESYNCKHRPROC eglCreateSyncKHR_p;
-	extern PFNEGLDESTROYSYNCKHRPROC eglDestroySyncKHR_p;
-	extern PFNEGLCLIENTWAITSYNCKHRPROC eglClientWaitSyncKHR_p;
-	extern PFNEGLGETSYNCATTRIBKHRPROC eglGetSyncAttribKHR_p;
-
-	extern PFNeglPresentationTimeANDROID eglPresentationTimeANDROID_p;
-	extern PFNeglGetNextFrameIdANDROID eglGetNextFrameIdANDROID_p;
-	extern PFNeglGetCompositorTimingANDROID eglGetCompositorTimingANDROID_p;
-	extern PFNeglGetFrameTimestampsANDROID eglGetFrameTimestampsANDROID_p;
-	extern PFNeglQueryTimestampSupportedANDROID eglQueryTimestampSupportedANDROID_p;
-	extern PFNeglQueryTimestampSupportedANDROID eglGetCompositorTimingSupportedANDROID_p;
-	extern PFNeglQueryTimestampSupportedANDROID eglGetFrameTimestampsSupportedANDROID_p;
-}
-
-namespace GLFuncPointers
-{
-	// GL_QCOM_shader_framebuffer_fetch_noncoherent
-	extern PFNGLFRAMEBUFFERFETCHBARRIERQCOMPROC	glFramebufferFetchBarrierQCOM;
-}
-
 struct FAndroidOpenGL : public FOpenGLES
 {
 	static FORCEINLINE bool HasHardwareHiddenSurfaceRemoval() { return bHasHardwareHiddenSurfaceRemoval; };
@@ -203,13 +142,4 @@ struct FAndroidOpenGL : public FOpenGLES
 	static void StopRemoteCompileServices();
 };
 
-typedef FAndroidOpenGL FOpenGL;
-
-
-/** Unreal tokens that maps to different OpenGL tokens by platform. */
-#undef UGL_DRAW_FRAMEBUFFER
-#define UGL_DRAW_FRAMEBUFFER	GL_DRAW_FRAMEBUFFER_NV
-#undef UGL_READ_FRAMEBUFFER
-#define UGL_READ_FRAMEBUFFER	GL_READ_FRAMEBUFFER_NV
-
-#endif // PLATFORM_ANDROID
+using FOpenGL = FAndroidOpenGL;
diff --git a/Engine/Source/Runtime/OpenGLDrv/Private/Android/AndroidOpenGLFunctions.cpp b/Engine/Source/Runtime/OpenGLDrv/Private/Android/AndroidOpenGLFunctions.cpp
new file mode 100644
--- /dev/null
+++ b/Engine/Source/Runtime/OpenGLDrv/Private/Android/AndroidOpenGLFunctions.cpp
@@ -0,0 +1,61 @@
+// Copyright Epic Games, Inc. All Rights Reserved.
+
+#include "AndroidOpenGLFunctions.h"
+
+PFNeglPresentationTimeANDROID eglPresentationTimeANDROID_p = NULL;
+PFNeglGetNextFrameIdANDROID eglGetNextFrameIdANDROID_p = NULL;
+PFNeglGetCompositorTimingANDROID eglGetCompositorTimingANDROID_p = NULL;
+PFNeglGetFrameTimestampsANDROID eglGetFrameTimestampsANDROID_p = NULL;
+PFNeglQueryTimestampSupportedANDROID eglQueryTimestampSupportedANDROID_p = NULL;
+PFNeglQueryTimestampSupportedANDROID eglGetCompositorTimingSupportedANDROID_p = NULL;
+PFNeglQueryTimestampSupportedANDROID eglGetFrameTimestampsSupportedANDROID_p = NULL;
+
+PFNEGLGETSYSTEMTIMENVPROC eglGetSystemTimeNV_p = NULL;
+PFNEGLCREATESYNCKHRPROC eglCreateSyncKHR_p = NULL;
+PFNEGLDESTROYSYNCKHRPROC eglDestroySyncKHR_p = NULL;
+PFNEGLCLIENTWAITSYNCKHRPROC eglClientWaitSyncKHR_p = NULL;
+PFNEGLGETSYNCATTRIBKHRPROC eglGetSyncAttribKHR_p = NULL;
+
+namespace GLFuncPointers
+{
+	// Offscreen MSAA rendering
+	PFNGLFRAMEBUFFERTEXTURE2DMULTISAMPLEEXTPROC	glFramebufferTexture2DMultisampleEXT = NULL;
+	PFNGLRENDERBUFFERSTORAGEMULTISAMPLEEXTPROC	glRenderbufferStorageMultisampleEXT = NULL;
+
+	PFNGLPUSHGROUPMARKEREXTPROC				glPushGroupMarkerEXT = NULL;
+	PFNGLPOPGROUPMARKEREXTPROC				glPopGroupMarkerEXT = NULL;
+	PFNGLLABELOBJECTEXTPROC					glLabelObjectEXT = NULL;
+	PFNGLGETOBJECTLABELEXTPROC				glGetObjectLabelEXT = NULL;
+
+	PFNGLBUFFERSTORAGEEXTPROC				glBufferStorageEXT = NULL;
+	// KHR_debug
+	PFNGLDEBUGMESSAGECONTROLKHRPROC			glDebugMessageControlKHR = NULL;
+	PFNGLDEBUGMESSAGEINSERTKHRPROC			glDebugMessageInsertKHR = NULL;
+	PFNGLDEBUGMESSAGECALLBACKKHRPROC		glDebugMessageCallbackKHR = NULL;
+	PFNGLGETDEBUGMESSAGELOGKHRPROC			glDebugMessageLogKHR = NULL;
+	PFNGLGETPOINTERVKHRPROC					glGetPointervKHR = NULL;
+	PFNGLPUSHDEBUGGROUPKHRPROC				glPushDebugGroupKHR = NULL;
+	PFNGLPOPDEBUGGROUPKHRPROC				glPopDebugGroupKHR = NULL;
+	PFNGLOBJECTLABELKHRPROC					glObjectLabelKHR = NULL;
+	PFNGLGETOBJECTLABELKHRPROC				glGetObjectLabelKHR = NULL;
+	PFNGLOBJECTPTRLABELKHRPROC				glObjectPtrLabelKHR = NULL;
+	PFNGLGETOBJECTPTRLABELKHRPROC			glGetObjectPtrLabelKHR = NULL;
+
+	// ES 3.2
+	PFNGLTEXBUFFEREXTPROC					glTexBufferEXT = nullptr;
+	PFNGLTEXBUFFERRANGEEXTPROC				glTexBufferRangeEXT = nullptr;
+	PFNGLCOPYIMAGESUBDATAEXTPROC			glCopyImageSubData = nullptr;
+	PFNGLENABLEIEXTPROC						glEnableiEXT = nullptr;
+	PFNGLDISABLEIEXTPROC					glDisableiEXT = nullptr;
+	PFNGLBLENDEQUATIONIEXTPROC				glBlendEquationiEXT = nullptr;
+	PFNGLBLENDEQUATIONSEPARATEIEXTPROC		glBlendEquationSeparateiEXT = nullptr;
+	PFNGLBLENDFUNCIEXTPROC					glBlendFunciEXT = nullptr;
+	PFNGLBLENDFUNCSEPARATEIEXTPROC			glBlendFuncSeparateiEXT = nullptr;
+	PFNGLCOLORMASKIEXTPROC					glColorMaskiEXT = nullptr;
+	PFNGLFRAMEBUFFERTEXTUREPROC				glFramebufferTexture = nullptr;
+
+	PFNGLFRAMEBUFFERTEXTUREMULTIVIEWOVRPROC				glFramebufferTextureMultiviewOVR = NULL;
+	PFNGLFRAMEBUFFERTEXTUREMULTISAMPLEMULTIVIEWOVRPROC	glFramebufferTextureMultisampleMultiviewOVR = NULL;
+
+	PFNGLFRAMEBUFFERFETCHBARRIERQCOMPROC glFramebufferFetchBarrierQCOM = NULL;
+};
diff --git a/Engine/Source/Runtime/OpenGLDrv/Private/Android/AndroidOpenGLFunctions.h b/Engine/Source/Runtime/OpenGLDrv/Private/Android/AndroidOpenGLFunctions.h
new file mode 100644
--- /dev/null
+++ b/Engine/Source/Runtime/OpenGLDrv/Private/Android/AndroidOpenGLFunctions.h
@@ -0,0 +1,87 @@
+// Copyright Epic Games, Inc. All Rights Reserved.
+
+#pragma once
+
+#include "OpenGLPlatform.h"
+
+typedef khronos_stime_nanoseconds_t EGLnsecsANDROID;
+
+typedef GLboolean(GL_APIENTRYP PFNeglPresentationTimeANDROID) (EGLDisplay dpy, EGLSurface surface, EGLnsecsANDROID time);
+typedef GLboolean(GL_APIENTRYP PFNeglGetNextFrameIdANDROID) (EGLDisplay dpy, EGLSurface surface, EGLuint64KHR* frameId);
+typedef GLboolean(GL_APIENTRYP PFNeglGetCompositorTimingANDROID) (EGLDisplay dpy, EGLSurface surface, EGLint numTimestamps, const EGLint* names, EGLnsecsANDROID* values);
+typedef GLboolean(GL_APIENTRYP PFNeglGetFrameTimestampsANDROID) (EGLDisplay dpy, EGLSurface surface, EGLuint64KHR frameId, EGLint numTimestamps, const EGLint* timestamps, EGLnsecsANDROID* values);
+typedef GLboolean(GL_APIENTRYP PFNeglQueryTimestampSupportedANDROID) (EGLDisplay dpy, EGLSurface surface, EGLint timestamp);
+
+extern "C"
+{
+	extern PFNEGLGETSYSTEMTIMENVPROC eglGetSystemTimeNV_p;
+	extern PFNEGLCREATESYNCKHRPROC eglCreateSyncKHR_p;
+	extern PFNEGLDESTROYSYNCKHRPROC eglDestroySyncKHR_p;
+	extern PFNEGLCLIENTWAITSYNCKHRPROC eglClientWaitSyncKHR_p;
+	extern PFNEGLGETSYNCATTRIBKHRPROC eglGetSyncAttribKHR_p;
+
+	extern PFNeglPresentationTimeANDROID eglPresentationTimeANDROID_p;
+	extern PFNeglGetNextFrameIdANDROID eglGetNextFrameIdANDROID_p;
+	extern PFNeglGetCompositorTimingANDROID eglGetCompositorTimingANDROID_p;
+	extern PFNeglGetFrameTimestampsANDROID eglGetFrameTimestampsANDROID_p;
+	extern PFNeglQueryTimestampSupportedANDROID eglQueryTimestampSupportedANDROID_p;
+	extern PFNeglQueryTimestampSupportedANDROID eglGetCompositorTimingSupportedANDROID_p;
+	extern PFNeglQueryTimestampSupportedANDROID eglGetFrameTimestampsSupportedANDROID_p;
+}
+
+namespace GLFuncPointers
+{
+	// GL_QCOM_shader_framebuffer_fetch_noncoherent
+	extern PFNGLFRAMEBUFFERFETCHBARRIERQCOMPROC	glFramebufferFetchBarrierQCOM;
+}
+
+// FIXME: include gl32.h
+typedef void (GL_APIENTRYP PFNGLFRAMEBUFFERTEXTUREPROC) (GLenum target, GLenum attachment, GLuint texture, GLint level);
+
+namespace GLFuncPointers
+{
+	// GL_EXT_multisampled_render_to_texture
+	extern PFNGLFRAMEBUFFERTEXTURE2DMULTISAMPLEEXTPROC	glFramebufferTexture2DMultisampleEXT;
+	extern PFNGLRENDERBUFFERSTORAGEMULTISAMPLEEXTPROC	glRenderbufferStorageMultisampleEXT;
+
+	// GL_EXT_debug_marker
+	extern PFNGLPUSHGROUPMARKEREXTPROC		glPushGroupMarkerEXT;
+	extern PFNGLPOPGROUPMARKEREXTPROC 		glPopGroupMarkerEXT;
+
+	// GL_EXT_debug_label
+	extern PFNGLLABELOBJECTEXTPROC			glLabelObjectEXT;
+	extern PFNGLGETOBJECTLABELEXTPROC		glGetObjectLabelEXT;
+	//GL_EXT_buffer_storage
+	extern PFNGLBUFFERSTORAGEEXTPROC		glBufferStorageEXT;
+
+	extern PFNGLDEBUGMESSAGECONTROLKHRPROC	glDebugMessageControlKHR;
+	extern PFNGLDEBUGMESSAGEINSERTKHRPROC	glDebugMessageInsertKHR;
+	extern PFNGLDEBUGMESSAGECALLBACKKHRPROC	glDebugMessageCallbackKHR;
+	extern PFNGLGETDEBUGMESSAGELOGKHRPROC	glDebugMessageLogKHR;
+	extern PFNGLGETPOINTERVKHRPROC			glGetPointervKHR;
+	extern PFNGLPUSHDEBUGGROUPKHRPROC		glPushDebugGroupKHR;
+	extern PFNGLPOPDEBUGGROUPKHRPROC		glPopDebugGroupKHR;
+	extern PFNGLOBJECTLABELKHRPROC			glObjectLabelKHR;
+	extern PFNGLGETOBJECTLABELKHRPROC		glGetObjectLabelKHR;
+	extern PFNGLOBJECTPTRLABELKHRPROC		glObjectPtrLabelKHR;
+	extern PFNGLGETOBJECTPTRLABELKHRPROC	glGetObjectPtrLabelKHR;
+
+	// ES 3.2
+	extern PFNGLTEXBUFFEREXTPROC				glTexBufferEXT;
+	extern PFNGLTEXBUFFERRANGEEXTPROC			glTexBufferRangeEXT;
+	extern PFNGLCOPYIMAGESUBDATAEXTPROC			glCopyImageSubData;
+	extern PFNGLENABLEIEXTPROC					glEnableiEXT;
+	extern PFNGLDISABLEIEXTPROC					glDisableiEXT;
+	extern PFNGLBLENDEQUATIONIEXTPROC			glBlendEquationiEXT;
+	extern PFNGLBLENDEQUATIONSEPARATEIEXTPROC	glBlendEquationSeparateiEXT;
+	extern PFNGLBLENDFUNCIEXTPROC				glBlendFunciEXT;
+	extern PFNGLBLENDFUNCSEPARATEIEXTPROC		glBlendFuncSeparateiEXT;
+	extern PFNGLCOLORMASKIEXTPROC				glColorMaskiEXT;
+	extern PFNGLFRAMEBUFFERTEXTUREPROC			glFramebufferTexture;
+
+	// Mobile multi-view
+	extern PFNGLFRAMEBUFFERTEXTUREMULTIVIEWOVRPROC glFramebufferTextureMultiviewOVR;
+	extern PFNGLFRAMEBUFFERTEXTUREMULTISAMPLEMULTIVIEWOVRPROC glFramebufferTextureMultisampleMultiviewOVR;
+};
+
+using namespace GLFuncPointers;
\ No newline at end of file
diff --git a/Engine/Source/Runtime/OpenGLDrv/Private/Android/AndroidOpenGLPlatform.cpp b/Engine/Source/Runtime/OpenGLDrv/Private/Android/AndroidOpenGLPlatform.cpp
new file mode 100644
--- /dev/null
+++ b/Engine/Source/Runtime/OpenGLDrv/Private/Android/AndroidOpenGLPlatform.cpp
@@ -0,0 +1,54 @@
+// Copyright Epic Games, Inc. All Rights Reserved.
+
+#include "AndroidOpenGLPlatform.h"
+
+#include "Android/AndroidApplication.h"
+#include "Android/AndroidEGL.h"
+#include "Android/AndroidOpenGLPrivate.h"
+#include "Android/AndroidPlatformMisc.h"
+
+FString FAndroidMisc::GetGPUFamily()
+{
+	return FAndroidGPUInfo::Get().GetGPUFamily();
+}
+
+FString FAndroidMisc::GetGLVersion()
+{
+	return FAndroidGPUInfo::Get().GLVersion;
+}
+
+bool FAndroidMisc::SupportsFloatingPointRenderTargets()
+{
+	return FAndroidGPUInfo::Get().bSupportsFloatingPointRenderTargets;
+}
+
+bool FAndroidMisc::SupportsShaderFramebufferFetch()
+{
+	return FAndroidGPUInfo::Get().bSupportsFrameBufferFetch;
+}
+
+bool FAndroidMisc::SupportsES30()
+{
+	return true;
+}
+
+void FAndroidMisc::GetValidTargetPlatforms(TArray<FString>& TargetPlatformNames)
+{
+	TargetPlatformNames = FAndroidGPUInfo::Get().TargetPlatformNames;
+}
+
+void FAndroidAppEntry::PlatformInit()
+{
+	// Try to create an ES3.2 EGL here for gpu queries and don't have to recreate the GL context.
+	AndroidEGL::GetInstance()->Init(AndroidEGL::AV_OpenGLES, 3, 2);
+}
+
+void FAndroidAppEntry::ReleaseEGL()
+{
+	AndroidEGL* EGL = AndroidEGL::GetInstance();
+	if (EGL->IsInitialized())
+	{
+		EGL->DestroyBackBuffer();
+		EGL->Terminate();
+	}
+}
diff --git a/Engine/Source/Runtime/OpenGLDrv/Private/Android/AndroidOpenGLPlatform.h b/Engine/Source/Runtime/OpenGLDrv/Private/Android/AndroidOpenGLPlatform.h
new file mode 100644
--- /dev/null
+++ b/Engine/Source/Runtime/OpenGLDrv/Private/Android/AndroidOpenGLPlatform.h
@@ -0,0 +1,131 @@
+// Copyright Epic Games, Inc. All Rights Reserved.
+
+#pragma once
+
+#include "OpenGLThirdParty.h"
+
+#define UGL_PLATFORM_SUPPORTS_GLES 1
+#define UGL_PLATFORM_SUPPORTS_GL3  0
+#define UGL_PLATFORM_SUPPORTS_GL4  0
+
+typedef GLsync			UGLsync;
+#define GLdouble		GLfloat
+typedef khronos_int64_t GLint64;
+typedef khronos_uint64_t GLuint64;
+#define GL_CLAMP		GL_CLAMP_TO_EDGE
+
+#define EGL_TIMESTAMPS_ANDROID 0x3430
+#define EGL_COMPOSITE_DEADLINE_ANDROID 0x3431
+#define EGL_COMPOSITE_INTERVAL_ANDROID 0x3432
+#define EGL_COMPOSITE_TO_PRESENT_LATENCY_ANDROID 0x3433
+#define EGL_REQUESTED_PRESENT_TIME_ANDROID 0x3434
+#define EGL_RENDERING_COMPLETE_TIME_ANDROID 0x3435
+#define EGL_COMPOSITION_LATCH_TIME_ANDROID 0x3436
+#define EGL_FIRST_COMPOSITION_START_TIME_ANDROID 0x3437
+#define EGL_LAST_COMPOSITION_START_TIME_ANDROID 0x3438
+#define EGL_FIRST_COMPOSITION_GPU_FINISHED_TIME_ANDROID 0x3439
+#define EGL_DISPLAY_PRESENT_TIME_ANDROID 0x343A
+#define EGL_DEQUEUE_READY_TIME_ANDROID 0x343B
+#define EGL_READS_DONE_TIME_ANDROID 0x343C
+
+#ifdef GL_AMD_debug_output
+	#undef GL_AMD_debug_output
+#endif
+
+/** Unreal tokens that maps to different OpenGL tokens by platform. */
+#define UGL_DRAW_FRAMEBUFFER	GL_DRAW_FRAMEBUFFER_NV
+#define UGL_READ_FRAMEBUFFER	GL_READ_FRAMEBUFFER_NV
+#define UGL_ABGR8				GL_UNSIGNED_BYTE
+#define UGL_ANY_SAMPLES_PASSED	GL_ANY_SAMPLES_PASSED
+#define UGL_SAMPLES_PASSED		GL_SAMPLES_PASSED
+#define UGL_TIME_ELAPSED		GL_TIME_ELAPSED_EXT
+#define UGL_CLAMP_TO_BORDER		GL_CLAMP_TO_EDGE
+
+/** Official OpenGL definitions */
+#ifndef GL_FILL
+#define GL_FILL 0x1B02
+#endif
+#ifndef GL_GEOMETRY_SHADER
+#define GL_GEOMETRY_SHADER 0x8DD9
+#endif
+#ifndef GL_SAMPLER_1D
+#define GL_SAMPLER_1D 0x8B5D
+#endif
+#ifndef GL_SAMPLER_1D_SHADOW
+#define GL_SAMPLER_1D_SHADOW 0x8B61
+#endif
+#ifndef GL_DOUBLE
+#define GL_DOUBLE 0x140A
+#endif
+#ifndef GL_BGRA
+#define GL_BGRA	GL_BGRA_EXT 
+#endif
+#ifndef GL_TEXTURE_BUFFER
+#define GL_TEXTURE_BUFFER 0x8C2A
+#endif
+#ifndef GL_RGBA16
+/* GL_EXT_texture_norm16 */
+#define GL_RGBA16 GL_RGBA16_EXT
+#endif
+#ifndef GL_R16
+/* GL_EXT_texture_norm16 */
+#define GL_R16 GL_R16_EXT
+#endif
+#ifndef GL_TEXTURE_BUFFER_OFFSET_ALIGNMENT
+#define GL_TEXTURE_BUFFER_OFFSET_ALIGNMENT 0x919F
+#endif
+#ifndef GL_POLYGON_OFFSET_LINE
+#define GL_POLYGON_OFFSET_LINE 0x2A02
+#endif
+#ifndef GL_POLYGON_OFFSET_POINT
+#define GL_POLYGON_OFFSET_POINT 0x2A01
+#endif
+#ifndef GL_TEXTURE_LOD_BIAS
+#define GL_TEXTURE_LOD_BIAS 0x8501
+#endif
+#ifndef GL_FRAMEBUFFER_SRGB
+#define GL_FRAMEBUFFER_SRGB 0x8DB9
+#endif
+#ifndef GL_SAMPLES_PASSED
+#define GL_SAMPLES_PASSED 0x8914
+#endif
+#ifndef GL_POINT
+#define GL_POINT 0x1B00
+#endif
+#ifndef GL_LINE
+#define GL_LINE 0x1B01
+#endif
+#ifndef GL_TEXTURE_1D
+#define GL_TEXTURE_1D			0x0DE0
+#endif
+#ifndef GL_TEXTURE_1D_ARRAY
+#define GL_TEXTURE_1D_ARRAY		0x8C18
+#endif
+#ifndef GL_TEXTURE_RECTANGLE
+#define GL_TEXTURE_RECTANGLE	0x84F5
+#endif
+#ifndef GL_MAX_TEXTURE_BUFFER_SIZE
+#define GL_MAX_TEXTURE_BUFFER_SIZE 0x8C2B
+#endif
+#ifndef GL_DEPTH_CLAMP
+#define GL_DEPTH_CLAMP 0x864F
+#endif
+
+/** For the shader stage bits that don't exist just use 0 */
+#define GL_GEOMETRY_SHADER_BIT				0x00000000
+#define GL_TESS_CONTROL_SHADER_BIT			0x00000000
+#define GL_TESS_EVALUATION_SHADER_BIT		0x00000000
+
+// Normalize debug macros due to naming differences across GL versions
+#if defined(GL_KHR_debug) && GL_KHR_debug
+#define GL_DEBUG_SOURCE_OTHER_ARB GL_DEBUG_SOURCE_OTHER_KHR
+#define GL_DEBUG_SOURCE_API_ARB GL_DEBUG_SOURCE_API_KHR
+#define GL_DEBUG_TYPE_ERROR_ARB GL_DEBUG_TYPE_ERROR_KHR
+#define GL_DEBUG_TYPE_OTHER_ARB GL_DEBUG_TYPE_OTHER_KHR
+#define GL_DEBUG_TYPE_MARKER GL_DEBUG_TYPE_MARKER_KHR
+#define GL_DEBUG_TYPE_PUSH_GROUP GL_DEBUG_TYPE_PUSH_GROUP_KHR
+#define GL_DEBUG_TYPE_POP_GROUP GL_DEBUG_TYPE_POP_GROUP_KHR
+#define GL_DEBUG_SEVERITY_HIGH_ARB GL_DEBUG_SEVERITY_HIGH_KHR
+#define GL_DEBUG_SEVERITY_LOW_ARB GL_DEBUG_SEVERITY_LOW_KHR
+#define GL_DEBUG_SEVERITY_NOTIFICATION GL_DEBUG_SEVERITY_NOTIFICATION_KHR
+#endif
diff --git a/Engine/Source/Runtime/OpenGLDrv/Private/Android/AndroidOpenGLPrivate.h b/Engine/Source/Runtime/OpenGLDrv/Private/Android/AndroidOpenGLPrivate.h
--- a/Engine/Source/Runtime/OpenGLDrv/Private/Android/AndroidOpenGLPrivate.h
+++ b/Engine/Source/Runtime/OpenGLDrv/Private/Android/AndroidOpenGLPrivate.h
@@ -5,6 +5,7 @@
 =============================================================================*/
 #pragma once
 
+#include "AndroidEGL.h"
 #include "Android/AndroidApplication.h"
 #include "libgpuinfo.hpp"
 #include "Internationalization/Regex.h"
diff --git a/Engine/Source/Runtime/OpenGLDrv/Private/Linux/LinuxOpenGL.h b/Engine/Source/Runtime/OpenGLDrv/Private/Linux/LinuxOpenGL.h
new file mode 100644
--- /dev/null
+++ b/Engine/Source/Runtime/OpenGLDrv/Private/Linux/LinuxOpenGL.h
@@ -0,0 +1,103 @@
+// Copyright Epic Games, Inc. All Rights Reserved.
+
+#pragma once
+
+#include "OpenGLThirdParty.h"
+#include "OpenGLPlatform.h"
+
+#include "OpenGL4.h"
+
+struct FLinuxOpenGL : public FOpenGL4
+{
+	static FORCEINLINE void InitDebugContext()
+	{
+		bDebugContext = glIsEnabled(GL_DEBUG_OUTPUT) != GL_FALSE;
+	}
+
+	static FORCEINLINE void LabelObject(GLenum Type, GLuint Object, const ANSICHAR* Name)
+	{
+		if (glObjectLabel && bDebugContext)
+		{
+			glObjectLabel(Type, Object, FCStringAnsi::Strlen(Name), Name);
+		}
+	}
+
+	static FORCEINLINE void PushGroupMarker(const ANSICHAR* Name)
+	{
+		if (glPushDebugGroup && bDebugContext)
+		{
+			glPushDebugGroup(GL_DEBUG_SOURCE_APPLICATION, 1, FCStringAnsi::Strlen(Name), Name);
+		}
+	}
+
+	static FORCEINLINE void PopGroupMarker()
+	{
+		if (glPopDebugGroup && bDebugContext)
+		{
+			glPopDebugGroup();
+		}
+	}
+
+	static FORCEINLINE bool TexStorage2D(GLenum Target, GLint Levels, GLint InternalFormat, GLsizei Width, GLsizei Height, GLenum Format, GLenum Type, ETextureCreateFlags Flags)
+	{
+		if (glTexStorage2D != NULL)
+		{
+			glTexStorage2D(Target, Levels, InternalFormat, Width, Height);
+			return true;
+		}
+
+		return false;
+	}
+
+	static FORCEINLINE void TexStorage3D(GLenum Target, GLint Levels, GLint InternalFormat, GLsizei Width, GLsizei Height, GLsizei Depth, GLenum Format, GLenum Type)
+	{
+		if (glTexStorage3D)
+		{
+			glTexStorage3D(Target, Levels, InternalFormat, Width, Height, Depth);
+		}
+		else
+		{
+			const bool bArrayTexture = Target == GL_TEXTURE_2D_ARRAY || Target == GL_TEXTURE_CUBE_MAP_ARRAY;
+
+			for (uint32 MipIndex = 0; MipIndex < uint32(Levels); MipIndex++)
+			{
+				glTexImage3D(
+					Target,
+					MipIndex,
+					InternalFormat,
+					FMath::Max<uint32>(1, (Width >> MipIndex)),
+					FMath::Max<uint32>(1, (Height >> MipIndex)),
+					(bArrayTexture) ? Depth : FMath::Max<uint32>(1, (Depth >> MipIndex)),
+					0,
+					Format,
+					Type,
+					NULL
+				);
+			}
+		}
+	}
+
+	static FORCEINLINE void CopyImageSubData(GLuint SrcName, GLenum SrcTarget, GLint SrcLevel, GLint SrcX, GLint SrcY, GLint SrcZ, GLuint DstName, GLenum DstTarget, GLint DstLevel, GLint DstX, GLint DstY, GLint DstZ, GLsizei Width, GLsizei Height, GLsizei Depth)
+	{
+		glCopyImageSubData(SrcName, SrcTarget, SrcLevel, SrcX, SrcY, SrcZ, DstName, DstTarget, DstLevel, DstX, DstY, DstZ, Width, Height, Depth);
+	}
+
+	static void ProcessExtensions(const FString& ExtensionsString)
+	{
+		FOpenGL4::ProcessExtensions(ExtensionsString);
+
+		FString VendorName(ANSI_TO_TCHAR((const ANSICHAR*)glGetString(GL_VENDOR)));
+
+		if (VendorName.Contains(TEXT("ATI ")))
+		{
+			// Workaround for AMD driver not handling GL_SRGB8_ALPHA8 in glTexStorage2D() properly (gets treated as non-sRGB)
+			// FIXME: obsolete ? this was the case in <= 2014
+			glTexStorage1D = nullptr;
+			glTexStorage2D = nullptr;
+			glTexStorage3D = nullptr;
+		}
+	}
+};
+
+using FOpenGL = FLinuxOpenGL;
+
diff --git a/Engine/Source/Runtime/OpenGLDrv/Private/Linux/LinuxOpenGLFunctions.cpp b/Engine/Source/Runtime/OpenGLDrv/Private/Linux/LinuxOpenGLFunctions.cpp
new file mode 100644
--- /dev/null
+++ b/Engine/Source/Runtime/OpenGLDrv/Private/Linux/LinuxOpenGLFunctions.cpp
@@ -0,0 +1,13 @@
+// Copyright Epic Games, Inc. All Rights Reserved.
+
+#include "LinuxOpenGLFunctions.h"
+
+/*------------------------------------------------------------------------------
+	OpenGL function pointers.
+------------------------------------------------------------------------------*/
+#define DEFINE_GL_ENTRYPOINTS(Type,Func) Type Func = NULL;
+namespace GLFuncPointers	// see explanation in LinuxOpenGLFunctions.h why we need the namespace
+{
+	ENUM_GL_ENTRYPOINTS_ALL(DEFINE_GL_ENTRYPOINTS);
+};
+#undef DEFINE_GL_ENTRYPOINTS
diff --git a/Engine/Source/Runtime/OpenGLDrv/Private/Linux/LinuxOpenGLFunctions.h b/Engine/Source/Runtime/OpenGLDrv/Private/Linux/LinuxOpenGLFunctions.h
new file mode 100644
--- /dev/null
+++ b/Engine/Source/Runtime/OpenGLDrv/Private/Linux/LinuxOpenGLFunctions.h
@@ -0,0 +1,365 @@
+// Copyright Epic Games, Inc. All Rights Reserved.
+
+#pragma once
+
+#include "OpenGLPlatform.h"
+
+/** List all OpenGL entry points used by Unreal. */
+#define ENUM_GL_ENTRYPOINTS(EnumMacro) \
+	EnumMacro(PFNGLBINDTEXTUREPROC,glBindTexture) \
+	EnumMacro(PFNGLBLENDFUNCPROC,glBlendFunc) \
+	EnumMacro(PFNGLCOLORMASKPROC,glColorMask) \
+	EnumMacro(PFNGLCOPYTEXIMAGE1DPROC,glCopyTexImage1D) \
+	EnumMacro(PFNGLCOPYTEXIMAGE2DPROC,glCopyTexImage2D) \
+	EnumMacro(PFNGLCOPYTEXSUBIMAGE1DPROC,glCopyTexSubImage1D) \
+	EnumMacro(PFNGLCOPYTEXSUBIMAGE2DPROC,glCopyTexSubImage2D) \
+	EnumMacro(PFNGLCULLFACEPROC,glCullFace) \
+	EnumMacro(PFNGLDELETETEXTURESPROC,glDeleteTextures) \
+	EnumMacro(PFNGLDEPTHFUNCPROC,glDepthFunc) \
+	EnumMacro(PFNGLDEPTHMASKPROC,glDepthMask) \
+	EnumMacro(PFNGLDEPTHRANGEPROC,glDepthRange) \
+	EnumMacro(PFNGLDISABLEPROC,glDisable) \
+	EnumMacro(PFNGLDRAWARRAYSPROC,glDrawArrays) \
+	EnumMacro(PFNGLDRAWBUFFERPROC,glDrawBuffer) \
+	EnumMacro(PFNGLDRAWELEMENTSPROC,glDrawElements) \
+	EnumMacro(PFNGLENABLEPROC,glEnable) \
+	EnumMacro(PFNGLFINISHPROC,glFinish) \
+	EnumMacro(PFNGLFLUSHPROC,glFlush) \
+	EnumMacro(PFNGLFRONTFACEPROC,glFrontFace) \
+	EnumMacro(PFNGLGENTEXTURESPROC,glGenTextures) \
+	EnumMacro(PFNGLGETBOOLEANVPROC,glGetBooleanv) \
+	EnumMacro(PFNGLGETDOUBLEVPROC,glGetDoublev) \
+	EnumMacro(PFNGLGETERRORPROC,glGetError) \
+	EnumMacro(PFNGLGETFLOATVPROC,glGetFloatv) \
+	EnumMacro(PFNGLGETINTEGERVPROC,glGetIntegerv) \
+	EnumMacro(PFNGLGETPOINTERVPROC,glGetPointerv) \
+	EnumMacro(PFNGLGETSTRINGPROC,glGetString) \
+	EnumMacro(PFNGLGETTEXIMAGEPROC,glGetTexImage) \
+	EnumMacro(PFNGLGETTEXLEVELPARAMETERFVPROC,glGetTexLevelParameterfv) \
+	EnumMacro(PFNGLGETTEXLEVELPARAMETERIVPROC,glGetTexLevelParameteriv) \
+	EnumMacro(PFNGLGETTEXPARAMETERFVPROC,glGetTexParameterfv) \
+	EnumMacro(PFNGLGETTEXPARAMETERIVPROC,glGetTexParameteriv) \
+	EnumMacro(PFNGLHINTPROC,glHint) \
+	EnumMacro(PFNGLISENABLEDPROC,glIsEnabled) \
+	EnumMacro(PFNGLISTEXTUREPROC,glIsTexture) \
+	EnumMacro(PFNGLLINEWIDTHPROC,glLineWidth) \
+	EnumMacro(PFNGLLOGICOPPROC,glLogicOp) \
+	EnumMacro(PFNGLPIXELSTOREFPROC,glPixelStoref) \
+	EnumMacro(PFNGLPIXELSTOREIPROC,glPixelStorei) \
+	EnumMacro(PFNGLPOINTSIZEPROC,glPointSize) \
+	EnumMacro(PFNGLPOLYGONMODEPROC,glPolygonMode) \
+	EnumMacro(PFNGLPOLYGONOFFSETPROC,glPolygonOffset) \
+	EnumMacro(PFNGLREADBUFFERPROC,glReadBuffer) \
+	EnumMacro(PFNGLREADPIXELSPROC,glReadPixels) \
+	EnumMacro(PFNGLSCISSORPROC,glScissor) \
+	EnumMacro(PFNGLCLEARDEPTHPROC,glClearDepth) \
+	EnumMacro(PFNGLCLEARCOLORPROC,glClearColor) \
+	EnumMacro(PFNGLCLEARSTENCILPROC,glClearStencil) \
+	EnumMacro(PFNGLCLEARPROC,glClear) \
+	EnumMacro(PFNGLSTENCILFUNCPROC,glStencilFunc) \
+	EnumMacro(PFNGLSTENCILMASKPROC,glStencilMask) \
+	EnumMacro(PFNGLSTENCILOPPROC,glStencilOp) \
+	EnumMacro(PFNGLTEXIMAGE1DPROC,glTexImage1D) \
+	EnumMacro(PFNGLTEXIMAGE2DPROC,glTexImage2D) \
+	EnumMacro(PFNGLTEXPARAMETERFPROC,glTexParameterf) \
+	EnumMacro(PFNGLTEXPARAMETERFVPROC,glTexParameterfv) \
+	EnumMacro(PFNGLTEXPARAMETERIPROC,glTexParameteri) \
+	EnumMacro(PFNGLTEXPARAMETERIVPROC,glTexParameteriv) \
+	EnumMacro(PFNGLTEXSUBIMAGE1DPROC,glTexSubImage1D) \
+	EnumMacro(PFNGLTEXSUBIMAGE2DPROC,glTexSubImage2D) \
+	EnumMacro(PFNGLVIEWPORTPROC,glViewport) \
+	EnumMacro(PFNGLBINDSAMPLERPROC,glBindSampler) \
+	EnumMacro(PFNGLDELETESAMPLERSPROC,glDeleteSamplers) \
+	EnumMacro(PFNGLGENSAMPLERSPROC,glGenSamplers) \
+	EnumMacro(PFNGLSAMPLERPARAMETERIPROC,glSamplerParameteri) \
+	EnumMacro(PFNGLCLEARBUFFERFVPROC,glClearBufferfv) \
+	EnumMacro(PFNGLCLEARBUFFERIVPROC,glClearBufferiv) \
+	EnumMacro(PFNGLCLEARBUFFERUIVPROC,glClearBufferuiv) \
+	EnumMacro(PFNGLCLEARBUFFERFIPROC,glClearBufferfi) \
+	EnumMacro(PFNGLCOLORMASKIPROC,glColorMaski) \
+	EnumMacro(PFNGLDISABLEIPROC,glDisablei) \
+	EnumMacro(PFNGLENABLEIPROC,glEnablei) \
+	EnumMacro(PFNGLGETBOOLEANI_VPROC,glGetBooleani_v) \
+	EnumMacro(PFNGLGETINTEGERI_VPROC,glGetIntegeri_v) \
+	EnumMacro(PFNGLISENABLEDIPROC,glIsEnabledi) \
+	EnumMacro(PFNGLBLENDEQUATIONSEPARATEIARBPROC,glBlendEquationSeparateiARB) \
+	EnumMacro(PFNGLBLENDEQUATIONIARBPROC,glBlendEquationiARB) \
+	EnumMacro(PFNGLBLENDFUNCSEPARATEIARBPROC,glBlendFuncSeparateiARB) \
+	EnumMacro(PFNGLBLENDFUNCIARBPROC,glBlendFunciARB) \
+	EnumMacro(PFNGLBLENDCOLORPROC,glBlendColor) \
+	EnumMacro(PFNGLBLENDEQUATIONPROC,glBlendEquation) \
+	EnumMacro(PFNGLDRAWRANGEELEMENTSPROC,glDrawRangeElements) \
+	EnumMacro(PFNGLTEXIMAGE3DPROC,glTexImage3D) \
+	EnumMacro(PFNGLTEXSUBIMAGE3DPROC,glTexSubImage3D) \
+	EnumMacro(PFNGLCOPYTEXSUBIMAGE3DPROC,glCopyTexSubImage3D) \
+	EnumMacro(PFNGLACTIVETEXTUREPROC,glActiveTexture) \
+	EnumMacro(PFNGLSAMPLECOVERAGEPROC,glSampleCoverage) \
+	EnumMacro(PFNGLCOMPRESSEDTEXIMAGE3DPROC,glCompressedTexImage3D) \
+	EnumMacro(PFNGLCOMPRESSEDTEXIMAGE2DPROC,glCompressedTexImage2D) \
+	EnumMacro(PFNGLCOMPRESSEDTEXIMAGE1DPROC,glCompressedTexImage1D) \
+	EnumMacro(PFNGLCOMPRESSEDTEXSUBIMAGE3DPROC,glCompressedTexSubImage3D) \
+	EnumMacro(PFNGLCOMPRESSEDTEXSUBIMAGE2DPROC,glCompressedTexSubImage2D) \
+	EnumMacro(PFNGLCOMPRESSEDTEXSUBIMAGE1DPROC,glCompressedTexSubImage1D) \
+	EnumMacro(PFNGLGETCOMPRESSEDTEXIMAGEPROC,glGetCompressedTexImage) \
+	EnumMacro(PFNGLBLENDFUNCSEPARATEPROC,glBlendFuncSeparate) \
+	EnumMacro(PFNGLMULTIDRAWARRAYSPROC,glMultiDrawArrays) \
+	EnumMacro(PFNGLMULTIDRAWELEMENTSPROC,glMultiDrawElements) \
+	EnumMacro(PFNGLPOINTPARAMETERFPROC,glPointParameterf) \
+	EnumMacro(PFNGLPOINTPARAMETERFVPROC,glPointParameterfv) \
+	EnumMacro(PFNGLPOINTPARAMETERIPROC,glPointParameteri) \
+	EnumMacro(PFNGLPOINTPARAMETERIVPROC,glPointParameteriv) \
+	EnumMacro(PFNGLGENQUERIESPROC,glGenQueries) \
+	EnumMacro(PFNGLDELETEQUERIESPROC,glDeleteQueries) \
+	EnumMacro(PFNGLISQUERYPROC,glIsQuery) \
+	EnumMacro(PFNGLBEGINQUERYPROC,glBeginQuery) \
+	EnumMacro(PFNGLENDQUERYPROC,glEndQuery) \
+	EnumMacro(PFNGLGETQUERYIVPROC,glGetQueryiv) \
+	EnumMacro(PFNGLGETQUERYOBJECTIVPROC,glGetQueryObjectiv) \
+	EnumMacro(PFNGLGETQUERYOBJECTUIVPROC,glGetQueryObjectuiv) \
+	EnumMacro(PFNGLBINDBUFFERPROC,glBindBuffer) \
+	EnumMacro(PFNGLBINDBUFFERBASEPROC,glBindBufferBase) \
+	EnumMacro(PFNGLDELETEBUFFERSPROC,glDeleteBuffers) \
+	EnumMacro(PFNGLGENBUFFERSPROC,glGenBuffers) \
+	EnumMacro(PFNGLISBUFFERPROC,glIsBuffer) \
+	EnumMacro(PFNGLBUFFERDATAPROC,glBufferData) \
+	EnumMacro(PFNGLBUFFERSUBDATAPROC,glBufferSubData) \
+	EnumMacro(PFNGLGETBUFFERSUBDATAPROC,glGetBufferSubData) \
+	EnumMacro(PFNGLMAPBUFFERPROC,glMapBuffer) \
+	EnumMacro(PFNGLUNMAPBUFFERPROC,glUnmapBuffer) \
+	EnumMacro(PFNGLGETBUFFERPARAMETERIVPROC,glGetBufferParameteriv) \
+	EnumMacro(PFNGLGETBUFFERPOINTERVPROC,glGetBufferPointerv) \
+	EnumMacro(PFNGLBLENDEQUATIONSEPARATEPROC,glBlendEquationSeparate) \
+	EnumMacro(PFNGLDRAWBUFFERSPROC,glDrawBuffers) \
+	EnumMacro(PFNGLSTENCILOPSEPARATEPROC,glStencilOpSeparate) \
+	EnumMacro(PFNGLSTENCILFUNCSEPARATEPROC,glStencilFuncSeparate) \
+	EnumMacro(PFNGLSTENCILMASKSEPARATEPROC,glStencilMaskSeparate) \
+	EnumMacro(PFNGLATTACHSHADERPROC,glAttachShader) \
+	EnumMacro(PFNGLBINDATTRIBLOCATIONPROC,glBindAttribLocation) \
+	EnumMacro(PFNGLBINDFRAGDATALOCATIONPROC,glBindFragDataLocation) \
+	EnumMacro(PFNGLCOMPILESHADERPROC,glCompileShader) \
+	EnumMacro(PFNGLCREATEPROGRAMPROC,glCreateProgram) \
+	EnumMacro(PFNGLCREATESHADERPROC,glCreateShader) \
+	EnumMacro(PFNGLDELETEPROGRAMPROC,glDeleteProgram) \
+	EnumMacro(PFNGLDELETESHADERPROC,glDeleteShader) \
+	EnumMacro(PFNGLDETACHSHADERPROC,glDetachShader) \
+	EnumMacro(PFNGLDISABLEVERTEXATTRIBARRAYPROC,glDisableVertexAttribArray) \
+	EnumMacro(PFNGLENABLEVERTEXATTRIBARRAYPROC,glEnableVertexAttribArray) \
+	EnumMacro(PFNGLGETACTIVEATTRIBPROC,glGetActiveAttrib) \
+	EnumMacro(PFNGLGETACTIVEUNIFORMPROC,glGetActiveUniform) \
+	EnumMacro(PFNGLGETATTACHEDSHADERSPROC,glGetAttachedShaders) \
+	EnumMacro(PFNGLGETATTRIBLOCATIONPROC,glGetAttribLocation) \
+	EnumMacro(PFNGLGETPROGRAMIVPROC,glGetProgramiv) \
+	EnumMacro(PFNGLGETPROGRAMINFOLOGPROC,glGetProgramInfoLog) \
+	EnumMacro(PFNGLGETSHADERIVPROC,glGetShaderiv) \
+	EnumMacro(PFNGLGETSHADERINFOLOGPROC,glGetShaderInfoLog) \
+	EnumMacro(PFNGLGETSHADERSOURCEPROC,glGetShaderSource) \
+	EnumMacro(PFNGLGETUNIFORMLOCATIONPROC,glGetUniformLocation) \
+	EnumMacro(PFNGLGETUNIFORMBLOCKINDEXPROC,glGetUniformBlockIndex) \
+	EnumMacro(PFNGLGETUNIFORMFVPROC,glGetUniformfv) \
+	EnumMacro(PFNGLGETUNIFORMIVPROC,glGetUniformiv) \
+	EnumMacro(PFNGLGETVERTEXATTRIBDVPROC,glGetVertexAttribdv) \
+	EnumMacro(PFNGLGETVERTEXATTRIBFVPROC,glGetVertexAttribfv) \
+	EnumMacro(PFNGLGETVERTEXATTRIBIVPROC,glGetVertexAttribiv) \
+	EnumMacro(PFNGLGETVERTEXATTRIBPOINTERVPROC,glGetVertexAttribPointerv) \
+	EnumMacro(PFNGLISPROGRAMPROC,glIsProgram) \
+	EnumMacro(PFNGLISSHADERPROC,glIsShader) \
+	EnumMacro(PFNGLLINKPROGRAMPROC,glLinkProgram) \
+	EnumMacro(PFNGLSHADERSOURCEPROC,glShaderSource) \
+	EnumMacro(PFNGLUSEPROGRAMPROC,glUseProgram) \
+	EnumMacro(PFNGLUNIFORM1FPROC,glUniform1f) \
+	EnumMacro(PFNGLUNIFORM2FPROC,glUniform2f) \
+	EnumMacro(PFNGLUNIFORM3FPROC,glUniform3f) \
+	EnumMacro(PFNGLUNIFORM4FPROC,glUniform4f) \
+	EnumMacro(PFNGLUNIFORM1IPROC,glUniform1i) \
+	EnumMacro(PFNGLUNIFORM2IPROC,glUniform2i) \
+	EnumMacro(PFNGLUNIFORM3IPROC,glUniform3i) \
+	EnumMacro(PFNGLUNIFORM4IPROC,glUniform4i) \
+	EnumMacro(PFNGLUNIFORM1FVPROC,glUniform1fv) \
+	EnumMacro(PFNGLUNIFORM2FVPROC,glUniform2fv) \
+	EnumMacro(PFNGLUNIFORM3FVPROC,glUniform3fv) \
+	EnumMacro(PFNGLUNIFORM4FVPROC,glUniform4fv) \
+	EnumMacro(PFNGLUNIFORM1IVPROC,glUniform1iv) \
+	EnumMacro(PFNGLUNIFORM2IVPROC,glUniform2iv) \
+	EnumMacro(PFNGLUNIFORM3IVPROC,glUniform3iv) \
+	EnumMacro(PFNGLUNIFORM4IVPROC,glUniform4iv) \
+	EnumMacro(PFNGLUNIFORM1UIVPROC,glUniform1uiv) \
+	EnumMacro(PFNGLUNIFORM2UIVPROC,glUniform2uiv) \
+	EnumMacro(PFNGLUNIFORM3UIVPROC,glUniform3uiv) \
+	EnumMacro(PFNGLUNIFORM4UIVPROC,glUniform4uiv) \
+	EnumMacro(PFNGLUNIFORMBLOCKBINDINGPROC,glUniformBlockBinding) \
+	EnumMacro(PFNGLUNIFORMMATRIX2FVPROC,glUniformMatrix2fv) \
+	EnumMacro(PFNGLUNIFORMMATRIX3FVPROC,glUniformMatrix3fv) \
+	EnumMacro(PFNGLUNIFORMMATRIX4FVPROC,glUniformMatrix4fv) \
+	EnumMacro(PFNGLVALIDATEPROGRAMPROC,glValidateProgram) \
+	EnumMacro(PFNGLVERTEXATTRIB1DPROC,glVertexAttrib1d) \
+	EnumMacro(PFNGLVERTEXATTRIB1DVPROC,glVertexAttrib1dv) \
+	EnumMacro(PFNGLVERTEXATTRIB1FPROC,glVertexAttrib1f) \
+	EnumMacro(PFNGLVERTEXATTRIB1FVPROC,glVertexAttrib1fv) \
+	EnumMacro(PFNGLVERTEXATTRIB1SPROC,glVertexAttrib1s) \
+	EnumMacro(PFNGLVERTEXATTRIB1SVPROC,glVertexAttrib1sv) \
+	EnumMacro(PFNGLVERTEXATTRIB2DPROC,glVertexAttrib2d) \
+	EnumMacro(PFNGLVERTEXATTRIB2DVPROC,glVertexAttrib2dv) \
+	EnumMacro(PFNGLVERTEXATTRIB2FPROC,glVertexAttrib2f) \
+	EnumMacro(PFNGLVERTEXATTRIB2FVPROC,glVertexAttrib2fv) \
+	EnumMacro(PFNGLVERTEXATTRIB2SPROC,glVertexAttrib2s) \
+	EnumMacro(PFNGLVERTEXATTRIB2SVPROC,glVertexAttrib2sv) \
+	EnumMacro(PFNGLVERTEXATTRIB3DPROC,glVertexAttrib3d) \
+	EnumMacro(PFNGLVERTEXATTRIB3DVPROC,glVertexAttrib3dv) \
+	EnumMacro(PFNGLVERTEXATTRIB3FPROC,glVertexAttrib3f) \
+	EnumMacro(PFNGLVERTEXATTRIB3FVPROC,glVertexAttrib3fv) \
+	EnumMacro(PFNGLVERTEXATTRIB3SPROC,glVertexAttrib3s) \
+	EnumMacro(PFNGLVERTEXATTRIB3SVPROC,glVertexAttrib3sv) \
+	EnumMacro(PFNGLVERTEXATTRIB4NBVPROC,glVertexAttrib4Nbv) \
+	EnumMacro(PFNGLVERTEXATTRIB4NIVPROC,glVertexAttrib4Niv) \
+	EnumMacro(PFNGLVERTEXATTRIB4NSVPROC,glVertexAttrib4Nsv) \
+	EnumMacro(PFNGLVERTEXATTRIB4NUBPROC,glVertexAttrib4Nub) \
+	EnumMacro(PFNGLVERTEXATTRIB4NUBVPROC,glVertexAttrib4Nubv) \
+	EnumMacro(PFNGLVERTEXATTRIB4NUIVPROC,glVertexAttrib4Nuiv) \
+	EnumMacro(PFNGLVERTEXATTRIB4NUSVPROC,glVertexAttrib4Nusv) \
+	EnumMacro(PFNGLVERTEXATTRIB4BVPROC,glVertexAttrib4bv) \
+	EnumMacro(PFNGLVERTEXATTRIB4DPROC,glVertexAttrib4d) \
+	EnumMacro(PFNGLVERTEXATTRIB4DVPROC,glVertexAttrib4dv) \
+	EnumMacro(PFNGLVERTEXATTRIB4FPROC,glVertexAttrib4f) \
+	EnumMacro(PFNGLVERTEXATTRIB4FVPROC,glVertexAttrib4fv) \
+	EnumMacro(PFNGLVERTEXATTRIB4IVPROC,glVertexAttrib4iv) \
+	EnumMacro(PFNGLVERTEXATTRIB4SPROC,glVertexAttrib4s) \
+	EnumMacro(PFNGLVERTEXATTRIB4SVPROC,glVertexAttrib4sv) \
+	EnumMacro(PFNGLVERTEXATTRIB4UBVPROC,glVertexAttrib4ubv) \
+	EnumMacro(PFNGLVERTEXATTRIB4UIVPROC,glVertexAttrib4uiv) \
+	EnumMacro(PFNGLVERTEXATTRIB4USVPROC,glVertexAttrib4usv) \
+	EnumMacro(PFNGLVERTEXATTRIBI4IVPROC,glVertexAttribI4iv) \
+	EnumMacro(PFNGLVERTEXATTRIBI4UIVPROC,glVertexAttribI4uiv) \
+	EnumMacro(PFNGLVERTEXATTRIBI4SVPROC,glVertexAttribI4sv) \
+	EnumMacro(PFNGLVERTEXATTRIBI4USVPROC,glVertexAttribI4usv) \
+	EnumMacro(PFNGLVERTEXATTRIBI4BVPROC,glVertexAttribI4bv) \
+	EnumMacro(PFNGLVERTEXATTRIBI4UBVPROC,glVertexAttribI4ubv) \
+	EnumMacro(PFNGLVERTEXATTRIBPOINTERPROC,glVertexAttribPointer) \
+	EnumMacro(PFNGLVERTEXATTRIBIPOINTERPROC,glVertexAttribIPointer) \
+	EnumMacro(PFNGLUNIFORMMATRIX2X3FVPROC,glUniformMatrix2x3fv) \
+	EnumMacro(PFNGLUNIFORMMATRIX3X2FVPROC,glUniformMatrix3x2fv) \
+	EnumMacro(PFNGLUNIFORMMATRIX2X4FVPROC,glUniformMatrix2x4fv) \
+	EnumMacro(PFNGLUNIFORMMATRIX4X2FVPROC,glUniformMatrix4x2fv) \
+	EnumMacro(PFNGLUNIFORMMATRIX3X4FVPROC,glUniformMatrix3x4fv) \
+	EnumMacro(PFNGLUNIFORMMATRIX4X3FVPROC,glUniformMatrix4x3fv) \
+	EnumMacro(PFNGLISRENDERBUFFERPROC,glIsRenderbuffer) \
+	EnumMacro(PFNGLBINDRENDERBUFFERPROC,glBindRenderbuffer) \
+	EnumMacro(PFNGLDELETERENDERBUFFERSPROC,glDeleteRenderbuffers) \
+	EnumMacro(PFNGLGENRENDERBUFFERSPROC,glGenRenderbuffers) \
+	EnumMacro(PFNGLRENDERBUFFERSTORAGEPROC,glRenderbufferStorage) \
+	EnumMacro(PFNGLGETRENDERBUFFERPARAMETERIVPROC,glGetRenderbufferParameteriv) \
+	EnumMacro(PFNGLISFRAMEBUFFERPROC,glIsFramebuffer) \
+	EnumMacro(PFNGLBINDFRAMEBUFFERPROC,glBindFramebuffer) \
+	EnumMacro(PFNGLDELETEFRAMEBUFFERSPROC,glDeleteFramebuffers) \
+	EnumMacro(PFNGLGENFRAMEBUFFERSPROC,glGenFramebuffers) \
+	EnumMacro(PFNGLCHECKFRAMEBUFFERSTATUSPROC,glCheckFramebufferStatus) \
+	EnumMacro(PFNGLFRAMEBUFFERTEXTURE1DPROC,glFramebufferTexture1D) \
+	EnumMacro(PFNGLFRAMEBUFFERTEXTURE2DPROC,glFramebufferTexture2D) \
+	EnumMacro(PFNGLFRAMEBUFFERTEXTURE3DPROC,glFramebufferTexture3D) \
+	EnumMacro(PFNGLFRAMEBUFFERTEXTUREPROC,glFramebufferTexture) \
+	EnumMacro(PFNGLFRAMEBUFFERRENDERBUFFERPROC,glFramebufferRenderbuffer) \
+	EnumMacro(PFNGLGETFRAMEBUFFERATTACHMENTPARAMETERIVPROC,glGetFramebufferAttachmentParameteriv) \
+	EnumMacro(PFNGLGENERATEMIPMAPPROC,glGenerateMipmap) \
+	EnumMacro(PFNGLBLITFRAMEBUFFERPROC,glBlitFramebuffer) \
+	EnumMacro(PFNGLRENDERBUFFERSTORAGEMULTISAMPLEPROC,glRenderbufferStorageMultisample) \
+	EnumMacro(PFNGLFRAMEBUFFERTEXTURELAYERPROC,glFramebufferTextureLayer) \
+	EnumMacro(PFNGLMAPBUFFERRANGEPROC,glMapBufferRange) \
+	EnumMacro(PFNGLFLUSHMAPPEDBUFFERRANGEPROC,glFlushMappedBufferRange) \
+	EnumMacro(PFNGLVERTEXATTRIBDIVISORPROC,glVertexAttribDivisor) \
+	EnumMacro(PFNGLDRAWARRAYSINSTANCEDPROC,glDrawArraysInstanced) \
+	EnumMacro(PFNGLDRAWELEMENTSINSTANCEDPROC,glDrawElementsInstanced) \
+	EnumMacro(PFNGLGETSTRINGIPROC,glGetStringi) \
+	EnumMacro(PFNGLGENVERTEXARRAYSPROC,glGenVertexArrays) \
+	EnumMacro(PFNGLDELETEVERTEXARRAYSPROC,glDeleteVertexArrays) \
+	EnumMacro(PFNGLBINDVERTEXARRAYPROC,glBindVertexArray) \
+	EnumMacro(PFNGLCOPYBUFFERSUBDATAPROC,glCopyBufferSubData) \
+	EnumMacro(PFNGLTEXBUFFERPROC,glTexBuffer) \
+	EnumMacro(PFNGLTEXBUFFERRANGEPROC,glTexBufferRange) \
+	EnumMacro(PFNGLTEXIMAGE2DMULTISAMPLEPROC,glTexImage2DMultisample) \
+	EnumMacro(PFNGLQUERYCOUNTERPROC, glQueryCounter)\
+	EnumMacro(PFNGLISSYNCPROC, glIsSync)\
+	EnumMacro(PFNGLDELETESYNCPROC, glDeleteSync)\
+	EnumMacro(PFNGLGETQUERYOBJECTUI64VPROC, glGetQueryObjectui64v)\
+	EnumMacro(PFNGLFENCESYNCPROC, glFenceSync)\
+	EnumMacro(PFNGLGETSYNCIVPROC, glGetSynciv)\
+	EnumMacro(PFNGLCLIENTWAITSYNCPROC, glClientWaitSync)\
+	EnumMacro(PFNGLBINDBUFFERRANGEPROC, glBindBufferRange) \
+	EnumMacro(PFNGLGETPROGRAMRESOURCEINDEXPROC, glGetProgramResourceIndex)
+
+#define ENUM_GL_ENTRYPOINTS_OPTIONAL(EnumMacro) \
+	EnumMacro(PFNGLCLIPCONTROLPROC,glClipControl) \
+	EnumMacro(PFNGLDEBUGMESSAGECALLBACKARBPROC,glDebugMessageCallbackARB) \
+	EnumMacro(PFNGLDEBUGMESSAGECONTROLARBPROC,glDebugMessageControlARB) \
+	EnumMacro(PFNGLDEBUGMESSAGECALLBACKAMDPROC,glDebugMessageCallbackAMD) \
+	EnumMacro(PFNGLDEBUGMESSAGEENABLEAMDPROC,glDebugMessageEnableAMD) \
+	EnumMacro(PFNGLGETACTIVEUNIFORMSIVPROC,glGetActiveUniformsiv) \
+	EnumMacro(PFNGLGETVERTEXATTRIBIUIVPROC,glGetVertexAttribIuiv) \
+	EnumMacro(PFNGLGETACTIVEUNIFORMNAMEPROC,glGetActiveUniformName) \
+	EnumMacro(PFNGLGETUNIFORMUIVPROC,glGetUniformuiv) \
+	EnumMacro(PFNGLGETACTIVEUNIFORMBLOCKIVPROC,glGetActiveUniformBlockiv) \
+	EnumMacro(PFNGLGETBUFFERPARAMETERI64VPROC,glGetBufferParameteri64v) \
+	EnumMacro(PFNGLGETACTIVEUNIFORMBLOCKNAMEPROC,glGetActiveUniformBlockName) \
+	EnumMacro(PFNGLGETSAMPLERPARAMETERFVPROC,glGetSamplerParameterfv) \
+	EnumMacro(PFNGLGETSAMPLERPARAMETERIVPROC,glGetSamplerParameteriv) \
+	EnumMacro(PFNGLDISPATCHCOMPUTEPROC, glDispatchCompute) \
+	EnumMacro(PFNGLDISPATCHCOMPUTEINDIRECTPROC, glDispatchComputeIndirect) \
+	EnumMacro(PFNGLBINDIMAGETEXTUREPROC, glBindImageTexture) \
+	EnumMacro(PFNGLMEMORYBARRIERPROC, glMemoryBarrier) \
+	EnumMacro(PFNGLBLENDEQUATIONIPROC, glBlendEquationi) \
+	EnumMacro(PFNGLBLENDEQUATIONSEPARATEIPROC, glBlendEquationSeparatei) \
+	EnumMacro(PFNGLBLENDFUNCIPROC, glBlendFunci) \
+	EnumMacro(PFNGLBLENDFUNCSEPARATEIPROC, glBlendFuncSeparatei)\
+	EnumMacro(PFNGLCLEARBUFFERDATAPROC, glClearBufferData)\
+	EnumMacro(PFNGLCLEARBUFFERSUBDATAPROC, glClearBufferSubData)\
+	EnumMacro(PFNGLPUSHDEBUGGROUPPROC, glPushDebugGroup)\
+	EnumMacro(PFNGLPOPDEBUGGROUPPROC, glPopDebugGroup)\
+	EnumMacro(PFNGLOBJECTLABELPROC, glObjectLabel)\
+	EnumMacro(PFNGLOBJECTPTRLABELPROC, glObjectPtrLabel)\
+	EnumMacro(PFNGLPATCHPARAMETERIPROC, glPatchParameteri)\
+	EnumMacro(PFNGLBINDVERTEXBUFFERPROC, glBindVertexBuffer)\
+	EnumMacro(PFNGLVERTEXATTRIBFORMATPROC, glVertexAttribFormat)\
+	EnumMacro(PFNGLVERTEXATTRIBIFORMATPROC, glVertexAttribIFormat)\
+	EnumMacro(PFNGLVERTEXATTRIBBINDINGPROC, glVertexAttribBinding)\
+	EnumMacro(PFNGLVERTEXBINDINGDIVISORPROC, glVertexBindingDivisor)\
+	EnumMacro(PFNGLCOPYIMAGESUBDATAPROC, glCopyImageSubData)\
+	EnumMacro(PFNGLTEXSTORAGE1DPROC, glTexStorage1D)\
+	EnumMacro(PFNGLTEXSTORAGE2DPROC, glTexStorage2D)\
+	EnumMacro(PFNGLTEXSTORAGE3DPROC, glTexStorage3D)\
+	EnumMacro(PFNGLBUFFERSTORAGEPROC, glBufferStorage)\
+	EnumMacro(PFNGLTEXTUREVIEWPROC, glTextureView)\
+	EnumMacro(PFNGLTEXSTORAGE2DMULTISAMPLEPROC, glTexStorage2DMultisample)\
+	EnumMacro(PFNGLDRAWELEMENTSINDIRECTPROC, glDrawElementsIndirect)\
+	EnumMacro(PFNGLDRAWARRAYSINDIRECTPROC, glDrawArraysIndirect)\
+	EnumMacro(PFNGLDEPTHBOUNDSEXTPROC, glDepthBoundsEXT)\
+	EnumMacro(PFNGLPROGRAMPARAMETERIPROC, glProgramParameteri)\
+	EnumMacro(PFNGLUSEPROGRAMSTAGESPROC, glUseProgramStages)\
+	EnumMacro(PFNGLBINDPROGRAMPIPELINEPROC, glBindProgramPipeline)\
+	EnumMacro(PFNGLDELETEPROGRAMPIPELINESPROC, glDeleteProgramPipelines)\
+	EnumMacro(PFNGLGENPROGRAMPIPELINESPROC, glGenProgramPipelines)\
+	EnumMacro(PFNGLPROGRAMUNIFORM1IPROC, glProgramUniform1i)\
+	EnumMacro(PFNGLPROGRAMUNIFORM4IVPROC, glProgramUniform4iv)\
+	EnumMacro(PFNGLPROGRAMUNIFORM4FVPROC, glProgramUniform4fv)\
+	EnumMacro(PFNGLPROGRAMUNIFORM4UIVPROC, glProgramUniform4uiv)\
+	EnumMacro(PFNGLGETPROGRAMPIPELINEIVPROC, glGetProgramPipelineiv)\
+	EnumMacro(PFNGLVALIDATEPROGRAMPIPELINEPROC, glValidateProgramPipeline)\
+	EnumMacro(PFNGLGETPROGRAMPIPELINEINFOLOGPROC, glGetProgramPipelineInfoLog)\
+	EnumMacro(PFNGLISPROGRAMPIPELINEPROC, glIsProgramPipeline)
+
+/** List of all OpenGL entry points. */
+#define ENUM_GL_ENTRYPOINTS_ALL(EnumMacro) \
+	ENUM_GL_ENTRYPOINTS(EnumMacro) \
+	ENUM_GL_ENTRYPOINTS_OPTIONAL(EnumMacro)
+
+/** Declare all GL functions. */
+#define DECLARE_GL_ENTRYPOINTS(Type,Func) extern Type OPENGLDRV_API Func;
+
+// We need to make pointer names different from GL functions otherwise we may end up getting
+// addresses of those symbols when looking for extensions.
+namespace GLFuncPointers
+{
+	ENUM_GL_ENTRYPOINTS_ALL(DECLARE_GL_ENTRYPOINTS);
+};
+
+#undef DECLARE_GL_ENTRYPOINTS
+
+// this using is needed since the rest of code uses plain GL names
+using namespace GLFuncPointers;
diff --git a/Engine/Source/Runtime/OpenGLDrv/Private/Linux/LinuxOpenGLPlatform.cpp b/Engine/Source/Runtime/OpenGLDrv/Private/Linux/LinuxOpenGLPlatform.cpp
new file mode 100644
--- /dev/null
+++ b/Engine/Source/Runtime/OpenGLDrv/Private/Linux/LinuxOpenGLPlatform.cpp
@@ -0,0 +1,994 @@
+// Copyright Epic Games, Inc. All Rights Reserved.
+
+#include "LinuxOpenGLPlatform.h"
+
+#include "Misc/ScopeLock.h"
+#include "OpenGLDrv.h"
+#include "SDL.h"
+#include "OpenGLDrvPrivate.h"
+#include "ComponentReregisterContext.h"
+#include "Linux/LinuxPlatformApplicationMisc.h"
+#include "GenericPlatform/GenericPlatformFramePacer.h"
+#include "RHIUtilities.h"
+
+typedef SDL_Window* SDL_HWindow;
+typedef SDL_GLContext	SDL_HGLContext;
+
+/*------------------------------------------------------------------------------
+	OpenGL context management.
+------------------------------------------------------------------------------*/
+void Linux_ContextMakeCurrent(SDL_HWindow hWnd, SDL_HGLContext hGLDC)
+{
+	GLint Result = SDL_GL_MakeCurrent(hWnd, hGLDC);
+	if (Result != 0)
+	{
+		// this is a warning and not error, since Slate sometimes destroys windows before
+		// releasing RHI resources associated with them. This code can result in leaks - proper resolution is tracked as UE-7388
+		FString SdlError(UTF8_TO_TCHAR(SDL_GetError()));
+		UE_LOG(LogLinux, Warning, TEXT("SDL_GL_MakeCurrent() failed, SDL error: '%s'"), *SdlError);
+	}
+}
+
+SDL_HGLContext Linux_GetCurrentContext()
+{
+	return SDL_GL_GetCurrentContext();
+}
+
+/** Platform specific OpenGL context. */
+struct FPlatformOpenGLContext
+{
+	SDL_HWindow		hWnd;
+	SDL_HGLContext	hGLContext;		//	this is a (void*) pointer
+
+	bool bReleaseWindowOnDestroy;
+	int32 SyncInterval;
+	GLuint	ViewportFramebuffer;
+	GLuint	VertexArrayObject;	// one has to be generated and set for each context (OpenGL 3.2 Core requirements)
+};
+
+
+class FScopeContext
+{
+public:
+	FScopeContext(FPlatformOpenGLContext* Context)
+	{
+		check(Context);
+		hPreWnd = SDL_GL_GetCurrentWindow();
+		hPreGLContext = SDL_GL_GetCurrentContext();
+
+		bSameDCAndContext = (hPreGLContext == Context->hGLContext);
+
+		if (!bSameDCAndContext)
+		{
+			if (hPreGLContext)
+			{
+				glFlush();
+			}
+			// no need to glFlush() on Windows, it does flush by itself before switching contexts
+
+			Linux_ContextMakeCurrent(Context->hWnd, Context->hGLContext);
+		}
+	}
+
+	~FScopeContext(void)
+	{
+		if (!bSameDCAndContext)
+		{
+			glFlush();	// not needed on Windows, it does flush by itself before switching contexts
+			if (hPreGLContext)
+			{
+				Linux_ContextMakeCurrent(hPreWnd, hPreGLContext);
+			}
+			else
+			{
+				Linux_ContextMakeCurrent(NULL, NULL);
+			}
+		}
+
+	}
+
+private:
+	SDL_HWindow			hPreWnd;
+	SDL_HGLContext		hPreGLContext;		//	this is a pointer, (void*)
+	bool				bSameDCAndContext;
+};
+
+
+
+void Linux_DeleteQueriesForCurrentContext(SDL_HGLContext hGLContext);
+
+/**
+ * Create a dummy window used to construct OpenGL contexts.
+ */
+void Linux_PlatformCreateDummyGLWindow(FPlatformOpenGLContext* OutContext)
+{
+	static bool bInitializedWindowClass = false;
+
+	// Create a dummy window.
+	SDL_HWindow DummyWindow = SDL_CreateWindow(NULL,
+		0, 0, 1, 1,
+		SDL_WINDOW_OPENGL | SDL_WINDOW_BORDERLESS | SDL_WINDOW_HIDDEN | SDL_WINDOW_SKIP_TASKBAR);
+	if (DummyWindow == nullptr)
+	{
+		FString SdlError(UTF8_TO_TCHAR(SDL_GetError()));
+
+		FText ErrorMessage = FText::Format(NSLOCTEXT("Renderer", "LinuxCannotCreatePlatformWindow", "Cannot create OpenGL-enabled SDL window. SDL error: '{0}'."), FText::FromString(SdlError));
+		FPlatformMisc::MessageBoxExt(EAppMsgType::Ok,
+			*ErrorMessage.ToString(),
+			*(NSLOCTEXT("Renderer", "LinuxCannotCreatePlatformWindowTitle", "Cannot create SDL window.").ToString()));
+		FPlatformMisc::RequestExit(true, TEXT("Linux_PlatformCreateDummyGLWindow"));
+		// unreachable
+		return;
+	}
+	else
+	{
+		SDL_SetWindowTitle(DummyWindow, "UnrealEditor Dummy GL window");
+	}
+
+	OutContext->hWnd = DummyWindow;
+	OutContext->bReleaseWindowOnDestroy = true;
+}
+
+/**
+ * Enable/Disable debug context from the commandline
+ */
+bool Linux_PlatformOpenGLDebugCtx()
+{
+	return IsOGLDebugOutputEnabled();
+}
+
+/**
+ * Create a core profile OpenGL context.
+ */
+void Linux_PlatformCreateOpenGLContextCore(FPlatformOpenGLContext* OutContext)
+{
+	check(OutContext);
+
+	SDL_HWindow prevWindow = SDL_GL_GetCurrentWindow();
+	SDL_HGLContext prevContext = SDL_GL_GetCurrentContext();
+
+	OutContext->SyncInterval = -1;	// invalid value to enforce setup on first buffer swap
+	OutContext->ViewportFramebuffer = 0;
+
+	OutContext->hGLContext = SDL_GL_CreateContext(OutContext->hWnd);
+	if (OutContext->hGLContext == nullptr)
+	{
+		FString SdlError(UTF8_TO_TCHAR(SDL_GetError()));
+
+		// ignore errors getting version, it will be clear from the logs
+		int OpenGLMajorVersion = -1;
+		int OpenGLMinorVersion = -1;
+		SDL_GL_GetAttribute(SDL_GL_CONTEXT_MAJOR_VERSION, &OpenGLMajorVersion);
+		SDL_GL_GetAttribute(SDL_GL_CONTEXT_MINOR_VERSION, &OpenGLMinorVersion);
+
+		UE_LOG(LogInit, Error, TEXT("Linux_PlatformCreateOpenGLContextCore - Could not create OpenGL %d.%d context, SDL error: '%s'"),
+			OpenGLMajorVersion, OpenGLMinorVersion,
+			*SdlError
+		);
+		// unreachable
+		return;
+	}
+
+	SDL_GL_MakeCurrent(prevWindow, prevContext);
+}
+
+void PlatformReleaseOpenGLContext(FPlatformOpenGLDevice* Device, FPlatformOpenGLContext* Context);
+
+extern void OnQueryInvalidation(void);
+
+/** Platform specific OpenGL device. */
+struct FPlatformOpenGLDevice
+{
+	FPlatformOpenGLContext	SharedContext;
+	FPlatformOpenGLContext	RenderingContext;
+	int32					NumUsedContexts;
+
+	/** Guards against operating on viewport contexts from more than one thread at the same time. */
+	FCriticalSection* ContextUsageGuard;
+
+	FPlatformOpenGLDevice() : NumUsedContexts(0)
+	{
+		extern void InitDebugContext();
+
+		ContextUsageGuard = new FCriticalSection;
+
+		verifyf(SDL_GL_SetAttribute(SDL_GL_SHARE_WITH_CURRENT_CONTEXT, 0) == 0, TEXT("SDL_GL_SetAttribute(SDL_GL_SHARE_WITH_CURRENT_CONTEXT, 0) failed: %s\n."), UTF8_TO_TCHAR(SDL_GetError()));
+
+		Linux_PlatformCreateDummyGLWindow(&SharedContext);
+		Linux_PlatformCreateOpenGLContextCore(&SharedContext);
+
+		if (SharedContext.hGLContext == nullptr)
+		{
+			FPlatformMisc::MessageBoxExt(EAppMsgType::Ok,
+				*(NSLOCTEXT("Renderer", "LinuxInsufficientDriversText", "Cannot create OpenGL context. Check that the drivers and hardware support at least OpenGL 4.3 (or re-run with -opengl3)").ToString()),
+				*(NSLOCTEXT("Renderer", "LinuxInsufficientDriversTitle", "Insufficient drivers or hardware").ToString()));
+			FPlatformMisc::RequestExit(true, TEXT("FPlatformOpenGLDevice()"));
+			// unreachable
+			return;
+		}
+
+		{
+			FScopeContext ScopeContext(&SharedContext);
+			InitDebugContext();
+			glGenVertexArrays(1, &SharedContext.VertexArrayObject);
+			glBindVertexArray(SharedContext.VertexArrayObject);
+			InitDefaultGLContextState();
+		}
+
+		verifyf(SDL_GL_SetAttribute(SDL_GL_SHARE_WITH_CURRENT_CONTEXT, 1) == 0, TEXT("SDL_GL_SetAttribute(SDL_GL_SHARE_WITH_CURRENT_CONTEXT, 1) failed: %s\n."), UTF8_TO_TCHAR(SDL_GetError()));
+
+		Linux_ContextMakeCurrent(SharedContext.hWnd, SharedContext.hGLContext);
+
+		Linux_PlatformCreateDummyGLWindow(&RenderingContext);
+		Linux_PlatformCreateOpenGLContextCore(&RenderingContext);
+
+		check(RenderingContext.hGLContext);
+
+		{
+			FScopeContext ScopeContext(&RenderingContext);
+			InitDebugContext();
+			glGenVertexArrays(1, &RenderingContext.VertexArrayObject);
+			glBindVertexArray(RenderingContext.VertexArrayObject);
+			InitDefaultGLContextState();
+		}
+	}
+
+	~FPlatformOpenGLDevice()
+	{
+		check(NumUsedContexts == 0);
+
+		Linux_ContextMakeCurrent(NULL, NULL);
+
+		OnQueryInvalidation();
+		PlatformReleaseOpenGLContext(this, &RenderingContext);
+		PlatformReleaseOpenGLContext(this, &SharedContext);
+
+		delete ContextUsageGuard;
+	}
+};
+
+FPlatformOpenGLDevice* PlatformCreateOpenGLDevice()
+{
+	return new FPlatformOpenGLDevice;
+}
+
+bool PlatformCanEnableGPUCapture()
+{
+	return false;
+}
+
+void PlatformDestroyOpenGLDevice(FPlatformOpenGLDevice* Device)
+{
+	delete Device;
+}
+
+FPlatformOpenGLContext* PlatformGetOpenGLRenderingContext(FPlatformOpenGLDevice* Device)
+{
+	return &Device->RenderingContext;
+}
+
+/**
+ * Create an OpenGL context.
+ */
+FPlatformOpenGLContext* PlatformCreateOpenGLContext(FPlatformOpenGLDevice* Device, void* InWindowHandle)
+{
+	check(InWindowHandle);
+	FPlatformOpenGLContext* Context = new FPlatformOpenGLContext;
+
+	Context->hWnd = (SDL_HWindow)InWindowHandle;
+	Context->bReleaseWindowOnDestroy = false;
+
+	check(Device->SharedContext.hGLContext)
+	{
+		FScopeContext Scope(&(Device->SharedContext));
+		verifyf(SDL_GL_SetAttribute(SDL_GL_SHARE_WITH_CURRENT_CONTEXT, 1) == 0, TEXT("SDL_GL_SetAttribute(SDL_GL_SHARE_WITH_CURRENT_CONTEXT, 1) failed: %s\n."), UTF8_TO_TCHAR(SDL_GetError()));
+		Linux_PlatformCreateOpenGLContextCore(Context);
+	}
+
+	check(Context->hGLContext);
+	{
+		FScopeContext Scope(Context);
+		InitDefaultGLContextState();
+	}
+
+	return Context;
+}
+
+/**
+ * Release an OpenGL context.
+ */
+void PlatformReleaseOpenGLContext(FPlatformOpenGLDevice* Device, FPlatformOpenGLContext* Context)
+{
+	check(Context && Context->hGLContext);
+
+	{
+		FScopeLock ScopeLock(Device->ContextUsageGuard);
+		{
+			FScopeContext ScopeContext(Context);
+
+			Linux_DeleteQueriesForCurrentContext(Context->hGLContext);
+			glBindVertexArray(0);
+			glDeleteVertexArrays(1, &Context->VertexArrayObject);
+
+			if (Context->ViewportFramebuffer)
+			{
+				glDeleteFramebuffers(1, &Context->ViewportFramebuffer);	// this can be done from any context shared with ours, as long as it's not nil.
+				Context->ViewportFramebuffer = 0;
+			}
+		}
+
+		SDL_GL_DeleteContext(Context->hGLContext);
+		Context->hGLContext = NULL;
+	}
+
+	check(Context->hWnd);
+
+	if (Context->bReleaseWindowOnDestroy)
+	{
+		SDL_DestroyWindow(Context->hWnd);
+	}
+
+	Context->hWnd = NULL;
+}
+
+/**
+ * Destroy an OpenGL context.
+ */
+void PlatformDestroyOpenGLContext(FPlatformOpenGLDevice* Device, FPlatformOpenGLContext* Context)
+{
+	PlatformReleaseOpenGLContext(Device, Context);
+	delete Context;
+}
+
+void* PlatformGetWindow(FPlatformOpenGLContext* Context, void** AddParam)
+{
+	check(Context && Context->hWnd);
+
+	if (AddParam)
+	{
+		*AddParam = &Context->hGLContext;
+	}
+
+	return (void*)&Context->hWnd;
+}
+
+const TCHAR* PlatformDescribeSyncInterval(int32 SyncInterval)
+{
+	switch (SyncInterval)
+	{
+	case -1:
+		return TEXT("Late swap");
+	case 0:
+		return TEXT("Immediate");
+	case 1:
+		return TEXT("Synchronized with retrace");
+	default: break;
+	}
+	return TEXT("Unknown");
+}
+
+/**
+ * Main function for transferring data to on-screen buffers.
+ * On Windows it temporarily switches OpenGL context, on Mac only context's output view.
+ */
+bool PlatformBlitToViewport(FPlatformOpenGLDevice* Device,
+	const FOpenGLViewport& Viewport,
+	uint32 BackbufferSizeX,
+	uint32 BackbufferSizeY,
+	bool bPresent,
+	bool bLockToVsync)
+{
+	FPlatformOpenGLContext* const Context = Viewport.GetGLContext();
+
+	check(Context && Context->hWnd);
+
+	FScopeLock ScopeLock(Device->ContextUsageGuard);
+
+	int32 SyncInterval = (int32)RHIGetSyncInterval();
+
+	{
+		FScopeContext ScopeContext(Context);
+
+		if (Viewport.GetCustomPresent())
+		{
+			glDisable(GL_FRAMEBUFFER_SRGB);
+			bool bShouldPresent = Viewport.GetCustomPresent()->Present(SyncInterval);
+			glEnable(GL_FRAMEBUFFER_SRGB);
+			if (!bShouldPresent)
+			{
+				return false;
+			}
+		}
+
+		glBindFramebuffer(GL_DRAW_FRAMEBUFFER, 0);
+		glDrawBuffer(GL_BACK);
+		glBindFramebuffer(GL_READ_FRAMEBUFFER, Context->ViewportFramebuffer);
+		glReadBuffer(GL_COLOR_ATTACHMENT0);
+		glDisable(GL_FRAMEBUFFER_SRGB);
+
+		int WinW = 0;
+		int WinH = 0;
+		SDL_GL_GetDrawableSize(Context->hWnd, &WinW, &WinH);
+		GLenum BlitFilter;
+		GLint DestX0, DestY0, DestX1, DestY1;
+
+		if (WinH == 0 || WinW == 0)
+		{
+			// Nothing to blit
+			return false;
+		}
+
+		if ((WinW == BackbufferSizeX) && (WinH == BackbufferSizeY))
+		{
+			// We match up. We're probably in windowed mode, or an exact
+			//  match for FULLSCREEN_DESKTOP mode. Use a NEAREST blit and
+			//  don't clear the window system's framebuffer first.
+			BlitFilter = GL_NEAREST;
+			DestX0 = 0;
+			DestY0 = BackbufferSizeY;  // flip vertically.
+			DestX1 = BackbufferSizeX;
+			DestY1 = 0;
+		}
+		else
+		{
+			// we need to scale to match the size of the window system's
+			// framebuffer, so scale linearly, and adjust for letterboxing.
+			BlitFilter = GL_LINEAR;
+
+			const uint32 w = BackbufferSizeX;
+			const uint32 h = BackbufferSizeY;
+			const float WantedAspect = (w > h) ? (((float)w) / ((float)h)) : (((float)h) / ((float)w));
+			const float PhysicalAspect = (((float)WinW) / ((float)WinH));
+
+			bool bMustClear;  // have to clear the window framebuffer if letterboxing.
+			if (PhysicalAspect == WantedAspect)  // Perfect aspect ratio; no letterboxing needed?
+			{
+				bMustClear = false;
+				DestX0 = 0;
+				DestY0 = WinH;  // flip vertically.
+				DestX1 = WinW;
+				DestY1 = 0;
+			}
+			else if (PhysicalAspect > WantedAspect)  // view is wider than wanted aspect?
+			{
+				bMustClear = true;
+				const float ScaledW = WinH * WantedAspect;
+				const float ScaledX = (((float)WinW) - ScaledW) / 2.0f;
+				DestX0 = ScaledX;
+				DestY0 = WinH;  // flip vertically.
+				DestX1 = ScaledX + ScaledW;
+				DestY1 = 0;
+			}
+			else  // view is taller than wanted aspect?
+			{
+				bMustClear = true;
+				const float ScaledH = WinW / WantedAspect;
+				const float ScaledY = (((float)WinH) - ScaledH) / 2.0f;
+				DestX0 = 0;
+				DestY0 = ScaledY + ScaledH;  // flip vertically.
+				DestX1 = WinW;
+				DestY1 = ScaledY;
+			}
+
+			// if the Steam Overlay is running, it might write garbage into our
+			//  letterbox area, so if we have a letterbox, clear the framebuffer.
+			if (bMustClear)
+			{
+				glClearColor(0.0f, 0.0f, 0.0f, 1.0f);
+				glClear(GL_COLOR_BUFFER_BIT);
+			}
+		}
+
+		// Get it to the window system's framebuffer.
+		glBlitFramebuffer(0, 0, BackbufferSizeX, BackbufferSizeY,
+			DestX0, DestY0, DestX1, DestY1,
+			GL_COLOR_BUFFER_BIT,
+			BlitFilter);
+
+		if (bPresent)
+		{
+			int32 RealSyncInterval = bLockToVsync ? SyncInterval : 0;
+
+			if (Context->SyncInterval != RealSyncInterval)
+			{
+				//	0 for immediate updates
+				//	1 for updates synchronized with the vertical retrace
+				//	-1 for late swap tearing; 
+
+				UE_LOG(LogLinux, Log, TEXT("Setting swap interval to '%s'"), PlatformDescribeSyncInterval(RealSyncInterval));
+				int SetSwapResult = SDL_GL_SetSwapInterval(RealSyncInterval);
+				// if late tearing is not supported, this needs to be retried with a valid value.
+				if (SetSwapResult == -1)
+				{
+					if (RealSyncInterval == -1)
+					{
+						// fallback to synchronized
+						int FallbackInterval = 1;
+						UE_LOG(LogLinux, Log, TEXT("Unable to set desired swap interval, falling back to '%s'"), PlatformDescribeSyncInterval(FallbackInterval));
+						SetSwapResult = SDL_GL_SetSwapInterval(FallbackInterval);
+					}
+				}
+
+				if (SetSwapResult == -1)
+				{
+					UE_LOG(LogLinux, Warning, TEXT("Unable to set desired swap interval '%s'"), PlatformDescribeSyncInterval(RealSyncInterval));
+				}
+
+				// even if we failed to set it, update the context value to prevent further attempts
+				Context->SyncInterval = RealSyncInterval;
+			}
+
+			SDL_GL_SwapWindow(Context->hWnd);
+
+			glEnable(GL_FRAMEBUFFER_SRGB);
+		}
+	}
+	return true;
+}
+
+void PlatformFlushIfNeeded()
+{
+	glFinish();
+}
+
+void PlatformRebindResources(FPlatformOpenGLDevice* Device)
+{
+	// @todo: Figure out if we need to rebind frame & renderbuffers after switching contexts
+}
+
+void PlatformRenderingContextSetup(FPlatformOpenGLDevice* Device)
+{
+	check(Device && Device->RenderingContext.hWnd && Device->RenderingContext.hGLContext);
+
+	if (Linux_GetCurrentContext())
+	{
+		glFlush();
+	}
+
+	Linux_ContextMakeCurrent(Device->RenderingContext.hWnd, Device->RenderingContext.hGLContext);
+}
+
+void PlatformSharedContextSetup(FPlatformOpenGLDevice* Device)
+{
+	check(Device && Device->SharedContext.hWnd && Device->SharedContext.hGLContext);
+
+	// no need to glFlush() on Windows, it does flush by itself before switching contexts
+	if (Linux_GetCurrentContext())
+	{
+		glFlush();
+	}
+
+	Linux_ContextMakeCurrent(Device->SharedContext.hWnd, Device->SharedContext.hGLContext);
+}
+
+
+void PlatformNULLContextSetup()
+{
+	if (Linux_GetCurrentContext())
+	{
+		glFlush();
+	}
+
+	Linux_ContextMakeCurrent(NULL, NULL);
+}
+
+
+
+/**
+ * Resize the GL context.
+ */
+void PlatformResizeGLContext(FPlatformOpenGLDevice* Device,
+	FPlatformOpenGLContext* Context,
+	uint32 SizeX, uint32 SizeY,
+	bool bFullscreen,
+	bool bWasFullscreen,
+	GLenum BackBufferTarget,
+	GLuint BackBufferResource)
+{
+	FScopeLock ScopeLock(Device->ContextUsageGuard);
+	{
+		FScopeContext ScopeContext(Context);
+
+		if (Context->ViewportFramebuffer == 0)
+		{
+			glGenFramebuffers(1, &Context->ViewportFramebuffer);
+		}
+		glBindFramebuffer(GL_FRAMEBUFFER, Context->ViewportFramebuffer);
+		glFramebufferTexture2D(GL_FRAMEBUFFER, GL_COLOR_ATTACHMENT0, BackBufferTarget, BackBufferResource, 0);
+		FOpenGL::CheckFrameBuffer();
+
+		glViewport(0, 0, SizeX, SizeY);
+
+		static GLfloat ZeroColor[4] = { 0.0f, 0.0f, 0.0f, 0.0f };
+		glClearBufferfv(GL_COLOR, 0, ZeroColor);
+	}
+
+
+	if (bFullscreen)
+	{
+		//	Deregister all components
+		//	this line will fix the geometry missing and color distortion bug on Linux/Nvidia machine
+		//	it detach all the components and re attach all the components
+		FGlobalComponentReregisterContext RecreateComponents;
+	}
+	else if (bWasFullscreen)
+	{
+		FGlobalComponentReregisterContext RecreateComponents;
+	}
+}
+
+void PlatformGetSupportedResolution(uint32& Width, uint32& Height)
+{
+	uint32 InitializedMode = false;
+	uint32 BestWidth = 0;
+	uint32 BestHeight = 0;
+	uint32 ModeIndex = 0;
+
+	SDL_DisplayMode DisplayMode;
+	FMemory::Memzero(&DisplayMode, sizeof(DisplayMode));
+
+	while (!SDL_GetDisplayMode(0, ModeIndex++, &DisplayMode))
+	{
+		bool IsEqualOrBetterWidth = FMath::Abs((int32)DisplayMode.w - (int32)Width) <= FMath::Abs((int32)BestWidth - (int32)Width);
+		bool IsEqualOrBetterHeight = FMath::Abs((int32)DisplayMode.h - (int32)Height) <= FMath::Abs((int32)BestHeight - (int32)Height);
+		if (!InitializedMode || (IsEqualOrBetterWidth && IsEqualOrBetterHeight))
+		{
+			BestWidth = DisplayMode.w;
+			BestHeight = DisplayMode.h;
+			InitializedMode = true;
+		}
+	}
+	check(InitializedMode);
+	Width = BestWidth;
+	Height = BestHeight;
+}
+
+
+bool PlatformGetAvailableResolutions(FScreenResolutionArray& Resolutions, bool bIgnoreRefreshRate)
+{
+	int32 MinAllowableResolutionX = 0;
+	int32 MinAllowableResolutionY = 0;
+	int32 MaxAllowableResolutionX = 10480;
+	int32 MaxAllowableResolutionY = 10480;
+	int32 MinAllowableRefreshRate = 0;
+	int32 MaxAllowableRefreshRate = 10480;
+
+	if (MaxAllowableResolutionX == 0)
+	{
+		MaxAllowableResolutionX = 10480;
+	}
+	if (MaxAllowableResolutionY == 0)
+	{
+		MaxAllowableResolutionY = 10480;
+	}
+	if (MaxAllowableRefreshRate == 0)
+	{
+		MaxAllowableRefreshRate = 10480;
+	}
+
+	uint32 ModeIndex = 0;
+	SDL_DisplayMode DisplayMode;
+	FMemory::Memzero(&DisplayMode, sizeof(SDL_DisplayMode));
+
+	while (!SDL_GetDisplayMode(0, ModeIndex++, &DisplayMode))
+	{
+		if (((int32)DisplayMode.w >= MinAllowableResolutionX) &&
+			((int32)DisplayMode.w <= MaxAllowableResolutionX) &&
+			((int32)DisplayMode.h >= MinAllowableResolutionY) &&
+			((int32)DisplayMode.h <= MaxAllowableResolutionY)
+			)
+		{
+			bool bAddIt = true;
+			if (bIgnoreRefreshRate == false)
+			{
+				if (((int32)DisplayMode.refresh_rate < MinAllowableRefreshRate) ||
+					((int32)DisplayMode.refresh_rate > MaxAllowableRefreshRate))
+				{
+					continue;
+				}
+			}
+			else
+			{
+				// See if it is in the list already
+				for (int32 CheckIndex = 0; CheckIndex < Resolutions.Num(); CheckIndex++)
+				{
+					FScreenResolutionRHI& CheckResolution = Resolutions[CheckIndex];
+					if ((CheckResolution.Width == DisplayMode.w) &&
+						(CheckResolution.Height == DisplayMode.h))
+					{
+						// Already in the list...
+						bAddIt = false;
+						break;
+					}
+				}
+			}
+
+			if (bAddIt)
+			{
+				// Add the mode to the list
+				int32 Temp2Index = Resolutions.AddZeroed();
+				FScreenResolutionRHI& ScreenResolution = Resolutions[Temp2Index];
+
+				ScreenResolution.Width = DisplayMode.w;
+				ScreenResolution.Height = DisplayMode.h;
+				ScreenResolution.RefreshRate = DisplayMode.refresh_rate;
+			}
+		}
+	}
+
+	return true;
+}
+
+void PlatformRestoreDesktopDisplayMode()
+{
+}
+
+
+bool PlatformInitOpenGL()
+{
+	static bool bInitialized = false;
+	static bool bOpenGLSupported = false;
+
+	if (!FLinuxPlatformApplicationMisc::InitSDL()) //	will not initialize more than once
+	{
+		UE_LOG(LogInit, Error, TEXT("PlatformInitOpenGL() : InitSDL() failed, cannot initialize OpenGL."));
+		// unreachable
+		return false;
+	}
+
+#if DO_CHECK
+	uint32 InitializedSubsystems = SDL_WasInit(SDL_INIT_EVERYTHING);
+	check(InitializedSubsystems & SDL_INIT_VIDEO);
+#endif // DO_CHECK
+
+	if (!bInitialized)
+	{
+		if (SDL_GL_LoadLibrary(NULL))
+		{
+			FPlatformMisc::MessageBoxExt(EAppMsgType::Ok,
+				*FString::Printf(TEXT("%s. SDL error: \"%s\""), *(NSLOCTEXT("Renderer", "LinuxCannotLoadLibGLText", "Unable to dynamically load libGL").ToString()), UTF8_TO_TCHAR(SDL_GetError())),
+				*(NSLOCTEXT("Renderer", "LinuxInsufficientDriversTitle", "Insufficient drivers or hardware").ToString()));
+			FPlatformMisc::RequestExit(true, TEXT("PlatformInitOpenGL"));
+			// unreachable
+			return false;
+		}
+
+		const int MajorVersion = 4;
+		const int MinorVersion = 3;
+
+		int DebugFlag = 0;
+
+		if (Linux_PlatformOpenGLDebugCtx())
+		{
+			DebugFlag = SDL_GL_CONTEXT_DEBUG_FLAG;
+		}
+
+		if (SDL_GL_SetAttribute(SDL_GL_CONTEXT_MAJOR_VERSION, MajorVersion) != 0)
+		{
+			UE_LOG(LogLinux, Fatal, TEXT("SDL_GL_SetAttribute(SDL_GL_CONTEXT_MAJOR_VERSION, %d) failed: %s"), MajorVersion, UTF8_TO_TCHAR(SDL_GetError()));
+		}
+
+		if (SDL_GL_SetAttribute(SDL_GL_CONTEXT_MINOR_VERSION, MinorVersion) != 0)
+		{
+			UE_LOG(LogLinux, Fatal, TEXT("SDL_GL_SetAttribute(SDL_GL_CONTEXT_MINOR_VERSION, %d) failed: %s"), MinorVersion, UTF8_TO_TCHAR(SDL_GetError()));
+		}
+
+		if (SDL_GL_SetAttribute(SDL_GL_CONTEXT_FLAGS, DebugFlag) != 0)
+		{
+			UE_LOG(LogLinux, Fatal, TEXT("SDL_GL_SetAttribute(SDL_GL_CONTEXT_FLAGS, %d) failed: %s"), DebugFlag, UTF8_TO_TCHAR(SDL_GetError()));
+		}
+
+		if (SDL_GL_SetAttribute(SDL_GL_CONTEXT_PROFILE_MASK, SDL_GL_CONTEXT_PROFILE_CORE) != 0)
+		{
+			UE_LOG(LogLinux, Fatal, TEXT("SDL_GL_SetAttribute(SDL_GL_CONTEXT_PROFILE_MASK, SDL_GL_CONTEXT_PROFILE_CORE) failed: %s"), UTF8_TO_TCHAR(SDL_GetError()));
+		}
+
+		if (FParse::Param(FCommandLine::Get(), TEXT("quad_buffer_stereo")))
+		{
+			if (SDL_GL_SetAttribute(SDL_GL_STEREO, 1) != 0)
+			{
+				UE_LOG(LogLinux, Fatal, TEXT("SDL_GL_SetAttribute(SDL_GL_STEREO, 1) failed: %s"), UTF8_TO_TCHAR(SDL_GetError()));
+			}
+		}
+
+		// Create a dummy context to verify opengl support.
+		FPlatformOpenGLContext DummyContext;
+		Linux_PlatformCreateDummyGLWindow(&DummyContext);
+		Linux_PlatformCreateOpenGLContextCore(&DummyContext);
+
+		if (DummyContext.hGLContext)
+		{
+			bOpenGLSupported = true;
+			Linux_ContextMakeCurrent(DummyContext.hWnd, DummyContext.hGLContext);
+		}
+		else
+		{
+			UE_LOG(LogRHI, Error, TEXT("OpenGL %d.%d not supported by driver"), MajorVersion, MinorVersion);
+		}
+
+		if (bOpenGLSupported)
+		{
+			// Initialize all entry points required by Unreal.
+#define GET_GL_ENTRYPOINTS(Type,Func) GLFuncPointers::Func = reinterpret_cast<Type>(SDL_GL_GetProcAddress(#Func));
+			ENUM_GL_ENTRYPOINTS(GET_GL_ENTRYPOINTS);
+			ENUM_GL_ENTRYPOINTS_OPTIONAL(GET_GL_ENTRYPOINTS);
+#undef GET_GL_ENTRYPOINTS
+
+			// Check that all of the entry points have been initialized.
+			bool bFoundAllEntryPoints = true;
+#define CHECK_GL_ENTRYPOINTS(Type,Func) if (Func == NULL) { bFoundAllEntryPoints = false; UE_LOG(LogRHI, Fatal, TEXT("Failed to find entry point for %s"), TEXT(#Func)); }
+			ENUM_GL_ENTRYPOINTS(CHECK_GL_ENTRYPOINTS);
+#undef CHECK_GL_ENTRYPOINTS
+			checkf(bFoundAllEntryPoints, TEXT("Failed to find all OpenGL entry points."));
+		}
+
+		// The dummy context can now be released.
+		if (DummyContext.hGLContext)
+		{
+			Linux_ContextMakeCurrent(NULL, NULL);
+			SDL_GL_DeleteContext(DummyContext.hGLContext);
+		}
+		check(DummyContext.bReleaseWindowOnDestroy);
+		SDL_DestroyWindow(DummyContext.hWnd);
+
+		bInitialized = true;
+	}
+
+	return bOpenGLSupported;
+}
+
+bool PlatformOpenGLContextValid()
+{
+	return (Linux_GetCurrentContext() != NULL);
+}
+
+int32 PlatformGlGetError()
+{
+	return glGetError();
+}
+
+EOpenGLCurrentContext PlatformOpenGLCurrentContext(FPlatformOpenGLDevice* Device)
+{
+	SDL_HGLContext hGLContext = Linux_GetCurrentContext();
+
+	if (LIKELY(hGLContext == Device->RenderingContext.hGLContext))	// most common case
+	{
+		return CONTEXT_Rendering;
+	}
+	else if (hGLContext == Device->SharedContext.hGLContext)
+	{
+		return CONTEXT_Shared;
+	}
+	else if (hGLContext)
+	{
+		return CONTEXT_Other;
+	}
+	return CONTEXT_Invalid;
+}
+
+void* PlatformOpenGLCurrentContextHandle(FPlatformOpenGLDevice* Device)
+{
+	return Linux_GetCurrentContext();
+}
+
+void PlatformGetBackbufferDimensions(uint32& OutWidth, uint32& OutHeight)
+{
+	SDL_HWindow CurrentWindow = SDL_GL_GetCurrentWindow();
+
+	int Width = 0;
+	int Height = 0;
+
+	if (CurrentWindow)
+	{
+		SDL_GL_GetDrawableSize(CurrentWindow, &Width, &Height);
+	}
+
+	OutWidth = Width;
+	OutHeight = Height;
+}
+
+// =============================================================
+
+struct FOpenGLReleasedQuery
+{
+	SDL_HGLContext	hGLContext;
+	GLuint			Query;
+};
+
+static TArray<FOpenGLReleasedQuery>	ReleasedQueries;
+static FCriticalSection* ReleasedQueriesGuard;
+
+void PlatformGetNewRenderQuery(GLuint* OutQuery, uint64* OutQueryContext)
+{
+	if (!ReleasedQueriesGuard)
+	{
+		ReleasedQueriesGuard = new FCriticalSection;
+	}
+
+	{
+		FScopeLock Lock(ReleasedQueriesGuard);
+
+#ifdef UE_BUILD_DEBUG
+		check(OutQuery && OutQueryContext);
+#endif
+
+		SDL_HGLContext hGLContext = Linux_GetCurrentContext();
+		check(hGLContext);
+
+		GLuint NewQuery = 0;
+
+		// Check for possible query reuse
+		const int32 ArraySize = ReleasedQueries.Num();
+		for (int32 Index = 0; Index < ArraySize; ++Index)
+		{
+			if (ReleasedQueries[Index].hGLContext == hGLContext)
+			{
+				NewQuery = ReleasedQueries[Index].Query;
+				ReleasedQueries.RemoveAtSwap(Index);
+				break;
+			}
+		}
+
+		if (!NewQuery)
+		{
+			FOpenGL::GenQueries(1, &NewQuery);
+		}
+
+		*OutQuery = NewQuery;
+		*OutQueryContext = (uint64)hGLContext;
+
+	}
+
+}
+
+
+
+void PlatformReleaseRenderQuery(GLuint Query, uint64 QueryContext)
+{
+	SDL_HGLContext hGLContext = Linux_GetCurrentContext();
+	if ((uint64)hGLContext == QueryContext)
+	{
+		FOpenGL::DeleteQueries(1, &Query);
+	}
+	else
+	{
+		FScopeLock Lock(ReleasedQueriesGuard);
+#ifdef UE_BUILD_DEBUG
+		check(Query && QueryContext && ReleasedQueriesGuard);
+#endif
+		FOpenGLReleasedQuery ReleasedQuery;
+		ReleasedQuery.hGLContext = (SDL_HGLContext)QueryContext;
+		ReleasedQuery.Query = Query;
+		ReleasedQueries.Add(ReleasedQuery);
+	}
+
+}
+
+bool PlatformContextIsCurrent(uint64 QueryContext)
+{
+	return (uint64)Linux_GetCurrentContext() == QueryContext;
+}
+
+FOpenGLTexture* PlatformCreateBuiltinBackBuffer(FOpenGLDynamicRHI* OpenGLRHI, uint32 SizeX, uint32 SizeY)
+{
+	return nullptr;
+}
+
+void Linux_DeleteQueriesForCurrentContext(SDL_HGLContext hGLContext)
+{
+	if (!ReleasedQueriesGuard)
+	{
+		ReleasedQueriesGuard = new FCriticalSection;
+	}
+
+	FScopeLock Lock(ReleasedQueriesGuard);
+	for (int32 Index = 0; Index < ReleasedQueries.Num(); ++Index)
+	{
+		if (ReleasedQueries[Index].hGLContext == hGLContext)
+		{
+			FOpenGL::DeleteQueries(1, &ReleasedQueries[Index].Query);
+			ReleasedQueries.RemoveAtSwap(Index);
+			--Index;
+		}
+	}
+}
diff --git a/Engine/Source/Runtime/OpenGLDrv/Private/Linux/LinuxOpenGLPlatform.h b/Engine/Source/Runtime/OpenGLDrv/Private/Linux/LinuxOpenGLPlatform.h
new file mode 100644
--- /dev/null
+++ b/Engine/Source/Runtime/OpenGLDrv/Private/Linux/LinuxOpenGLPlatform.h
@@ -0,0 +1,18 @@
+// Copyright Epic Games, Inc. All Rights Reserved.
+
+#pragma once
+
+#include "OpenGLThirdParty.h"
+
+#define UGL_PLATFORM_SUPPORTS_GLES 0
+#define UGL_PLATFORM_SUPPORTS_GL3  1
+#define UGL_PLATFORM_SUPPORTS_GL4  1
+
+/** Unreal tokens that maps to different OpenGL tokens by platform. */
+#define UGL_DRAW_FRAMEBUFFER	GL_DRAW_FRAMEBUFFER
+#define UGL_READ_FRAMEBUFFER	GL_READ_FRAMEBUFFER
+#define UGL_ABGR8				GL_UNSIGNED_INT_8_8_8_8_REV
+#define UGL_ANY_SAMPLES_PASSED	GL_ANY_SAMPLES_PASSED
+#define UGL_SAMPLES_PASSED		GL_SAMPLES_PASSED
+#define UGL_TIME_ELAPSED		GL_TIME_ELAPSED
+#define UGL_CLAMP_TO_BORDER		GL_CLAMP_TO_BORDER
diff --git a/Engine/Source/Runtime/OpenGLDrv/Private/Linux/OpenGLLinux.cpp b/Engine/Source/Runtime/OpenGLDrv/Private/Linux/OpenGLLinux.cpp
deleted file mode 100644
diff --git a/Engine/Source/Runtime/OpenGLDrv/Private/Linux/OpenGLLinux.h b/Engine/Source/Runtime/OpenGLDrv/Private/Linux/OpenGLLinux.h
deleted file mode 100644
diff --git a/Engine/Source/Runtime/OpenGLDrv/Private/OpenGL.h b/Engine/Source/Runtime/OpenGLDrv/Private/OpenGL.h
--- a/Engine/Source/Runtime/OpenGLDrv/Private/OpenGL.h
+++ b/Engine/Source/Runtime/OpenGLDrv/Private/OpenGL.h
@@ -5,644 +5,6 @@
 =============================================================================*/
 #pragma once
 
-// HEADER_UNIT_SKIP - REVISIT
+#include "HAL/PreprocessorHelpers.h"
 
-#include "Containers/ContainersFwd.h"
-#include "Logging/LogMacros.h"
-
-struct FPlatformOpenGLContext;
-struct FPlatformOpenGLDevice;
-
-/** OpenGL Logging. */
-OPENGLDRV_API DECLARE_LOG_CATEGORY_EXTERN(LogOpenGL,Log,VeryVerbose);
-
-#define UGL_REQUIRED_VOID			{ UE_LOG(LogOpenGL,Fatal,TEXT("%s is not supported."), ANSI_TO_TCHAR(__FUNCTION__)); }
-#define UGL_REQUIRED(ReturnValue)	{ UE_LOG(LogOpenGL,Fatal,TEXT("%s is not supported."), ANSI_TO_TCHAR(__FUNCTION__)); return (ReturnValue); }
-#define UGL_OPTIONAL_VOID			{ }
-#define UGL_OPTIONAL(ReturnValue)	{ return (ReturnValue); }
-
-struct FPlatformOpenGLDevice;
-struct FPlatformOpenGLContext;
-
-
-#define UGL_SUPPORTS_PIXELBUFFERS		1
-#define UGL_SUPPORTS_UNIFORMBUFFERS		1
-
-#ifndef OPENGL_ES
-#define OPENGL_ES	0
-#endif
-#ifndef OPENGL_GL3
-#define OPENGL_GL3	0
-#endif
-#ifndef OPENGL_GL4
-#define OPENGL_GL4	0
-#endif
-
-/** Official OpenGL definitions */
-#ifndef GL_HALF_FLOAT
-#define GL_HALF_FLOAT 0x140B
-#endif
-
-#ifndef GL_RGBA16F
-#define GL_RGBA16F    0x881A
-#endif
-
-
-// for platform extensions: from OpenGLShaders.h
-extern "C" struct FOpenGLShaderDeviceCapabilities;
-
-// for platform extensions: from OpenGLResources.h
-typedef TArray<ANSICHAR> FAnsiCharArray;
-
-// for platform extensions: from OpenGLDrvPrivate.h
-extern "C" struct FOpenGLTextureFormat;
-
-// Base static class
-class FOpenGLBase
-{
-public:
-	enum class EResourceLockMode : uint8
-	{
-		RLM_ReadWrite,
-		RLM_ReadOnly,
-		RLM_WriteOnly,
-		RLM_WriteOnlyUnsynchronized,
-		RLM_WriteOnlyPersistent,
-		RLM_ReadOnlyPersistent,
-	};
-
-	enum EQueryMode
-	{
-		QM_Result,
-		QM_ResultAvailable,
-	};
-
-	enum EFenceResult
-	{
-		FR_AlreadySignaled,
-		FR_TimeoutExpired,
-		FR_ConditionSatisfied,
-		FR_WaitFailed,
-	};
-
-	static void ProcessQueryGLInt();
-	static void ProcessExtensions(const FString& ExtensionsString);
-	static void SetupDefaultGLContextState(const FString& ExtensionsString) {};
-
-	static FORCEINLINE bool SupportsUniformBuffers()					{ return true; }
-	static FORCEINLINE bool SupportsStructuredBuffers()					{ return true; }
-	static FORCEINLINE bool SupportsTimestampQueries()					{ return true; }
-	static FORCEINLINE bool SupportsDisjointTimeQueries()				{ return false; } // @todo: if enabled, causes crash on PC
-	static FORCEINLINE bool SupportsExactOcclusionQueries()				{ return true; }
-	static FORCEINLINE bool SupportsDepthStencilReadSurface()			{ return true; }
-	static FORCEINLINE bool SupportsFloatReadSurface()					{ return true; }
-	static FORCEINLINE bool SupportsWideMRT()							{ return true; }
-	static FORCEINLINE bool SupportsPolygonMode()						{ return true; }
-	static FORCEINLINE bool SupportsTexture3D()							{ return true; }
-	static FORCEINLINE bool SupportsMobileMultiView()					{ return false; }
-	static FORCEINLINE bool SupportsImageExternal()						{ return false; }
-	static FORCEINLINE bool SupportsTextureLODBias()					{ return true; }
-	static FORCEINLINE bool SupportsTextureCompare()					{ return true; }
-	static FORCEINLINE bool SupportsDrawIndexOffset()					{ return true; }
-	static FORCEINLINE bool SupportsDiscardFrameBuffer()				{ return false; }
-	static FORCEINLINE bool SupportsIndexedExtensions()					{ return true; }
-	static FORCEINLINE bool SupportsColorBufferFloat()					{ return true; }
-	static FORCEINLINE bool SupportsColorBufferHalfFloat()				{ return true; }
-	static FORCEINLINE bool SupportsVolumeTextureRendering()			{ return false; }
-	static FORCEINLINE bool SupportsShaderFramebufferFetch()			{ return false; }
-	static FORCEINLINE bool SupportsShaderMRTFramebufferFetch()			{ return false; }
-	static FORCEINLINE bool SupportsShaderDepthStencilFetch()			{ return false; }
-	static FORCEINLINE bool SupportsPixelLocalStorage()					{ return false; }
-	static FORCEINLINE bool SupportsVertexArrayBGRA()					{ return true; }
-	static FORCEINLINE bool SupportsBGRA8888()							{ return true; }
-	static FORCEINLINE bool SupportsDXT()								{ return true; }
-	static FORCEINLINE bool SupportsASTC()								{ return bSupportsASTC; }
-	static FORCEINLINE bool SupportsASTCHDR()							{ return bSupportsASTCHDR; }
-	static FORCEINLINE bool SupportsETC2()								{ return false; }
-	static FORCEINLINE bool SupportsFramebufferSRGBEnable()				{ return true; }
-	static FORCEINLINE bool SupportsFastBufferData()					{ return true; }
-	static FORCEINLINE bool SupportsTextureFilterAnisotropic()			{ return bSupportsTextureFilterAnisotropic; }
-	static FORCEINLINE bool SupportsSeparateAlphaBlend()				{ return bSupportsDrawBuffersBlend; }
-	static FORCEINLINE void EnableSupportsClipControl()					{ bSupportsClipControl = true; }
-	static FORCEINLINE bool SupportsClipControl()						{ return bSupportsClipControl; }
-	static FORCEINLINE bool SupportsSeamlessCubeMap()					{ return false; }
-	static FORCEINLINE bool SupportsDrawIndirect()						{ return false; }
-	static FORCEINLINE bool SupportsGenerateMipmap()					{ return false; }
-	static FORCEINLINE bool SupportsBufferStorage()						{ return false; }
-	static FORCEINLINE bool SupportsDepthBoundsTest()					{ return false; }
-	static FORCEINLINE bool SupportsTextureRange()						{ return false; }
-	static FORCEINLINE bool SupportsBindlessTexture()					{ return false; }
-	static FORCEINLINE bool HasHardwareHiddenSurfaceRemoval()			{ return false; }
-	static FORCEINLINE bool AmdWorkaround()								{ return false; }
-	static FORCEINLINE bool SupportsProgramBinary()						{ return false; }
-	static FORCEINLINE bool SupportsDepthClamp()						{ return true; }
-	
-	static FORCEINLINE bool SupportsASTCDecodeMode()					{ return false; }
-
-	static FORCEINLINE GLenum GetDepthFormat()							{ return GL_DEPTH_COMPONENT16; }
-	static FORCEINLINE GLenum GetShadowDepthFormat()					{ return GL_DEPTH_COMPONENT16; }
-
-	static FORCEINLINE GLint GetMaxTextureImageUnits()			{ check(MaxTextureImageUnits != -1); return MaxTextureImageUnits; }
-	static FORCEINLINE GLint GetMaxVertexTextureImageUnits()	{ check(MaxVertexTextureImageUnits != -1); return MaxVertexTextureImageUnits; }
-	static FORCEINLINE GLint GetMaxGeometryTextureImageUnits()	{ check(MaxGeometryTextureImageUnits != -1); return MaxGeometryTextureImageUnits; }
-	static FORCEINLINE GLint GetMaxComputeTextureImageUnits()	{ check(MaxComputeTextureImageUnits != -1); return MaxComputeTextureImageUnits; }
-	static FORCEINLINE GLint GetMaxCombinedTextureImageUnits()	{ check(MaxCombinedTextureImageUnits != -1); return MaxCombinedTextureImageUnits; }
-	static FORCEINLINE GLint GetTextureBufferAlignment()		{ return TextureBufferAlignment; }
-
-
-	// Indices per unit are set in this order [Pixel, Vertex, Geometry]
-	static FORCEINLINE GLint GetFirstPixelTextureUnit()			{ return 0; }
-	static FORCEINLINE GLint GetFirstVertexTextureUnit()		{ return GetFirstPixelTextureUnit() + GetMaxTextureImageUnits(); }
-	static FORCEINLINE GLint GetFirstGeometryTextureUnit()		{ return GetFirstVertexTextureUnit() + GetMaxVertexTextureImageUnits(); }
-
-	static FORCEINLINE GLint GetFirstComputeTextureUnit()		{ return 0; }
-	
-	// Image load/store units
-	static FORCEINLINE GLint GetFirstComputeUAVUnit()			{ return 0; }
-	static FORCEINLINE GLint GetMaxComputeUAVUnits()			{ return 0; }
-	static FORCEINLINE GLint GetFirstVertexUAVUnit()			{ return 0; }
-	static FORCEINLINE GLint GetFirstPixelUAVUnit()				{ return 0; }
-	static FORCEINLINE GLint GetMaxPixelUAVUnits()				{ return 0; }
-	static FORCEINLINE GLint GetMaxCombinedUAVUnits()			{ return 0; }
-	
-	static FORCEINLINE GLint GetMaxVaryingVectors()				{ check(MaxVaryingVectors != -1); return MaxVaryingVectors; }
-	static FORCEINLINE GLint GetMaxPixelUniformComponents()		{ check(MaxPixelUniformComponents != -1); return MaxPixelUniformComponents; }
-	static FORCEINLINE GLint GetMaxVertexUniformComponents()	{ check(MaxVertexUniformComponents != -1); return MaxVertexUniformComponents; }
-	static FORCEINLINE GLint GetMaxGeometryUniformComponents()	{ check(MaxGeometryUniformComponents != -1); return MaxGeometryUniformComponents; }
-	static FORCEINLINE GLint GetMaxComputeUniformComponents()	{ return 0; }
-
-	static FORCEINLINE uint64 GetVideoMemorySize()				{ return 0; }
-
-	static FORCEINLINE bool IsDebugContent()					{ return false; }
-	static FORCEINLINE void InitDebugContext()					{ }
-
-	static FORCEINLINE int32 GetReadHalfFloatPixelsEnum() UGL_REQUIRED(0)
-
-	static FORCEINLINE GLint GetMaxMSAASamplesTileMem()			{ return 0; /* not supported */ }
-
-	// Silently ignored if not implemented:
-	static FORCEINLINE void QueryTimestampCounter(GLuint QueryID) UGL_OPTIONAL_VOID
-	static FORCEINLINE void BeginQuery(GLenum QueryType, GLuint QueryId) UGL_OPTIONAL_VOID
-	static FORCEINLINE void EndQuery(GLenum QueryType) UGL_OPTIONAL_VOID
-	static FORCEINLINE void GetQueryObject(GLuint QueryId, EQueryMode QueryMode, uint64 *OutResult) UGL_OPTIONAL_VOID
-	static FORCEINLINE void BindFragDataLocation(GLuint Program, GLuint Color, const GLchar *Name) UGL_OPTIONAL_VOID
-	static FORCEINLINE void ReadBuffer(GLenum Mode) UGL_OPTIONAL_VOID
-	static FORCEINLINE void DrawBuffer(GLenum Mode) UGL_OPTIONAL_VOID
-	static FORCEINLINE void DeleteSync(UGLsync Sync) UGL_OPTIONAL_VOID
-	static FORCEINLINE UGLsync FenceSync(GLenum Condition, GLbitfield Flags) UGL_OPTIONAL(UGLsync())
-	static FORCEINLINE bool IsSync(UGLsync Sync) UGL_OPTIONAL(false)
-	static FORCEINLINE EFenceResult ClientWaitSync(UGLsync Sync, GLbitfield Flags, GLuint64 Timeout) UGL_OPTIONAL(FR_WaitFailed)
-	static FORCEINLINE void GenSamplers(GLsizei Count, GLuint *Samplers) UGL_OPTIONAL_VOID
-	static FORCEINLINE void DeleteSamplers(GLsizei Count, GLuint *Samplers) UGL_OPTIONAL_VOID
-	static FORCEINLINE void SetSamplerParameter(GLuint Sampler, GLenum Parameter, GLint Value) UGL_OPTIONAL_VOID
-	static FORCEINLINE void BindSampler(GLuint Unit, GLuint Sampler) UGL_OPTIONAL_VOID
-	static FORCEINLINE void PolygonMode(GLenum Face, GLenum Mode) UGL_OPTIONAL_VOID
-	static FORCEINLINE void VertexAttribDivisor(GLuint Index, GLuint Divisor) UGL_OPTIONAL_VOID
-	static FORCEINLINE void PushGroupMarker(const ANSICHAR* Name) UGL_OPTIONAL_VOID
-	static FORCEINLINE void PopGroupMarker() UGL_OPTIONAL_VOID
-	static FORCEINLINE void LabelObject(GLenum Type, GLuint Object, const ANSICHAR* Name) UGL_OPTIONAL_VOID
-	static FORCEINLINE GLsizei GetLabelObject(GLenum Type, GLuint Object, GLsizei BufferSize, ANSICHAR* OutName) UGL_OPTIONAL(0)
-	static FORCEINLINE void InvalidateFramebuffer(GLenum Target, GLsizei NumAttachments, const GLenum* Attachments) UGL_OPTIONAL_VOID
-
-	// Will assert at run-time if not implemented:
-	static FORCEINLINE void* MapBufferRange(GLenum Type, uint32 InOffset, uint32 InSize, EResourceLockMode LockMode) UGL_REQUIRED(NULL)
-	static FORCEINLINE void UnmapBufferRange(GLenum Type, uint32 InOffset, uint32 InSize) UGL_REQUIRED_VOID
-	static FORCEINLINE void UnmapBuffer(GLenum Type) UGL_REQUIRED_VOID
-	static FORCEINLINE void GenQueries(GLsizei NumQueries, GLuint* QueryIDs) UGL_REQUIRED_VOID
-	static FORCEINLINE void DeleteQueries(GLsizei NumQueries, const GLuint* QueryIDs) UGL_REQUIRED_VOID
-	static FORCEINLINE void GetQueryObject(GLuint QueryId, EQueryMode QueryMode, GLuint *OutResult) UGL_REQUIRED_VOID
-	static FORCEINLINE void BindBufferBase(GLenum Target, GLuint Index, GLuint Buffer) UGL_REQUIRED_VOID
-	static FORCEINLINE void BindBufferRange(GLenum Target, GLuint Index, GLuint Buffer, GLintptr Offset, GLsizeiptr Size) UGL_REQUIRED_VOID
-	static FORCEINLINE GLuint GetUniformBlockIndex(GLuint Program, const GLchar *UniformBlockName) UGL_REQUIRED(-1)
-	static FORCEINLINE void UniformBlockBinding(GLuint Program, GLuint UniformBlockIndex, GLuint UniformBlockBinding) UGL_REQUIRED_VOID
-	static FORCEINLINE void Uniform4uiv(GLint Location, GLsizei Count, const GLuint* Value) UGL_REQUIRED_VOID
-	static FORCEINLINE void TexParameter(GLenum Target, GLenum Parameter, GLint Value) UGL_REQUIRED_VOID
-	static FORCEINLINE void FramebufferTexture(GLenum Target, GLenum Attachment, GLuint Texture, GLint Level) UGL_REQUIRED_VOID
-	static FORCEINLINE void FramebufferTexture2D(GLenum Target, GLenum Attachment, GLenum TexTarget, GLuint Texture, GLint Level)
-	{
-		glFramebufferTexture2D(Target, Attachment, TexTarget, Texture, Level);
-	}
-	static FORCEINLINE void FramebufferTexture2DMultisample(GLenum Target, GLenum Attachment, GLenum TexTarget, GLuint Texture, GLint Level, GLint NumSamples) UGL_REQUIRED_VOID
-	static FORCEINLINE void FramebufferTexture3D(GLenum Target, GLenum Attachment, GLenum TexTarget, GLuint Texture, GLint Level, GLint ZOffset) UGL_REQUIRED_VOID
-	static FORCEINLINE void FramebufferTextureLayer(GLenum Target, GLenum Attachment, GLuint Texture, GLint Level, GLint Layer) UGL_REQUIRED_VOID
-	static FORCEINLINE void FramebufferRenderbuffer(GLenum Target, GLenum Attachment, GLenum RenderBufferTarget, GLuint RenderBuffer)
-	{
-		glFramebufferRenderbuffer(Target, Attachment, RenderBufferTarget, RenderBuffer);
-	}
-	static FORCEINLINE void FramebufferTextureMultiviewOVR(GLenum Target, GLenum Attachment, GLuint Texture, GLint Level, GLint BaseViewIndex, GLsizei NumViews) UGL_REQUIRED_VOID
-	static FORCEINLINE void FramebufferTextureMultisampleMultiviewOVR(GLenum Target, GLenum Attachment, GLuint Texture, GLint Level, GLsizei NumSamples, GLint BaseViewIndex, GLsizei NumViews) UGL_REQUIRED_VOID
-	static FORCEINLINE void BlitFramebuffer(GLint SrcX0, GLint SrcY0, GLint SrcX1, GLint SrcY1, GLint DstX0, GLint DstY0, GLint DstX1, GLint DstY1, GLbitfield Mask, GLenum Filter) UGL_REQUIRED_VOID
-	static FORCEINLINE void DrawBuffers(GLsizei NumBuffers, const GLenum *Buffers) UGL_REQUIRED_VOID
-	static FORCEINLINE void DepthRange(GLdouble Near, GLdouble Far) UGL_REQUIRED_VOID
-	static FORCEINLINE void EnableIndexed(GLenum Parameter, GLuint Index) UGL_REQUIRED_VOID
-	static FORCEINLINE void DisableIndexed(GLenum Parameter, GLuint Index) UGL_REQUIRED_VOID
-	static FORCEINLINE void ColorMaskIndexed(GLuint Index, GLboolean Red, GLboolean Green, GLboolean Blue, GLboolean Alpha) UGL_REQUIRED_VOID
-	static FORCEINLINE void VertexAttribPointer(GLuint Index, GLint Size, GLenum Type, GLboolean Normalized, GLsizei Stride, const GLvoid* Pointer) UGL_REQUIRED_VOID
-	static FORCEINLINE void VertexAttribIPointer(GLuint Index, GLint Size, GLenum Type, GLsizei Stride, const GLvoid* Pointer) UGL_REQUIRED_VOID
-	static FORCEINLINE void VertexAttrib4Nsv(GLuint AttributeIndex, const GLshort* Values) UGL_REQUIRED_VOID
-	static FORCEINLINE void VertexAttrib4sv(GLuint AttributeIndex, const GLshort* Values) UGL_REQUIRED_VOID
-	static FORCEINLINE void VertexAttribI4sv(GLuint AttributeIndex, const GLshort* Values) UGL_REQUIRED_VOID
-	static FORCEINLINE void VertexAttribI4usv(GLuint AttributeIndex, const GLushort* Values) UGL_REQUIRED_VOID
-	static FORCEINLINE void VertexAttrib4Nubv(GLuint AttributeIndex, const GLubyte* Values) UGL_REQUIRED_VOID
-	static FORCEINLINE void VertexAttrib4ubv(GLuint AttributeIndex, const GLubyte* Values) UGL_REQUIRED_VOID
-	static FORCEINLINE void VertexAttribI4ubv(GLuint AttributeIndex, const GLubyte* Values) UGL_REQUIRED_VOID
-	static FORCEINLINE void VertexAttrib4Nbv(GLuint AttributeIndex, const GLbyte* Values) UGL_REQUIRED_VOID
-	static FORCEINLINE void VertexAttrib4bv(GLuint AttributeIndex, const GLbyte* Values) UGL_REQUIRED_VOID
-	static FORCEINLINE void VertexAttribI4bv(GLuint AttributeIndex, const GLbyte* Values) UGL_REQUIRED_VOID
-	static FORCEINLINE void VertexAttrib4dv(GLuint AttributeIndex, const GLdouble* Values) UGL_REQUIRED_VOID
-	static FORCEINLINE void VertexAttribI4iv(GLuint AttributeIndex, const GLint* Values) UGL_REQUIRED_VOID
-	static FORCEINLINE void VertexAttribI4uiv(GLuint AttributeIndex, const GLuint* Values) UGL_REQUIRED_VOID
-	static FORCEINLINE void DrawArraysInstanced(GLenum Mode, GLint First, GLsizei Count, GLsizei InstanceCount) UGL_REQUIRED_VOID
-	static FORCEINLINE void DrawElementsInstanced(GLenum Mode, GLsizei Count, GLenum Type, const GLvoid* Indices, GLsizei InstanceCount) UGL_REQUIRED_VOID
-	static FORCEINLINE void DrawRangeElements(GLenum Mode, GLuint Start, GLuint End, GLsizei Count, GLenum Type, const GLvoid* Indices) UGL_REQUIRED_VOID
-	static FORCEINLINE void ClearBufferfv(GLenum Buffer, GLint DrawBufferIndex, const GLfloat* Value) UGL_REQUIRED_VOID
-	static FORCEINLINE void ClearBufferfi(GLenum Buffer, GLint DrawBufferIndex, GLfloat Depth, GLint Stencil) UGL_REQUIRED_VOID
-	static FORCEINLINE void ClearBufferiv(GLenum Buffer, GLint DrawBufferIndex, const GLint* Value) UGL_REQUIRED_VOID
-	static FORCEINLINE void ClearDepth(GLdouble Depth) UGL_REQUIRED_VOID
-	static FORCEINLINE void TexImage3D(GLenum Target, GLint Level, GLint InternalFormat, GLsizei Width, GLsizei Height, GLsizei Depth, GLint Border, GLenum Format, GLenum Type, const GLvoid* PixelData) UGL_REQUIRED_VOID
-	static FORCEINLINE void CompressedTexImage3D(GLenum Target, GLint Level, GLenum InternalFormat, GLsizei Width, GLsizei Height, GLsizei Depth, GLint Border, GLsizei ImageSize, const GLvoid* PixelData) UGL_REQUIRED_VOID
-	static FORCEINLINE void TexImage2DMultisample(GLenum Target, GLsizei Samples, GLint InternalFormat, GLsizei Width, GLsizei Height, GLboolean FixedSampleLocations) UGL_REQUIRED_VOID
-	static FORCEINLINE void TexBuffer(GLenum Target, GLenum InternalFormat, GLuint Buffer) UGL_REQUIRED_VOID
-	static FORCEINLINE void TexBufferRange(GLenum Target, GLenum InternalFormat, GLuint Buffer, GLintptr Offset, GLsizeiptr Size) UGL_REQUIRED_VOID
-	static FORCEINLINE void TexSubImage3D(GLenum Target, GLint Level, GLint XOffset, GLint YOffset, GLint ZOffset, GLsizei Width, GLsizei Height, GLsizei Depth, GLenum Format, GLenum Type, const GLvoid* PixelData) UGL_REQUIRED_VOID
-	static FORCEINLINE void	CopyTexSubImage2D(GLenum Target, GLint Level, GLint XOffset, GLint YOffset, GLint X, GLint Y, GLsizei Width, GLsizei Height) UGL_REQUIRED_VOID
-	static FORCEINLINE void	CopyTexSubImage3D(GLenum Target, GLint Level, GLint XOffset, GLint YOffset, GLint ZOffset, GLint X, GLint Y, GLsizei Width, GLsizei Height) UGL_REQUIRED_VOID
-	static FORCEINLINE void GetCompressedTexImage(GLenum Target, GLint Level, GLvoid* OutImageData) UGL_REQUIRED_VOID
-	static FORCEINLINE void GetTexImage(GLenum Target, GLint Level, GLenum Format, GLenum Type, GLvoid* OutPixelData) UGL_REQUIRED_VOID
-	static FORCEINLINE void CopyBufferSubData(GLenum ReadTarget, GLenum WriteTarget, GLintptr ReadOffset, GLintptr WriteOffset, GLsizeiptr Size) UGL_REQUIRED_VOID
-	static FORCEINLINE const ANSICHAR* GetStringIndexed(GLenum Name, GLuint Index) UGL_REQUIRED(NULL)
-	static FORCEINLINE GLuint GetMajorVersion() UGL_REQUIRED(0)
-	static FORCEINLINE GLuint GetMinorVersion() UGL_REQUIRED(0)
-	static FORCEINLINE ERHIFeatureLevel::Type GetFeatureLevel() UGL_REQUIRED(ERHIFeatureLevel::SM5)
-	static FORCEINLINE EShaderPlatform GetShaderPlatform() UGL_REQUIRED(SP_NumPlatforms)
-	static FORCEINLINE FString GetAdapterName() UGL_REQUIRED(TEXT(""))
-	static FORCEINLINE void BlendFuncSeparatei(GLuint Buf, GLenum SrcRGB, GLenum DstRGB, GLenum SrcAlpha, GLenum DstAlpha) UGL_REQUIRED_VOID
-	static FORCEINLINE void BlendEquationSeparatei(GLuint Buf, GLenum ModeRGB, GLenum ModeAlpha) UGL_REQUIRED_VOID
-	static FORCEINLINE void BlendFunci(GLuint Buf, GLenum Src, GLenum Dst) UGL_REQUIRED_VOID
-	static FORCEINLINE void BlendEquationi(GLuint Buf, GLenum Mode) UGL_REQUIRED_VOID
-	static FORCEINLINE void PatchParameteri(GLenum Pname, GLint Value) UGL_REQUIRED_VOID
-	static FORCEINLINE void BindImageTexture(GLuint Unit, GLuint Texture, GLint Level, GLboolean Layered, GLint Layer, GLenum Access, GLenum Format) UGL_REQUIRED_VOID
-	static FORCEINLINE void DispatchCompute(GLuint NumGroupsX, GLuint NumGroupsY, GLuint NumGroupsZ) UGL_REQUIRED_VOID
-	static FORCEINLINE void DispatchComputeIndirect(GLintptr Offset) UGL_REQUIRED_VOID
-	static FORCEINLINE void MemoryBarrier(GLbitfield Barriers) UGL_REQUIRED_VOID
-	static FORCEINLINE bool TexStorage2D(GLenum Target, GLint Levels, GLint InternalFormat, GLsizei Width, GLsizei Height, GLenum Format, GLenum Type, ETextureCreateFlags Flags) UGL_OPTIONAL(false)
-	static FORCEINLINE bool TexStorage2DMultisample(GLenum Target, GLsizei Samples, GLint InternalFormat, GLsizei Width, GLsizei Height, GLboolean FixedSampleLocations) UGL_OPTIONAL(false)
-	static FORCEINLINE void RenderbufferStorageMultisample(GLenum Target, GLsizei Samples, GLint InternalFormat, GLsizei Width, GLsizei Height) UGL_REQUIRED_VOID
-	static FORCEINLINE void TexStorage3D(GLenum Target, GLint Levels, GLint InternalFormat, GLsizei Width, GLsizei Height, GLsizei Depth, GLenum Format, GLenum Type) UGL_REQUIRED_VOID
-	static FORCEINLINE void CompressedTexSubImage3D(GLenum Target, GLint Level, GLint XOffset, GLint YOffset, GLint ZOffset, GLsizei Width, GLsizei Height, GLsizei Depth, GLenum Format, GLsizei ImageSize, const GLvoid* PixelData) UGL_REQUIRED_VOID
-	static FORCEINLINE void CopyImageSubData(GLuint SrcName, GLenum SrcTarget, GLint SrcLevel, GLint SrcX, GLint SrcY, GLint SrcZ, GLuint DstName, GLenum DstTarget, GLint DstLevel, GLint DstX, GLint DstY, GLint DstZ, GLsizei Width, GLsizei Height, GLsizei Depth) UGL_REQUIRED_VOID
-	static FORCEINLINE void TextureView(GLuint ViewName, GLenum ViewTarget, GLuint SrcName, GLenum InternalFormat, GLuint MinLevel, GLuint NumLevels, GLuint MinLayer, GLuint NumLayers) UGL_REQUIRED_VOID
-	static FORCEINLINE void DrawArraysIndirect(GLenum Mode, const void *Offset) UGL_REQUIRED_VOID
-	static FORCEINLINE void DrawElementsIndirect(GLenum Mode, GLenum Type, const void *Offset) UGL_REQUIRED_VOID
-	static FORCEINLINE void GenerateMipmap(GLenum Target) UGL_REQUIRED_VOID
-	static FORCEINLINE void BindVertexBuffer(GLuint BindingIndex, GLuint Nuffer, GLintptr Offset, GLsizei Stride) UGL_REQUIRED_VOID
-	static FORCEINLINE void VertexAttribFormat(GLuint AttribIndex, GLint Size, GLenum Type, GLboolean Normalized, GLuint RelativeOffset) UGL_REQUIRED_VOID
-	static FORCEINLINE void VertexAttribIFormat(GLuint AttribIndex, GLint Size, GLenum Type, GLuint RelativeOffset) UGL_REQUIRED_VOID
-	static FORCEINLINE void VertexAttribBinding(GLuint AttribIndex, GLuint BindingIndex) UGL_REQUIRED_VOID
-	static FORCEINLINE void ClearBufferData(GLenum Target, GLenum InternalFormat, GLenum Format, GLenum Type, const uint32* Data) UGL_REQUIRED_VOID
-	static FORCEINLINE void VertexBindingDivisor(GLuint BindingIndex, GLuint Divisor) UGL_REQUIRED_VOID
-	static FORCEINLINE void BufferStorage(GLenum Target, GLsizeiptr Size, const void *Data, GLbitfield Flags) UGL_REQUIRED_VOID
-	static FORCEINLINE void DepthBounds(GLfloat Min, GLfloat Max) UGL_REQUIRED_VOID
-	static FORCEINLINE void TextureRange(GLenum Target, GLsizei Length, const GLvoid *Pointer) UGL_OPTIONAL_VOID
-	static FORCEINLINE void ProgramParameter (GLuint Program, GLenum PName, GLint Value) UGL_OPTIONAL_VOID
-	static FORCEINLINE void UseProgramStages(GLuint Pipeline, GLbitfield Stages, GLuint Program) { glAttachShader(Pipeline, Program); }
-	static FORCEINLINE void BindProgramPipeline(GLuint Pipeline) { glUseProgram(Pipeline); }
-	static FORCEINLINE void DeleteShader(GLuint Program) { glDeleteShader(Program); }
-	static FORCEINLINE void DeleteProgramPipelines(GLsizei Number, const GLuint *Pipelines) { for(GLsizei i = 0; i < Number; i++) { glDeleteProgram(Pipelines[i]); } }
-	static FORCEINLINE void GenProgramPipelines(GLsizei Number, GLuint *Pipelines) { check(Pipelines); for(GLsizei i = 0; i < Number; i++) { Pipelines[i] = glCreateProgram(); } }
-	static FORCEINLINE void ProgramUniform1i(GLuint Program, GLint Location, GLint V0) { glUniform1i( Location, V0 ); }
-	static FORCEINLINE void ProgramUniform4iv(GLuint Program, GLint Location, GLsizei Count, const GLint *Value) { glUniform4iv(Location, Count, Value); }
-	static FORCEINLINE void ProgramUniform4fv(GLuint Program, GLint Location, GLsizei Count, const GLfloat *Value) { glUniform4fv(Location, Count, Value); }
-	static FORCEINLINE void ProgramUniform4uiv(GLuint Program, GLint Location, GLsizei Count, const GLuint *Value) UGL_REQUIRED_VOID
-	static FORCEINLINE void GetProgramPipelineiv(GLuint Pipeline, GLenum Pname, GLint* Params) UGL_OPTIONAL_VOID
-	static FORCEINLINE void ValidateProgramPipeline(GLuint Pipeline) UGL_OPTIONAL_VOID
-	static FORCEINLINE void GetProgramPipelineInfoLog(GLuint Pipeline, GLsizei BufSize, GLsizei* Length, GLchar* InfoLog) UGL_OPTIONAL_VOID
-	static FORCEINLINE bool IsProgramPipeline(GLuint Pipeline) UGL_OPTIONAL(false)
-
-	static FORCEINLINE GLuint64 GetTextureSamplerHandle(GLuint Texture, GLuint Sampler) UGL_REQUIRED(0)
-	static FORCEINLINE GLuint64 GetTextureHandle(GLuint Texture) UGL_REQUIRED(0)
-	static FORCEINLINE void MakeTextureHandleResident(GLuint64 TextureHandle) UGL_REQUIRED_VOID
-	static FORCEINLINE void MakeTextureHandleNonResident(GLuint64 TextureHandle) UGL_REQUIRED_VOID
-	static FORCEINLINE void UniformHandleui64(GLint Location, GLuint64 Value) UGL_REQUIRED_VOID
-
-	static FORCEINLINE void GetProgramBinary(GLuint Program, GLsizei BufSize, GLsizei *Length, GLenum *BinaryFormat, void *Binary) UGL_OPTIONAL_VOID
-	static FORCEINLINE void ProgramBinary(GLuint Program, GLenum BinaryFormat, const void *Binary, GLsizei Length) UGL_OPTIONAL_VOID
-
-	static FORCEINLINE void FrameBufferFetchBarrier() UGL_OPTIONAL_VOID
-	
-	static FPlatformOpenGLDevice*	CreateDevice() UGL_REQUIRED(NULL)
-	static FPlatformOpenGLContext*	CreateContext( FPlatformOpenGLDevice* Device, void* WindowHandle ) UGL_REQUIRED(NULL)
-
-	static FORCEINLINE void CheckFrameBuffer()
-	{
-#if UE_BUILD_DEBUG || UE_BUILD_DEVELOPMENT 
-		GLenum CompleteResult = glCheckFramebufferStatus(GL_FRAMEBUFFER);
-		if (CompleteResult != GL_FRAMEBUFFER_COMPLETE)
-		{
-				UE_LOG(LogRHI, Fatal,TEXT("Framebuffer not complete. Status = 0x%x"), CompleteResult);
-		}
-#endif 
-	}
-
-	static FORCEINLINE void BufferSubData(GLenum Target, GLintptr Offset, GLsizeiptr Size, const GLvoid* Data)	{ glBufferSubData(Target, Offset, Size, Data); }
-	static FORCEINLINE void DeleteBuffers(GLsizei Number, const GLuint* Buffers)	{ glDeleteBuffers(Number, Buffers); }
-	static FORCEINLINE void DeleteTextures(GLsizei Number, const GLuint* Textures)	{ glDeleteTextures(Number, Textures); }
-	static FORCEINLINE void Flush()													{ glFlush(); }
-	static FORCEINLINE GLuint CreateShader(GLenum Type)								{ return glCreateShader(Type); }
-	static FORCEINLINE GLuint CreateProgram()										{ return glCreateProgram(); }
-	static FORCEINLINE bool TimerQueryDisjoint()									{ return false; }
-
-	// Calling glBufferData() to discard-reupload is slower than calling glBufferSubData() on some platforms,
-	// because changing glBufferData() with a different size (from before) may incur extra validation.
-	// To use glBufferData() discard trick: set to this to true - otherwise, glBufferSubData() will used.
-	static FORCEINLINE bool DiscardFrameBufferToResize()							{ return true; }
-
-	static FORCEINLINE EPixelFormat PreferredPixelFormatHint(EPixelFormat PreferredPixelFormat)
-	{
-		// Use a default pixel format if none was specified	
-		if (PreferredPixelFormat == EPixelFormat::PF_Unknown)
-		{
-			PreferredPixelFormat = EPixelFormat::PF_B8G8R8A8;
-		}
-		return PreferredPixelFormat;
-	}
-
-	// for platform extensions
-	static void PE_GetCurrentOpenGLShaderDeviceCapabilities(FOpenGLShaderDeviceCapabilities& Capabilities);
-	static bool PE_GLSLToDeviceCompatibleGLSL(FAnsiCharArray& GlslCodeOriginal, const FString& ShaderName, GLenum TypeEnum, const FOpenGLShaderDeviceCapabilities& Capabilities, FAnsiCharArray& GlslCode) UGL_OPTIONAL(false)
-	static void PE_SetupTextureFormat(void(*SetupTextureFormat)(EPixelFormat, const FOpenGLTextureFormat&)) UGL_OPTIONAL_VOID
-
-	static GLenum GetPlatfrom5551Format() UGL_OPTIONAL(GL_UNSIGNED_SHORT_5_5_5_1)
-protected:
-	static GLint MaxTextureImageUnits;
-	static GLint MaxCombinedTextureImageUnits;
-	static GLint MaxComputeTextureImageUnits;
-	static GLint MaxVertexTextureImageUnits;
-	static GLint MaxGeometryTextureImageUnits;
-	static GLint MaxVertexUniformComponents;
-	static GLint MaxPixelUniformComponents;
-	static GLint MaxGeometryUniformComponents;
-	static GLint MaxVaryingVectors;
-	static GLint TextureBufferAlignment;
-
-	/** GL_ARB_clip_control */
-	static bool bSupportsClipControl;
-
-	/** GL_KHR_texture_compression_astc_ldr */
-	static bool bSupportsASTC;
-
-	/** GL_KHR_texture_compression_astc_hdr */
-	static bool bSupportsASTCHDR;
-
-	/** GL_ARB_seamless_cube_map */
-	static bool bSupportsSeamlessCubemap;
-	
-	/** Can we render to texture 2D array or 3D */
-	static bool bSupportsVolumeTextureRendering;
-	
-	/** GL_EXT_texture_filter_anisotropic Can we use anisotropic filtering? */
-	static bool bSupportsTextureFilterAnisotropic;
-
-	/** GL_ARB_draw_buffers_blend */
-	static bool bSupportsDrawBuffersBlend;
-
-	/** Workaround AMD driver issues. */
-	static bool bAmdWorkaround;
-
-};
-
-/** Unreal tokens that maps to different OpenGL tokens by platform. */
-#define UGL_DRAW_FRAMEBUFFER	GL_DRAW_FRAMEBUFFER
-#define UGL_READ_FRAMEBUFFER	GL_READ_FRAMEBUFFER
-#define UGL_ABGR8				GL_UNSIGNED_INT_8_8_8_8_REV
-#define UGL_ANY_SAMPLES_PASSED	GL_ANY_SAMPLES_PASSED
-#define UGL_SAMPLES_PASSED		GL_SAMPLES_PASSED
-#define UGL_TIME_ELAPSED		GL_TIME_ELAPSED
-#define UGL_CLAMP_TO_BORDER		GL_CLAMP_TO_BORDER
-
-/** OpenGL extensions */
-
-// http://www.opengl.org/registry/specs/EXT/texture_compression_s3tc.txt
-#if !defined(GL_EXT_texture_compression_s3tc)
-#define GL_COMPRESSED_RGB_S3TC_DXT1_EXT			0x83F0
-#define GL_COMPRESSED_RGBA_S3TC_DXT1_EXT		0x83F1
-#define GL_COMPRESSED_RGBA_S3TC_DXT3_EXT		0x83F2
-#define GL_COMPRESSED_RGBA_S3TC_DXT5_EXT		0x83F3
-#endif
-
-// http://www.opengl.org/registry/specs/EXT/texture_sRGB.txt
-#if !defined(GL_EXT_texture_sRGB)
-#define GL_SRGB_EXT                       0x8C40
-#define GL_SRGB8_EXT                      0x8C41
-#define GL_SRGB_ALPHA_EXT                 0x8C42
-#define GL_SRGB8_ALPHA8_EXT               0x8C43
-#define GL_SLUMINANCE_ALPHA_EXT           0x8C44
-#define GL_SLUMINANCE8_ALPHA8_EXT         0x8C45
-#define GL_SLUMINANCE_EXT                 0x8C46
-#define GL_SLUMINANCE8_EXT                0x8C47
-#define GL_COMPRESSED_SRGB_EXT            0x8C48
-#define GL_COMPRESSED_SRGB_ALPHA_EXT      0x8C49
-#define GL_COMPRESSED_SLUMINANCE_EXT      0x8C4A
-#define GL_COMPRESSED_SLUMINANCE_ALPHA_EXT 0x8C4B
-#define GL_COMPRESSED_SRGB_S3TC_DXT1_EXT  0x8C4C
-#define GL_COMPRESSED_SRGB_ALPHA_S3TC_DXT1_EXT 0x8C4D
-#define GL_COMPRESSED_SRGB_ALPHA_S3TC_DXT3_EXT 0x8C4E
-#define GL_COMPRESSED_SRGB_ALPHA_S3TC_DXT5_EXT 0x8C4F
-#endif
-
-// http://www.opengl.org/registry/specs/ARB/texture_compression_rgtc.txt
-#if !defined(GL_ARB_texture_compression_rgtc)
-#define GL_COMPRESSED_RED_RGTC1           0x8DBB
-#define GL_COMPRESSED_SIGNED_RED_RGTC1    0x8DBC
-#define GL_COMPRESSED_RG_RGTC2            0x8DBD
-#define GL_COMPRESSED_SIGNED_RG_RGTC2     0x8DBE
-#endif
-
-/* http://www.khronos.org/registry/gles/extensions/NV/NV_sRGB_formats.txt */
-#if !defined(GL_NV_sRGB_formats)
-#define GL_SLUMINANCE_NV                                        0x8C46
-#define GL_SLUMINANCE_ALPHA_NV                                  0x8C44
-#define GL_SRGB8_NV                                             0x8C41
-#define GL_SLUMINANCE8_NV                                       0x8C47
-#define GL_SLUMINANCE8_ALPHA8_NV                                0x8C45
-#define GL_COMPRESSED_SRGB_S3TC_DXT1_NV                         0x8C4C
-#define GL_COMPRESSED_SRGB_ALPHA_S3TC_DXT1_NV                   0x8C4D
-#define GL_COMPRESSED_SRGB_ALPHA_S3TC_DXT3_NV                   0x8C4E
-#define GL_COMPRESSED_SRGB_ALPHA_S3TC_DXT5_NV                   0x8C4F
-#endif
-
-// http://www.opengl.org/registry/specs/KHR/texture_compression_astc_ldr.txt
-#if !defined(GL_KHR_texture_compression_astc_ldr)
-#define GL_COMPRESSED_RGBA_ASTC_4x4_KHR            0x93B0
-#define GL_COMPRESSED_RGBA_ASTC_5x4_KHR            0x93B1
-#define GL_COMPRESSED_RGBA_ASTC_5x5_KHR            0x93B2
-#define GL_COMPRESSED_RGBA_ASTC_6x5_KHR            0x93B3
-#define GL_COMPRESSED_RGBA_ASTC_6x6_KHR            0x93B4
-#define GL_COMPRESSED_RGBA_ASTC_8x5_KHR            0x93B5
-#define GL_COMPRESSED_RGBA_ASTC_8x6_KHR            0x93B6
-#define GL_COMPRESSED_RGBA_ASTC_8x8_KHR            0x93B7
-#define GL_COMPRESSED_RGBA_ASTC_10x5_KHR           0x93B8
-#define GL_COMPRESSED_RGBA_ASTC_10x6_KHR           0x93B9
-#define GL_COMPRESSED_RGBA_ASTC_10x8_KHR           0x93BA
-#define GL_COMPRESSED_RGBA_ASTC_10x10_KHR          0x93BB
-#define GL_COMPRESSED_RGBA_ASTC_12x10_KHR          0x93BC
-#define GL_COMPRESSED_RGBA_ASTC_12x12_KHR          0x93BD
-
-#define GL_COMPRESSED_SRGB8_ALPHA8_ASTC_4x4_KHR    0x93D0
-#define GL_COMPRESSED_SRGB8_ALPHA8_ASTC_5x4_KHR    0x93D1
-#define GL_COMPRESSED_SRGB8_ALPHA8_ASTC_5x5_KHR    0x93D2
-#define GL_COMPRESSED_SRGB8_ALPHA8_ASTC_6x5_KHR    0x93D3
-#define GL_COMPRESSED_SRGB8_ALPHA8_ASTC_6x6_KHR    0x93D4
-#define GL_COMPRESSED_SRGB8_ALPHA8_ASTC_8x5_KHR    0x93D5
-#define GL_COMPRESSED_SRGB8_ALPHA8_ASTC_8x6_KHR    0x93D6
-#define GL_COMPRESSED_SRGB8_ALPHA8_ASTC_8x8_KHR    0x93D7
-#define GL_COMPRESSED_SRGB8_ALPHA8_ASTC_10x5_KHR   0x93D8
-#define GL_COMPRESSED_SRGB8_ALPHA8_ASTC_10x6_KHR   0x93D9
-#define GL_COMPRESSED_SRGB8_ALPHA8_ASTC_10x8_KHR   0x93DA
-#define GL_COMPRESSED_SRGB8_ALPHA8_ASTC_10x10_KHR  0x93DB
-#define GL_COMPRESSED_SRGB8_ALPHA8_ASTC_12x10_KHR  0x93DC
-#define GL_COMPRESSED_SRGB8_ALPHA8_ASTC_12x12_KHR  0x93DD
-#endif
-
-#if !defined(GL_TESS_EVALUATION_SHADER)
-#define GL_TESS_EVALUATION_SHADER					0x8E87
-#endif
-#if !defined(GL_TESS_CONTROL_SHADER)
-#define GL_TESS_CONTROL_SHADER						0x8E88
-#endif
-#if !defined(GL_PATCHES)
-#define GL_PATCHES									0x000E
-#endif
-#if !defined(GL_PATCH_VERTICES)
-#define GL_PATCH_VERTICES							0x8E72
-#endif
-#if !defined(GL_MAX_GEOMETRY_TEXTURE_IMAGE_UNITS)
-#define GL_MAX_GEOMETRY_TEXTURE_IMAGE_UNITS			0x8C29
-#endif
-#if !defined(GL_MAX_TESS_CONTROL_TEXTURE_IMAGE_UNITS)
-#define GL_MAX_TESS_CONTROL_TEXTURE_IMAGE_UNITS		0x8E81
-#endif
-#if !defined(GL_MAX_TESS_EVALUATION_TEXTURE_IMAGE_UNITS)
-#define GL_MAX_TESS_EVALUATION_TEXTURE_IMAGE_UNITS	0x8E82
-#endif
-#if !defined(GL_READ_WRITE)
-#define GL_READ_WRITE								0x88BA
-#endif
-#if !defined(GL_ALL_BARRIER_BITS)
-#define GL_ALL_BARRIER_BITS							0xFFFFFFFF
-#endif
-#if !defined(GL_TEXTURE_CUBE_MAP_ARRAY)
-#define GL_TEXTURE_CUBE_MAP_ARRAY					0x9009
-#endif
-#if !defined(GL_UNIFORM_BLOCK_REFERENCED_BY_TESS_CONTROL_SHADER)
-#define GL_UNIFORM_BLOCK_REFERENCED_BY_TESS_CONTROL_SHADER		0x84F0
-#endif
-#if !defined(GL_UNIFORM_BLOCK_REFERENCED_BY_TESS_EVALUATION_SHADER)
-#define GL_UNIFORM_BLOCK_REFERENCED_BY_TESS_EVALUATION_SHADER	0x84F1
-#endif
-#if !defined(GL_UNIFORM_BLOCK_REFERENCED_BY_COMPUTE_SHADER)
-#define GL_UNIFORM_BLOCK_REFERENCED_BY_COMPUTE_SHADER			0x90EC
-#endif
-#ifndef GL_ARB_seamless_cube_map
-#define GL_TEXTURE_CUBE_MAP_SEAMLESS				0x884F
-#endif
-#ifndef GL_TIME_ELAPSED_EXT
-#define GL_TIME_ELAPSED_EXT							0x88BF
-#endif
-#ifndef GL_TIMESTAMP_EXT
-#define GL_TIMESTAMP_EXT							0x8E28
-#endif
-#ifndef GL_DISPATCH_INDIRECT_BUFFER
-#define GL_DISPATCH_INDIRECT_BUFFER					0x90EE
-#endif
-#ifndef GL_DRAW_INDIRECT_BUFFER
-#define GL_DRAW_INDIRECT_BUFFER						0x8F3F
-#endif
-#ifndef GL_MAP_WRITE_BIT
-#define GL_MAP_WRITE_BIT							0x0002
-#endif
-#ifndef GL_MAP_PERSISTENT_BIT
-#define GL_MAP_PERSISTENT_BIT						0x0040
-#endif
-#ifndef GL_MAP_COHERENT_BIT
-#define GL_MAP_COHERENT_BIT							0x0080
-#endif
-#ifndef GL_DEPTH_BOUNDS_TEST_EXT
-#define GL_DEPTH_BOUNDS_TEST_EXT					0x8890
-#endif
-#ifndef GL_DEPTH_STENCIL_TEXTURE_MODE
-#define GL_DEPTH_STENCIL_TEXTURE_MODE				0x90EA
-#endif
-#ifndef GL_TEXTURE_SWIZZLE_R
-#define GL_TEXTURE_SWIZZLE_R						0x8E42
-#endif
-#ifndef GL_TEXTURE_SWIZZLE_G
-#define GL_TEXTURE_SWIZZLE_G						0x8E43
-#endif
-#ifndef GL_TEXTURE_SWIZZLE_B
-#define GL_TEXTURE_SWIZZLE_B						0x8E44
-#endif
-#ifndef GL_TEXTURE_SWIZZLE_A
-#define GL_TEXTURE_SWIZZLE_A						0x8E45
-#endif
-#ifndef GL_RED
-#define GL_RED										0x1903
-#endif
-#ifndef GL_BLUE
-#define GL_BLUE										0x1905
-#endif
-#ifndef GL_STENCIL_INDEX
-#define GL_STENCIL_INDEX							0x1901
-#endif
-#ifndef GL_RGBA_INTEGER
-#define GL_RGBA_INTEGER								0x8D99
-#endif
-
-#ifndef GL_ARB_compute_shader
-#define GL_COMPUTE_SHADER                 0x91B9
-#define GL_MAX_COMPUTE_UNIFORM_BLOCKS     0x91BB
-#define GL_MAX_COMPUTE_TEXTURE_IMAGE_UNITS 0x91BC
-#define GL_MAX_COMPUTE_IMAGE_UNIFORMS     0x91BD
-#define GL_MAX_COMPUTE_SHARED_MEMORY_SIZE 0x8262
-#define GL_MAX_COMPUTE_UNIFORM_COMPONENTS 0x8263
-#define GL_MAX_COMPUTE_ATOMIC_COUNTER_BUFFERS 0x8264
-#define GL_MAX_COMPUTE_ATOMIC_COUNTERS    0x8265
-#define GL_MAX_COMBINED_COMPUTE_UNIFORM_COMPONENTS 0x8266
-#define GL_MAX_COMPUTE_LOCAL_INVOCATIONS  0x90EB
-#define GL_MAX_COMPUTE_WORK_GROUP_COUNT   0x91BE
-#define GL_MAX_COMPUTE_WORK_GROUP_SIZE    0x91BF
-#define GL_COMPUTE_LOCAL_WORK_SIZE        0x8267
-#define GL_UNIFORM_BLOCK_REFERENCED_BY_COMPUTE_SHADER 0x90EC
-#define GL_ATOMIC_COUNTER_BUFFER_REFERENCED_BY_COMPUTE_SHADER 0x90ED
-#define GL_DISPATCH_INDIRECT_BUFFER       0x90EE
-#define GL_DISPATCH_INDIRECT_BUFFER_BINDING 0x90EF
-#define GL_COMPUTE_SHADER_BIT             0x00000020
-#endif
-
-#ifndef GL_GPU_DISJOINT_EXT
-#define GL_GPU_DISJOINT_EXT							0x8FBB
-#endif
-
-#ifndef GL_APPLE_client_storage
-#define GL_UNPACK_CLIENT_STORAGE_APPLE				0x85B2
-#endif
-
-#ifndef GL_APPLE_texture_range
-#define GL_TEXTURE_RANGE_LENGTH_APPLE     0x85B7
-#define GL_TEXTURE_RANGE_POINTER_APPLE    0x85B8
-#define GL_TEXTURE_STORAGE_HINT_APPLE     0x85BC
-#define GL_TEXTURE_MINIMIZE_STORAGE_APPLE 0x85B6
-#define GL_STORAGE_PRIVATE_APPLE          0x85BD
-#define GL_STORAGE_CACHED_APPLE           0x85BE
-#define GL_STORAGE_SHARED_APPLE           0x85BF
-#endif
-
-#ifndef TEXTURE_ASTC_DECODE_PRECISION_EXT
-#define TEXTURE_ASTC_DECODE_PRECISION_EXT 0x8F69
-#endif
-
-#ifndef GL_EXT_shader_pixel_local_storage
-#define GL_MAX_SHADER_PIXEL_LOCAL_STORAGE_FAST_SIZE_EXT 0x8F63
-#define GL_MAX_SHADER_PIXEL_LOCAL_STORAGE_SIZE_EXT 0x8F67
-#define GL_SHADER_PIXEL_LOCAL_STORAGE_EXT 0x8F64
-#endif
\ No newline at end of file
+#include COMPILED_PLATFORM_HEADER(OpenGL.h)
diff --git a/Engine/Source/Runtime/OpenGLDrv/Private/OpenGL3.cpp b/Engine/Source/Runtime/OpenGLDrv/Private/OpenGL3.cpp
--- a/Engine/Source/Runtime/OpenGLDrv/Private/OpenGL3.cpp
+++ b/Engine/Source/Runtime/OpenGLDrv/Private/OpenGL3.cpp
@@ -4,12 +4,11 @@
 	OpenGL3.cpp: OpenGL 3.2 implementation.
 =============================================================================*/
 
-#include "CoreMinimal.h"
-#include "HAL/IConsoleManager.h"
-#include "OpenGLDrv.h"
-#include "OpenGLDrvPrivate.h"
+#include "OpenGL3.h"
+
+#if UGL_PLATFORM_SUPPORTS_GL3
 
-#if OPENGL_GL3
+#include "OpenGLDrvPrivate.h"
 
 bool FOpenGL3::bAndroidGLESCompatibilityMode = false;
 
@@ -80,4 +79,4 @@ void FOpenGL3::ProcessExtensions( const FString& ExtensionsString )
 	bAndroidGLESCompatibilityMode = GetFeatureLevel() == ERHIFeatureLevel::ES3_1 && ExtensionsString.Contains(TEXT("GL_ARB_ES3_1_compatibility")) && FParse::Param(FCommandLine::Get(), TEXT("GLESCompat"));
 }
 
-#endif
+#endif // UGL_PLATFORM_SUPPORTS_GL3
diff --git a/Engine/Source/Runtime/OpenGLDrv/Private/OpenGL3.h b/Engine/Source/Runtime/OpenGLDrv/Private/OpenGL3.h
--- a/Engine/Source/Runtime/OpenGLDrv/Private/OpenGL3.h
+++ b/Engine/Source/Runtime/OpenGLDrv/Private/OpenGL3.h
@@ -6,29 +6,16 @@
 
 #pragma once
 
-#include "UObject/UObjectHierarchyFwd.h"
-#include "Misc/AssertionMacros.h"
-#include "HAL/UnrealMemory.h"
-#include "Containers/Array.h"
-#include "Containers/UnrealString.h"
-#include "Misc/Parse.h"
-#include "Containers/StringConv.h"
-#include "CoreGlobals.h"
+#include "OpenGLPlatform.h"
+
+#if UGL_PLATFORM_SUPPORTS_GL3
+
 #include "Containers/Map.h"
-#include "CoreMinimal.h"
-#include "Misc/CommandLine.h"
+#include "OpenGLBase.h"
 
 struct FPlatformOpenGLContext;
 struct FPlatformOpenGLDevice;
 
-template<typename KeyType,typename ValueType,typename SetAllocator ,typename KeyFuncs > class TMap;
-
-typedef GLsync UGLsync;
-
-#define OPENGL_GL3		1
-
-#include "OpenGL.h"
-
 #define USE_OPENGL_NAME_CACHE 1
 #define OPENGL_NAME_CACHE_SIZE 1024
 
@@ -652,3 +639,5 @@ protected:
 	static bool bDebugContext;
 	static bool bAndroidGLESCompatibilityMode;
 };
+
+#endif // UGL_PLATFORM_SUPPORTS_GL3
diff --git a/Engine/Source/Runtime/OpenGLDrv/Private/OpenGL4.cpp b/Engine/Source/Runtime/OpenGLDrv/Private/OpenGL4.cpp
--- a/Engine/Source/Runtime/OpenGLDrv/Private/OpenGL4.cpp
+++ b/Engine/Source/Runtime/OpenGLDrv/Private/OpenGL4.cpp
@@ -4,11 +4,11 @@
 	OpenGL4.cpp: OpenGL 4.3 implementation.
 =============================================================================*/
 
-#include "CoreMinimal.h"
-#include "OpenGLDrv.h"
-#include "OpenGLDrvPrivate.h"
+#include "OpenGL4.h"
+
+#if UGL_PLATFORM_SUPPORTS_GL4
 
-#if OPENGL_GL4
+#include "OpenGLDrvPrivate.h"
 
 GLint FOpenGL4::MaxComputeUniformComponents = -1;
 GLint FOpenGL4::MaxCombinedUAVUnits = 0;
diff --git a/Engine/Source/Runtime/OpenGLDrv/Private/OpenGL4.h b/Engine/Source/Runtime/OpenGLDrv/Private/OpenGL4.h
--- a/Engine/Source/Runtime/OpenGLDrv/Private/OpenGL4.h
+++ b/Engine/Source/Runtime/OpenGLDrv/Private/OpenGL4.h
@@ -6,7 +6,9 @@
 
 #pragma once
 
-#define OPENGL_GL4		1
+#include "OpenGLPlatform.h"
+
+#if UGL_PLATFORM_SUPPORTS_GL4
 
 #include "UObject/UObjectHierarchyFwd.h"
 #include "Misc/AssertionMacros.h"
@@ -121,3 +123,5 @@ protected:
 
 	static bool bSupportsGPUMemoryInfo;
 };
+
+#endif // UGL_PLATFORM_SUPPORTS_GL4
diff --git a/Engine/Source/Runtime/OpenGLDrv/Private/OpenGLBase.h b/Engine/Source/Runtime/OpenGLDrv/Private/OpenGLBase.h
new file mode 100644
--- /dev/null
+++ b/Engine/Source/Runtime/OpenGLDrv/Private/OpenGLBase.h
@@ -0,0 +1,389 @@
+// Copyright Epic Games, Inc. All Rights Reserved.
+
+#pragma once
+
+#include "Containers/ContainersFwd.h"
+#include "Logging/LogMacros.h"
+#include "OpenGLFunctions.h"
+#include "OpenGLPlatform.h"
+#include "PixelFormat.h"
+#include "RHIFeatureLevel.h"
+#include "RHIShaderPlatform.h"
+
+struct FPlatformOpenGLContext;
+struct FPlatformOpenGLDevice;
+
+/** OpenGL Logging. */
+OPENGLDRV_API DECLARE_LOG_CATEGORY_EXTERN(LogOpenGL, Log, VeryVerbose);
+
+#define UGL_REQUIRED_VOID			{ UE_LOG(LogOpenGL,Fatal,TEXT("%s is not supported."), ANSI_TO_TCHAR(__FUNCTION__)); }
+#define UGL_REQUIRED(ReturnValue)	{ UE_LOG(LogOpenGL,Fatal,TEXT("%s is not supported."), ANSI_TO_TCHAR(__FUNCTION__)); return (ReturnValue); }
+#define UGL_OPTIONAL_VOID			{ }
+#define UGL_OPTIONAL(ReturnValue)	{ return (ReturnValue); }
+
+#define UGL_SUPPORTS_PIXELBUFFERS		1
+#define UGL_SUPPORTS_UNIFORMBUFFERS		1
+
+// for platform extensions: from OpenGLShaders.h
+extern "C" struct FOpenGLShaderDeviceCapabilities;
+
+// for platform extensions: from OpenGLResources.h
+typedef TArray<ANSICHAR> FAnsiCharArray;
+
+// for platform extensions: from OpenGLDrvPrivate.h
+extern "C" struct FOpenGLTextureFormat;
+
+// Base static class
+class FOpenGLBase
+{
+public:
+	enum class EResourceLockMode : uint8
+	{
+		RLM_ReadWrite,
+		RLM_ReadOnly,
+		RLM_WriteOnly,
+		RLM_WriteOnlyUnsynchronized,
+		RLM_WriteOnlyPersistent,
+		RLM_ReadOnlyPersistent,
+	};
+
+	enum EQueryMode
+	{
+		QM_Result,
+		QM_ResultAvailable,
+	};
+
+	enum EFenceResult
+	{
+		FR_AlreadySignaled,
+		FR_TimeoutExpired,
+		FR_ConditionSatisfied,
+		FR_WaitFailed,
+	};
+
+	static void ProcessQueryGLInt();
+	static void ProcessExtensions(const FString& ExtensionsString);
+	static void SetupDefaultGLContextState(const FString& ExtensionsString) {};
+
+	static FORCEINLINE bool SupportsUniformBuffers()					{ return true; }
+	static FORCEINLINE bool SupportsStructuredBuffers()					{ return true; }
+	static FORCEINLINE bool SupportsTimestampQueries()					{ return true; }
+	static FORCEINLINE bool SupportsDisjointTimeQueries()				{ return false; } // @todo: if enabled, causes crash on PC
+	static FORCEINLINE bool SupportsExactOcclusionQueries()				{ return true; }
+	static FORCEINLINE bool SupportsDepthStencilReadSurface()			{ return true; }
+	static FORCEINLINE bool SupportsFloatReadSurface()					{ return true; }
+	static FORCEINLINE bool SupportsWideMRT()							{ return true; }
+	static FORCEINLINE bool SupportsPolygonMode()						{ return true; }
+	static FORCEINLINE bool SupportsTexture3D()							{ return true; }
+	static FORCEINLINE bool SupportsMobileMultiView()					{ return false; }
+	static FORCEINLINE bool SupportsImageExternal()						{ return false; }
+	static FORCEINLINE bool SupportsTextureLODBias()					{ return true; }
+	static FORCEINLINE bool SupportsTextureCompare()					{ return true; }
+	static FORCEINLINE bool SupportsDrawIndexOffset()					{ return true; }
+	static FORCEINLINE bool SupportsDiscardFrameBuffer()				{ return false; }
+	static FORCEINLINE bool SupportsIndexedExtensions()					{ return true; }
+	static FORCEINLINE bool SupportsColorBufferFloat()					{ return true; }
+	static FORCEINLINE bool SupportsColorBufferHalfFloat()				{ return true; }
+	static FORCEINLINE bool SupportsVolumeTextureRendering()			{ return false; }
+	static FORCEINLINE bool SupportsShaderFramebufferFetch()			{ return false; }
+	static FORCEINLINE bool SupportsShaderMRTFramebufferFetch()			{ return false; }
+	static FORCEINLINE bool SupportsShaderDepthStencilFetch()			{ return false; }
+	static FORCEINLINE bool SupportsPixelLocalStorage()					{ return false; }
+	static FORCEINLINE bool SupportsVertexArrayBGRA()					{ return true; }
+	static FORCEINLINE bool SupportsBGRA8888()							{ return true; }
+	static FORCEINLINE bool SupportsDXT()								{ return true; }
+	static FORCEINLINE bool SupportsASTC()								{ return bSupportsASTC; }
+	static FORCEINLINE bool SupportsASTCHDR()							{ return bSupportsASTCHDR; }
+	static FORCEINLINE bool SupportsETC2()								{ return false; }
+	static FORCEINLINE bool SupportsFramebufferSRGBEnable()				{ return true; }
+	static FORCEINLINE bool SupportsFastBufferData()					{ return true; }
+	static FORCEINLINE bool SupportsTextureFilterAnisotropic()			{ return bSupportsTextureFilterAnisotropic; }
+	static FORCEINLINE bool SupportsSeparateAlphaBlend()				{ return bSupportsDrawBuffersBlend; }
+	static FORCEINLINE void EnableSupportsClipControl()					{ bSupportsClipControl = true; }
+	static FORCEINLINE bool SupportsClipControl()						{ return bSupportsClipControl; }
+	static FORCEINLINE bool SupportsSeamlessCubeMap()					{ return false; }
+	static FORCEINLINE bool SupportsDrawIndirect()						{ return false; }
+	static FORCEINLINE bool SupportsGenerateMipmap()					{ return false; }
+	static FORCEINLINE bool SupportsBufferStorage()						{ return false; }
+	static FORCEINLINE bool SupportsDepthBoundsTest()					{ return false; }
+	static FORCEINLINE bool SupportsTextureRange()						{ return false; }
+	static FORCEINLINE bool SupportsBindlessTexture()					{ return false; }
+	static FORCEINLINE bool HasHardwareHiddenSurfaceRemoval()			{ return false; }
+	static FORCEINLINE bool AmdWorkaround()								{ return false; }
+	static FORCEINLINE bool SupportsProgramBinary()						{ return false; }
+	static FORCEINLINE bool SupportsDepthClamp()						{ return true; }
+	
+	static FORCEINLINE bool SupportsASTCDecodeMode()					{ return false; }
+
+	static FORCEINLINE GLenum GetDepthFormat()							{ return GL_DEPTH_COMPONENT16; }
+	static FORCEINLINE GLenum GetShadowDepthFormat()					{ return GL_DEPTH_COMPONENT16; }
+
+	static FORCEINLINE GLint GetMaxTextureImageUnits()			{ check(MaxTextureImageUnits != -1); return MaxTextureImageUnits; }
+	static FORCEINLINE GLint GetMaxVertexTextureImageUnits()	{ check(MaxVertexTextureImageUnits != -1); return MaxVertexTextureImageUnits; }
+	static FORCEINLINE GLint GetMaxGeometryTextureImageUnits()	{ check(MaxGeometryTextureImageUnits != -1); return MaxGeometryTextureImageUnits; }
+	static FORCEINLINE GLint GetMaxComputeTextureImageUnits()	{ check(MaxComputeTextureImageUnits != -1); return MaxComputeTextureImageUnits; }
+	static FORCEINLINE GLint GetMaxCombinedTextureImageUnits()	{ check(MaxCombinedTextureImageUnits != -1); return MaxCombinedTextureImageUnits; }
+	static FORCEINLINE GLint GetTextureBufferAlignment()		{ return TextureBufferAlignment; }
+
+
+	// Indices per unit are set in this order [Pixel, Vertex, Geometry]
+	static FORCEINLINE GLint GetFirstPixelTextureUnit()			{ return 0; }
+	static FORCEINLINE GLint GetFirstVertexTextureUnit()		{ return GetFirstPixelTextureUnit() + GetMaxTextureImageUnits(); }
+	static FORCEINLINE GLint GetFirstGeometryTextureUnit()		{ return GetFirstVertexTextureUnit() + GetMaxVertexTextureImageUnits(); }
+
+	static FORCEINLINE GLint GetFirstComputeTextureUnit()		{ return 0; }
+	
+	// Image load/store units
+	static FORCEINLINE GLint GetFirstComputeUAVUnit()			{ return 0; }
+	static FORCEINLINE GLint GetMaxComputeUAVUnits()			{ return 0; }
+	static FORCEINLINE GLint GetFirstVertexUAVUnit()			{ return 0; }
+	static FORCEINLINE GLint GetFirstPixelUAVUnit()				{ return 0; }
+	static FORCEINLINE GLint GetMaxPixelUAVUnits()				{ return 0; }
+	static FORCEINLINE GLint GetMaxCombinedUAVUnits()			{ return 0; }
+	
+	static FORCEINLINE GLint GetMaxVaryingVectors()				{ check(MaxVaryingVectors != -1); return MaxVaryingVectors; }
+	static FORCEINLINE GLint GetMaxPixelUniformComponents()		{ check(MaxPixelUniformComponents != -1); return MaxPixelUniformComponents; }
+	static FORCEINLINE GLint GetMaxVertexUniformComponents()	{ check(MaxVertexUniformComponents != -1); return MaxVertexUniformComponents; }
+	static FORCEINLINE GLint GetMaxGeometryUniformComponents()	{ check(MaxGeometryUniformComponents != -1); return MaxGeometryUniformComponents; }
+	static FORCEINLINE GLint GetMaxComputeUniformComponents()	{ return 0; }
+
+	static FORCEINLINE uint64 GetVideoMemorySize()				{ return 0; }
+
+	static FORCEINLINE bool IsDebugContent()					{ return false; }
+	static FORCEINLINE void InitDebugContext()					{ }
+
+	static FORCEINLINE int32 GetReadHalfFloatPixelsEnum() UGL_REQUIRED(0)
+
+	static FORCEINLINE GLint GetMaxMSAASamplesTileMem()			{ return 0; /* not supported */ }
+
+	// Silently ignored if not implemented:
+	static FORCEINLINE void QueryTimestampCounter(GLuint QueryID) UGL_OPTIONAL_VOID
+	static FORCEINLINE void BeginQuery(GLenum QueryType, GLuint QueryId) UGL_OPTIONAL_VOID
+	static FORCEINLINE void EndQuery(GLenum QueryType) UGL_OPTIONAL_VOID
+	static FORCEINLINE void GetQueryObject(GLuint QueryId, EQueryMode QueryMode, uint64 *OutResult) UGL_OPTIONAL_VOID
+	static FORCEINLINE void BindFragDataLocation(GLuint Program, GLuint Color, const GLchar *Name) UGL_OPTIONAL_VOID
+	static FORCEINLINE void ReadBuffer(GLenum Mode) UGL_OPTIONAL_VOID
+	static FORCEINLINE void DrawBuffer(GLenum Mode) UGL_OPTIONAL_VOID
+	static FORCEINLINE void DeleteSync(UGLsync Sync) UGL_OPTIONAL_VOID
+	static FORCEINLINE UGLsync FenceSync(GLenum Condition, GLbitfield Flags) UGL_OPTIONAL(UGLsync())
+	static FORCEINLINE bool IsSync(UGLsync Sync) UGL_OPTIONAL(false)
+	static FORCEINLINE EFenceResult ClientWaitSync(UGLsync Sync, GLbitfield Flags, GLuint64 Timeout) UGL_OPTIONAL(FR_WaitFailed)
+	static FORCEINLINE void GenSamplers(GLsizei Count, GLuint *Samplers) UGL_OPTIONAL_VOID
+	static FORCEINLINE void DeleteSamplers(GLsizei Count, GLuint *Samplers) UGL_OPTIONAL_VOID
+	static FORCEINLINE void SetSamplerParameter(GLuint Sampler, GLenum Parameter, GLint Value) UGL_OPTIONAL_VOID
+	static FORCEINLINE void BindSampler(GLuint Unit, GLuint Sampler) UGL_OPTIONAL_VOID
+	static FORCEINLINE void PolygonMode(GLenum Face, GLenum Mode) UGL_OPTIONAL_VOID
+	static FORCEINLINE void VertexAttribDivisor(GLuint Index, GLuint Divisor) UGL_OPTIONAL_VOID
+	static FORCEINLINE void PushGroupMarker(const ANSICHAR* Name) UGL_OPTIONAL_VOID
+	static FORCEINLINE void PopGroupMarker() UGL_OPTIONAL_VOID
+	static FORCEINLINE void LabelObject(GLenum Type, GLuint Object, const ANSICHAR* Name) UGL_OPTIONAL_VOID
+	static FORCEINLINE GLsizei GetLabelObject(GLenum Type, GLuint Object, GLsizei BufferSize, ANSICHAR* OutName) UGL_OPTIONAL(0)
+	static FORCEINLINE void InvalidateFramebuffer(GLenum Target, GLsizei NumAttachments, const GLenum* Attachments) UGL_OPTIONAL_VOID
+
+	// Will assert at run-time if not implemented:
+	static FORCEINLINE void* MapBufferRange(GLenum Type, uint32 InOffset, uint32 InSize, EResourceLockMode LockMode) UGL_REQUIRED(NULL)
+	static FORCEINLINE void UnmapBufferRange(GLenum Type, uint32 InOffset, uint32 InSize) UGL_REQUIRED_VOID
+	static FORCEINLINE void UnmapBuffer(GLenum Type) UGL_REQUIRED_VOID
+	static FORCEINLINE void GenQueries(GLsizei NumQueries, GLuint* QueryIDs) UGL_REQUIRED_VOID
+	static FORCEINLINE void DeleteQueries(GLsizei NumQueries, const GLuint* QueryIDs) UGL_REQUIRED_VOID
+	static FORCEINLINE void GetQueryObject(GLuint QueryId, EQueryMode QueryMode, GLuint *OutResult) UGL_REQUIRED_VOID
+	static FORCEINLINE void BindBufferBase(GLenum Target, GLuint Index, GLuint Buffer) UGL_REQUIRED_VOID
+	static FORCEINLINE void BindBufferRange(GLenum Target, GLuint Index, GLuint Buffer, GLintptr Offset, GLsizeiptr Size) UGL_REQUIRED_VOID
+	static FORCEINLINE GLuint GetUniformBlockIndex(GLuint Program, const GLchar *UniformBlockName) UGL_REQUIRED(-1)
+	static FORCEINLINE void UniformBlockBinding(GLuint Program, GLuint UniformBlockIndex, GLuint UniformBlockBinding) UGL_REQUIRED_VOID
+	static FORCEINLINE void Uniform4uiv(GLint Location, GLsizei Count, const GLuint* Value) UGL_REQUIRED_VOID
+	static FORCEINLINE void TexParameter(GLenum Target, GLenum Parameter, GLint Value) UGL_REQUIRED_VOID
+	static FORCEINLINE void FramebufferTexture(GLenum Target, GLenum Attachment, GLuint Texture, GLint Level) UGL_REQUIRED_VOID
+	static FORCEINLINE void FramebufferTexture2D(GLenum Target, GLenum Attachment, GLenum TexTarget, GLuint Texture, GLint Level)
+	{
+		glFramebufferTexture2D(Target, Attachment, TexTarget, Texture, Level);
+	}
+	static FORCEINLINE void FramebufferTexture2DMultisample(GLenum Target, GLenum Attachment, GLenum TexTarget, GLuint Texture, GLint Level, GLint NumSamples) UGL_REQUIRED_VOID
+	static FORCEINLINE void FramebufferTexture3D(GLenum Target, GLenum Attachment, GLenum TexTarget, GLuint Texture, GLint Level, GLint ZOffset) UGL_REQUIRED_VOID
+	static FORCEINLINE void FramebufferTextureLayer(GLenum Target, GLenum Attachment, GLuint Texture, GLint Level, GLint Layer) UGL_REQUIRED_VOID
+	static FORCEINLINE void FramebufferRenderbuffer(GLenum Target, GLenum Attachment, GLenum RenderBufferTarget, GLuint RenderBuffer)
+	{
+		glFramebufferRenderbuffer(Target, Attachment, RenderBufferTarget, RenderBuffer);
+	}
+	static FORCEINLINE void FramebufferTextureMultiviewOVR(GLenum Target, GLenum Attachment, GLuint Texture, GLint Level, GLint BaseViewIndex, GLsizei NumViews) UGL_REQUIRED_VOID
+	static FORCEINLINE void FramebufferTextureMultisampleMultiviewOVR(GLenum Target, GLenum Attachment, GLuint Texture, GLint Level, GLsizei NumSamples, GLint BaseViewIndex, GLsizei NumViews) UGL_REQUIRED_VOID
+	static FORCEINLINE void BlitFramebuffer(GLint SrcX0, GLint SrcY0, GLint SrcX1, GLint SrcY1, GLint DstX0, GLint DstY0, GLint DstX1, GLint DstY1, GLbitfield Mask, GLenum Filter) UGL_REQUIRED_VOID
+	static FORCEINLINE void DrawBuffers(GLsizei NumBuffers, const GLenum *Buffers) UGL_REQUIRED_VOID
+	static FORCEINLINE void DepthRange(GLdouble Near, GLdouble Far) UGL_REQUIRED_VOID
+	static FORCEINLINE void EnableIndexed(GLenum Parameter, GLuint Index) UGL_REQUIRED_VOID
+	static FORCEINLINE void DisableIndexed(GLenum Parameter, GLuint Index) UGL_REQUIRED_VOID
+	static FORCEINLINE void ColorMaskIndexed(GLuint Index, GLboolean Red, GLboolean Green, GLboolean Blue, GLboolean Alpha) UGL_REQUIRED_VOID
+	static FORCEINLINE void VertexAttribPointer(GLuint Index, GLint Size, GLenum Type, GLboolean Normalized, GLsizei Stride, const GLvoid* Pointer) UGL_REQUIRED_VOID
+	static FORCEINLINE void VertexAttribIPointer(GLuint Index, GLint Size, GLenum Type, GLsizei Stride, const GLvoid* Pointer) UGL_REQUIRED_VOID
+	static FORCEINLINE void VertexAttrib4Nsv(GLuint AttributeIndex, const GLshort* Values) UGL_REQUIRED_VOID
+	static FORCEINLINE void VertexAttrib4sv(GLuint AttributeIndex, const GLshort* Values) UGL_REQUIRED_VOID
+	static FORCEINLINE void VertexAttribI4sv(GLuint AttributeIndex, const GLshort* Values) UGL_REQUIRED_VOID
+	static FORCEINLINE void VertexAttribI4usv(GLuint AttributeIndex, const GLushort* Values) UGL_REQUIRED_VOID
+	static FORCEINLINE void VertexAttrib4Nubv(GLuint AttributeIndex, const GLubyte* Values) UGL_REQUIRED_VOID
+	static FORCEINLINE void VertexAttrib4ubv(GLuint AttributeIndex, const GLubyte* Values) UGL_REQUIRED_VOID
+	static FORCEINLINE void VertexAttribI4ubv(GLuint AttributeIndex, const GLubyte* Values) UGL_REQUIRED_VOID
+	static FORCEINLINE void VertexAttrib4Nbv(GLuint AttributeIndex, const GLbyte* Values) UGL_REQUIRED_VOID
+	static FORCEINLINE void VertexAttrib4bv(GLuint AttributeIndex, const GLbyte* Values) UGL_REQUIRED_VOID
+	static FORCEINLINE void VertexAttribI4bv(GLuint AttributeIndex, const GLbyte* Values) UGL_REQUIRED_VOID
+	static FORCEINLINE void VertexAttrib4dv(GLuint AttributeIndex, const GLdouble* Values) UGL_REQUIRED_VOID
+	static FORCEINLINE void VertexAttribI4iv(GLuint AttributeIndex, const GLint* Values) UGL_REQUIRED_VOID
+	static FORCEINLINE void VertexAttribI4uiv(GLuint AttributeIndex, const GLuint* Values) UGL_REQUIRED_VOID
+	static FORCEINLINE void DrawArraysInstanced(GLenum Mode, GLint First, GLsizei Count, GLsizei InstanceCount) UGL_REQUIRED_VOID
+	static FORCEINLINE void DrawElementsInstanced(GLenum Mode, GLsizei Count, GLenum Type, const GLvoid* Indices, GLsizei InstanceCount) UGL_REQUIRED_VOID
+	static FORCEINLINE void DrawRangeElements(GLenum Mode, GLuint Start, GLuint End, GLsizei Count, GLenum Type, const GLvoid* Indices) UGL_REQUIRED_VOID
+	static FORCEINLINE void ClearBufferfv(GLenum Buffer, GLint DrawBufferIndex, const GLfloat* Value) UGL_REQUIRED_VOID
+	static FORCEINLINE void ClearBufferfi(GLenum Buffer, GLint DrawBufferIndex, GLfloat Depth, GLint Stencil) UGL_REQUIRED_VOID
+	static FORCEINLINE void ClearBufferiv(GLenum Buffer, GLint DrawBufferIndex, const GLint* Value) UGL_REQUIRED_VOID
+	static FORCEINLINE void ClearDepth(GLdouble Depth) UGL_REQUIRED_VOID
+	static FORCEINLINE void TexImage3D(GLenum Target, GLint Level, GLint InternalFormat, GLsizei Width, GLsizei Height, GLsizei Depth, GLint Border, GLenum Format, GLenum Type, const GLvoid* PixelData) UGL_REQUIRED_VOID
+	static FORCEINLINE void CompressedTexImage3D(GLenum Target, GLint Level, GLenum InternalFormat, GLsizei Width, GLsizei Height, GLsizei Depth, GLint Border, GLsizei ImageSize, const GLvoid* PixelData) UGL_REQUIRED_VOID
+	static FORCEINLINE void TexImage2DMultisample(GLenum Target, GLsizei Samples, GLint InternalFormat, GLsizei Width, GLsizei Height, GLboolean FixedSampleLocations) UGL_REQUIRED_VOID
+	static FORCEINLINE void TexBuffer(GLenum Target, GLenum InternalFormat, GLuint Buffer) UGL_REQUIRED_VOID
+	static FORCEINLINE void TexBufferRange(GLenum Target, GLenum InternalFormat, GLuint Buffer, GLintptr Offset, GLsizeiptr Size) UGL_REQUIRED_VOID
+	static FORCEINLINE void TexSubImage3D(GLenum Target, GLint Level, GLint XOffset, GLint YOffset, GLint ZOffset, GLsizei Width, GLsizei Height, GLsizei Depth, GLenum Format, GLenum Type, const GLvoid* PixelData) UGL_REQUIRED_VOID
+	static FORCEINLINE void	CopyTexSubImage2D(GLenum Target, GLint Level, GLint XOffset, GLint YOffset, GLint X, GLint Y, GLsizei Width, GLsizei Height) UGL_REQUIRED_VOID
+	static FORCEINLINE void	CopyTexSubImage3D(GLenum Target, GLint Level, GLint XOffset, GLint YOffset, GLint ZOffset, GLint X, GLint Y, GLsizei Width, GLsizei Height) UGL_REQUIRED_VOID
+	static FORCEINLINE void GetCompressedTexImage(GLenum Target, GLint Level, GLvoid* OutImageData) UGL_REQUIRED_VOID
+	static FORCEINLINE void GetTexImage(GLenum Target, GLint Level, GLenum Format, GLenum Type, GLvoid* OutPixelData) UGL_REQUIRED_VOID
+	static FORCEINLINE void CopyBufferSubData(GLenum ReadTarget, GLenum WriteTarget, GLintptr ReadOffset, GLintptr WriteOffset, GLsizeiptr Size) UGL_REQUIRED_VOID
+	static FORCEINLINE const ANSICHAR* GetStringIndexed(GLenum Name, GLuint Index) UGL_REQUIRED(NULL)
+	static FORCEINLINE GLuint GetMajorVersion() UGL_REQUIRED(0)
+	static FORCEINLINE GLuint GetMinorVersion() UGL_REQUIRED(0)
+	static FORCEINLINE ERHIFeatureLevel::Type GetFeatureLevel() UGL_REQUIRED(ERHIFeatureLevel::SM5)
+	static FORCEINLINE EShaderPlatform GetShaderPlatform() UGL_REQUIRED(SP_NumPlatforms)
+	static FORCEINLINE FString GetAdapterName() UGL_REQUIRED(TEXT(""))
+	static FORCEINLINE void BlendFuncSeparatei(GLuint Buf, GLenum SrcRGB, GLenum DstRGB, GLenum SrcAlpha, GLenum DstAlpha) UGL_REQUIRED_VOID
+	static FORCEINLINE void BlendEquationSeparatei(GLuint Buf, GLenum ModeRGB, GLenum ModeAlpha) UGL_REQUIRED_VOID
+	static FORCEINLINE void BlendFunci(GLuint Buf, GLenum Src, GLenum Dst) UGL_REQUIRED_VOID
+	static FORCEINLINE void BlendEquationi(GLuint Buf, GLenum Mode) UGL_REQUIRED_VOID
+	static FORCEINLINE void PatchParameteri(GLenum Pname, GLint Value) UGL_REQUIRED_VOID
+	static FORCEINLINE void BindImageTexture(GLuint Unit, GLuint Texture, GLint Level, GLboolean Layered, GLint Layer, GLenum Access, GLenum Format) UGL_REQUIRED_VOID
+	static FORCEINLINE void DispatchCompute(GLuint NumGroupsX, GLuint NumGroupsY, GLuint NumGroupsZ) UGL_REQUIRED_VOID
+	static FORCEINLINE void DispatchComputeIndirect(GLintptr Offset) UGL_REQUIRED_VOID
+	static FORCEINLINE void MemoryBarrier(GLbitfield Barriers) UGL_REQUIRED_VOID
+	static FORCEINLINE bool TexStorage2D(GLenum Target, GLint Levels, GLint InternalFormat, GLsizei Width, GLsizei Height, GLenum Format, GLenum Type, ETextureCreateFlags Flags) UGL_OPTIONAL(false)
+	static FORCEINLINE bool TexStorage2DMultisample(GLenum Target, GLsizei Samples, GLint InternalFormat, GLsizei Width, GLsizei Height, GLboolean FixedSampleLocations) UGL_OPTIONAL(false)
+	static FORCEINLINE void RenderbufferStorageMultisample(GLenum Target, GLsizei Samples, GLint InternalFormat, GLsizei Width, GLsizei Height) UGL_REQUIRED_VOID
+	static FORCEINLINE void TexStorage3D(GLenum Target, GLint Levels, GLint InternalFormat, GLsizei Width, GLsizei Height, GLsizei Depth, GLenum Format, GLenum Type) UGL_REQUIRED_VOID
+	static FORCEINLINE void CompressedTexSubImage3D(GLenum Target, GLint Level, GLint XOffset, GLint YOffset, GLint ZOffset, GLsizei Width, GLsizei Height, GLsizei Depth, GLenum Format, GLsizei ImageSize, const GLvoid* PixelData) UGL_REQUIRED_VOID
+	static FORCEINLINE void CopyImageSubData(GLuint SrcName, GLenum SrcTarget, GLint SrcLevel, GLint SrcX, GLint SrcY, GLint SrcZ, GLuint DstName, GLenum DstTarget, GLint DstLevel, GLint DstX, GLint DstY, GLint DstZ, GLsizei Width, GLsizei Height, GLsizei Depth) UGL_REQUIRED_VOID
+	static FORCEINLINE void TextureView(GLuint ViewName, GLenum ViewTarget, GLuint SrcName, GLenum InternalFormat, GLuint MinLevel, GLuint NumLevels, GLuint MinLayer, GLuint NumLayers) UGL_REQUIRED_VOID
+	static FORCEINLINE void DrawArraysIndirect(GLenum Mode, const void *Offset) UGL_REQUIRED_VOID
+	static FORCEINLINE void DrawElementsIndirect(GLenum Mode, GLenum Type, const void *Offset) UGL_REQUIRED_VOID
+	static FORCEINLINE void GenerateMipmap(GLenum Target) UGL_REQUIRED_VOID
+	static FORCEINLINE void BindVertexBuffer(GLuint BindingIndex, GLuint Nuffer, GLintptr Offset, GLsizei Stride) UGL_REQUIRED_VOID
+	static FORCEINLINE void VertexAttribFormat(GLuint AttribIndex, GLint Size, GLenum Type, GLboolean Normalized, GLuint RelativeOffset) UGL_REQUIRED_VOID
+	static FORCEINLINE void VertexAttribIFormat(GLuint AttribIndex, GLint Size, GLenum Type, GLuint RelativeOffset) UGL_REQUIRED_VOID
+	static FORCEINLINE void VertexAttribBinding(GLuint AttribIndex, GLuint BindingIndex) UGL_REQUIRED_VOID
+	static FORCEINLINE void ClearBufferData(GLenum Target, GLenum InternalFormat, GLenum Format, GLenum Type, const uint32* Data) UGL_REQUIRED_VOID
+	static FORCEINLINE void VertexBindingDivisor(GLuint BindingIndex, GLuint Divisor) UGL_REQUIRED_VOID
+	static FORCEINLINE void BufferStorage(GLenum Target, GLsizeiptr Size, const void *Data, GLbitfield Flags) UGL_REQUIRED_VOID
+	static FORCEINLINE void DepthBounds(GLfloat Min, GLfloat Max) UGL_REQUIRED_VOID
+	static FORCEINLINE void TextureRange(GLenum Target, GLsizei Length, const GLvoid *Pointer) UGL_OPTIONAL_VOID
+	static FORCEINLINE void ProgramParameter (GLuint Program, GLenum PName, GLint Value) UGL_OPTIONAL_VOID
+	static FORCEINLINE void UseProgramStages(GLuint Pipeline, GLbitfield Stages, GLuint Program) { glAttachShader(Pipeline, Program); }
+	static FORCEINLINE void BindProgramPipeline(GLuint Pipeline) { glUseProgram(Pipeline); }
+	static FORCEINLINE void DeleteShader(GLuint Program) { glDeleteShader(Program); }
+	static FORCEINLINE void DeleteProgramPipelines(GLsizei Number, const GLuint *Pipelines) { for(GLsizei i = 0; i < Number; i++) { glDeleteProgram(Pipelines[i]); } }
+	static FORCEINLINE void GenProgramPipelines(GLsizei Number, GLuint *Pipelines) { check(Pipelines); for(GLsizei i = 0; i < Number; i++) { Pipelines[i] = glCreateProgram(); } }
+	static FORCEINLINE void ProgramUniform1i(GLuint Program, GLint Location, GLint V0) { glUniform1i( Location, V0 ); }
+	static FORCEINLINE void ProgramUniform4iv(GLuint Program, GLint Location, GLsizei Count, const GLint *Value) { glUniform4iv(Location, Count, Value); }
+	static FORCEINLINE void ProgramUniform4fv(GLuint Program, GLint Location, GLsizei Count, const GLfloat *Value) { glUniform4fv(Location, Count, Value); }
+	static FORCEINLINE void ProgramUniform4uiv(GLuint Program, GLint Location, GLsizei Count, const GLuint *Value) UGL_REQUIRED_VOID
+	static FORCEINLINE void GetProgramPipelineiv(GLuint Pipeline, GLenum Pname, GLint* Params) UGL_OPTIONAL_VOID
+	static FORCEINLINE void ValidateProgramPipeline(GLuint Pipeline) UGL_OPTIONAL_VOID
+	static FORCEINLINE void GetProgramPipelineInfoLog(GLuint Pipeline, GLsizei BufSize, GLsizei* Length, GLchar* InfoLog) UGL_OPTIONAL_VOID
+	static FORCEINLINE bool IsProgramPipeline(GLuint Pipeline) UGL_OPTIONAL(false)
+
+	static FORCEINLINE GLuint64 GetTextureSamplerHandle(GLuint Texture, GLuint Sampler) UGL_REQUIRED(0)
+	static FORCEINLINE GLuint64 GetTextureHandle(GLuint Texture) UGL_REQUIRED(0)
+	static FORCEINLINE void MakeTextureHandleResident(GLuint64 TextureHandle) UGL_REQUIRED_VOID
+	static FORCEINLINE void MakeTextureHandleNonResident(GLuint64 TextureHandle) UGL_REQUIRED_VOID
+	static FORCEINLINE void UniformHandleui64(GLint Location, GLuint64 Value) UGL_REQUIRED_VOID
+
+	static FORCEINLINE void GetProgramBinary(GLuint Program, GLsizei BufSize, GLsizei *Length, GLenum *BinaryFormat, void *Binary) UGL_OPTIONAL_VOID
+	static FORCEINLINE void ProgramBinary(GLuint Program, GLenum BinaryFormat, const void *Binary, GLsizei Length) UGL_OPTIONAL_VOID
+
+	static FORCEINLINE void FrameBufferFetchBarrier() UGL_OPTIONAL_VOID
+	
+	static FPlatformOpenGLDevice*	CreateDevice() UGL_REQUIRED(NULL)
+	static FPlatformOpenGLContext*	CreateContext( FPlatformOpenGLDevice* Device, void* WindowHandle ) UGL_REQUIRED(NULL)
+
+	static FORCEINLINE void CheckFrameBuffer()
+	{
+#if UE_BUILD_DEBUG || UE_BUILD_DEVELOPMENT 
+		GLenum CompleteResult = glCheckFramebufferStatus(GL_FRAMEBUFFER);
+		if (CompleteResult != GL_FRAMEBUFFER_COMPLETE)
+		{
+				UE_LOG(LogOpenGL, Fatal,TEXT("Framebuffer not complete. Status = 0x%x"), CompleteResult);
+		}
+#endif 
+	}
+
+	static FORCEINLINE void BufferSubData(GLenum Target, GLintptr Offset, GLsizeiptr Size, const GLvoid* Data)	{ glBufferSubData(Target, Offset, Size, Data); }
+	static FORCEINLINE void DeleteBuffers(GLsizei Number, const GLuint* Buffers)	{ glDeleteBuffers(Number, Buffers); }
+	static FORCEINLINE void DeleteTextures(GLsizei Number, const GLuint* Textures)	{ glDeleteTextures(Number, Textures); }
+	static FORCEINLINE void Flush()													{ glFlush(); }
+	static FORCEINLINE GLuint CreateShader(GLenum Type)								{ return glCreateShader(Type); }
+	static FORCEINLINE GLuint CreateProgram()										{ return glCreateProgram(); }
+	static FORCEINLINE bool TimerQueryDisjoint()									{ return false; }
+
+	// Calling glBufferData() to discard-reupload is slower than calling glBufferSubData() on some platforms,
+	// because changing glBufferData() with a different size (from before) may incur extra validation.
+	// To use glBufferData() discard trick: set to this to true - otherwise, glBufferSubData() will used.
+	static FORCEINLINE bool DiscardFrameBufferToResize()							{ return true; }
+
+	static FORCEINLINE EPixelFormat PreferredPixelFormatHint(EPixelFormat PreferredPixelFormat)
+	{
+		// Use a default pixel format if none was specified	
+		if (PreferredPixelFormat == EPixelFormat::PF_Unknown)
+		{
+			PreferredPixelFormat = EPixelFormat::PF_B8G8R8A8;
+		}
+		return PreferredPixelFormat;
+	}
+
+	// for platform extensions
+	static void PE_GetCurrentOpenGLShaderDeviceCapabilities(FOpenGLShaderDeviceCapabilities& Capabilities);
+	static bool PE_GLSLToDeviceCompatibleGLSL(FAnsiCharArray& GlslCodeOriginal, const FString& ShaderName, GLenum TypeEnum, const FOpenGLShaderDeviceCapabilities& Capabilities, FAnsiCharArray& GlslCode) UGL_OPTIONAL(false)
+	static void PE_SetupTextureFormat(void(*SetupTextureFormat)(EPixelFormat, const FOpenGLTextureFormat&)) UGL_OPTIONAL_VOID
+
+	static GLenum GetPlatfrom5551Format() UGL_OPTIONAL(GL_UNSIGNED_SHORT_5_5_5_1)
+protected:
+	static GLint MaxTextureImageUnits;
+	static GLint MaxCombinedTextureImageUnits;
+	static GLint MaxComputeTextureImageUnits;
+	static GLint MaxVertexTextureImageUnits;
+	static GLint MaxGeometryTextureImageUnits;
+	static GLint MaxVertexUniformComponents;
+	static GLint MaxPixelUniformComponents;
+	static GLint MaxGeometryUniformComponents;
+	static GLint MaxVaryingVectors;
+	static GLint TextureBufferAlignment;
+
+	/** GL_ARB_clip_control */
+	static bool bSupportsClipControl;
+
+	/** GL_KHR_texture_compression_astc_ldr */
+	static bool bSupportsASTC;
+
+	/** GL_KHR_texture_compression_astc_hdr */
+	static bool bSupportsASTCHDR;
+
+	/** GL_ARB_seamless_cube_map */
+	static bool bSupportsSeamlessCubemap;
+	
+	/** Can we render to texture 2D array or 3D */
+	static bool bSupportsVolumeTextureRendering;
+	
+	/** GL_EXT_texture_filter_anisotropic Can we use anisotropic filtering? */
+	static bool bSupportsTextureFilterAnisotropic;
+
+	/** GL_ARB_draw_buffers_blend */
+	static bool bSupportsDrawBuffersBlend;
+
+	/** Workaround AMD driver issues. */
+	static bool bAmdWorkaround;
+
+};
+
diff --git a/Engine/Source/Runtime/OpenGLDrv/Private/OpenGLDrv.h b/Engine/Source/Runtime/OpenGLDrv/Private/OpenGLDrv.h
--- a/Engine/Source/Runtime/OpenGLDrv/Private/OpenGLDrv.h
+++ b/Engine/Source/Runtime/OpenGLDrv/Private/OpenGLDrv.h
@@ -18,28 +18,16 @@
 #include "Templates/EnableIf.h"
 #include "BoundShaderStateHistory.h"
 
-// @todo platplug: Replace all of these includes with a call to COMPILED_PLATFORM_HEADER(OpenGLDrvPrivate.h)
-//TODO: Move these to OpenGLDrvPrivate.h
-#if PLATFORM_WINDOWS
-	#include "Windows/OpenGLWindows.h"
-#elif PLATFORM_LINUX
-	#include "Linux/OpenGLLinux.h"
-#elif PLATFORM_ANDROID
-	#include "Android/AndroidOpenGL.h"
-#else
-#include COMPILED_PLATFORM_HEADER(OpenGLDrvPrivate.h)
-#endif
+#include "OpenGLState.h"
+#include "OpenGLPlatform.h"
+#include "OpenGLUtil.h"
+#include "RenderUtils.h"
 
 // Define here so don't have to do platform filtering
 #ifndef GL_TEXTURE_EXTERNAL_OES
 #define GL_TEXTURE_EXTERNAL_OES 0x8D65
 #endif
 
-// OpenGL RHI public headers.
-#include "OpenGLUtil.h"
-#include "OpenGLState.h"
-#include "RenderUtils.h"
-
 #define FOpenGLCachedUniformBuffer_Invalid 0xFFFFFFFF
 
 class FOpenGLDynamicRHI;
diff --git a/Engine/Source/Runtime/OpenGLDrv/Private/OpenGLES.cpp b/Engine/Source/Runtime/OpenGLDrv/Private/OpenGLES.cpp
--- a/Engine/Source/Runtime/OpenGLDrv/Private/OpenGLES.cpp
+++ b/Engine/Source/Runtime/OpenGLDrv/Private/OpenGLES.cpp
@@ -4,62 +4,11 @@
 	OpenGLES.cpp: OpenGL ES implementation.
 =============================================================================*/
 
-#include "CoreMinimal.h"
-#include "HAL/IConsoleManager.h"
-#include "OpenGLDrv.h"
-#include "OpenGLDrvPrivate.h"
-
-#if !PLATFORM_DESKTOP
-
-#if OPENGL_ES
+#include "OpenGLES.h"
 
-PFNEGLGETSYSTEMTIMENVPROC eglGetSystemTimeNV_p = NULL;
-PFNEGLCREATESYNCKHRPROC eglCreateSyncKHR_p = NULL;
-PFNEGLDESTROYSYNCKHRPROC eglDestroySyncKHR_p = NULL;
-PFNEGLCLIENTWAITSYNCKHRPROC eglClientWaitSyncKHR_p = NULL;
-PFNEGLGETSYNCATTRIBKHRPROC eglGetSyncAttribKHR_p = NULL;
+#if UGL_PLATFORM_SUPPORTS_GLES
 
-namespace GLFuncPointers
-{
-	// Offscreen MSAA rendering
-	PFNGLFRAMEBUFFERTEXTURE2DMULTISAMPLEEXTPROC	glFramebufferTexture2DMultisampleEXT = NULL;
-	PFNGLRENDERBUFFERSTORAGEMULTISAMPLEEXTPROC	glRenderbufferStorageMultisampleEXT = NULL;
-
-	PFNGLPUSHGROUPMARKEREXTPROC				glPushGroupMarkerEXT = NULL;
-	PFNGLPOPGROUPMARKEREXTPROC				glPopGroupMarkerEXT = NULL;
-	PFNGLLABELOBJECTEXTPROC					glLabelObjectEXT = NULL;
-	PFNGLGETOBJECTLABELEXTPROC				glGetObjectLabelEXT = NULL;
-
-	PFNGLBUFFERSTORAGEEXTPROC				glBufferStorageEXT = NULL;
-	// KHR_debug
-	PFNGLDEBUGMESSAGECONTROLKHRPROC			glDebugMessageControlKHR = NULL;
-	PFNGLDEBUGMESSAGEINSERTKHRPROC			glDebugMessageInsertKHR = NULL;
-	PFNGLDEBUGMESSAGECALLBACKKHRPROC		glDebugMessageCallbackKHR = NULL;
-	PFNGLGETDEBUGMESSAGELOGKHRPROC			glDebugMessageLogKHR = NULL;
-	PFNGLGETPOINTERVKHRPROC					glGetPointervKHR = NULL;
-	PFNGLPUSHDEBUGGROUPKHRPROC				glPushDebugGroupKHR = NULL;
-	PFNGLPOPDEBUGGROUPKHRPROC				glPopDebugGroupKHR = NULL;
-	PFNGLOBJECTLABELKHRPROC					glObjectLabelKHR = NULL;
-	PFNGLGETOBJECTLABELKHRPROC				glGetObjectLabelKHR = NULL;
-	PFNGLOBJECTPTRLABELKHRPROC				glObjectPtrLabelKHR = NULL;
-	PFNGLGETOBJECTPTRLABELKHRPROC			glGetObjectPtrLabelKHR = NULL;
-
-	// ES 3.2
-	PFNGLTEXBUFFEREXTPROC					glTexBufferEXT = nullptr;
-	PFNGLTEXBUFFERRANGEEXTPROC				glTexBufferRangeEXT = nullptr;
-	PFNGLCOPYIMAGESUBDATAEXTPROC			glCopyImageSubData = nullptr;
-	PFNGLENABLEIEXTPROC						glEnableiEXT = nullptr;
-	PFNGLDISABLEIEXTPROC					glDisableiEXT = nullptr;
-	PFNGLBLENDEQUATIONIEXTPROC				glBlendEquationiEXT = nullptr;
-	PFNGLBLENDEQUATIONSEPARATEIEXTPROC		glBlendEquationSeparateiEXT = nullptr;
-	PFNGLBLENDFUNCIEXTPROC					glBlendFunciEXT = nullptr;
-	PFNGLBLENDFUNCSEPARATEIEXTPROC			glBlendFuncSeparateiEXT = nullptr;
-	PFNGLCOLORMASKIEXTPROC					glColorMaskiEXT = nullptr;
-	PFNGLFRAMEBUFFERTEXTUREPROC				glFramebufferTexture = nullptr;
-
-	PFNGLFRAMEBUFFERTEXTUREMULTIVIEWOVRPROC				glFramebufferTextureMultiviewOVR = NULL;
-	PFNGLFRAMEBUFFERTEXTUREMULTISAMPLEMULTIVIEWOVRPROC	glFramebufferTextureMultisampleMultiviewOVR = NULL;
-};
+#include "OpenGLDrvPrivate.h"
 
 /** GL_EXT_disjoint_timer_query */
 bool FOpenGLES::bSupportsDisjointTimeQueries = false;
@@ -338,6 +287,4 @@ void FOpenGLES::ProcessExtensions(const FString& ExtensionsString)
 	}
 }
 
-#endif
-
-#endif //desktop
+#endif //UGL_PLATFORM_SUPPORTS_GLES
diff --git a/Engine/Source/Runtime/OpenGLDrv/Private/OpenGLES.h b/Engine/Source/Runtime/OpenGLDrv/Private/OpenGLES.h
--- a/Engine/Source/Runtime/OpenGLDrv/Private/OpenGLES.h
+++ b/Engine/Source/Runtime/OpenGLDrv/Private/OpenGLES.h
@@ -6,173 +6,13 @@
 
 #pragma once
 
-#include "HAL/Platform.h"
+#include "OpenGLPlatform.h"
 
-#if !PLATFORM_DESKTOP // need this to fix compile issues with Win configuration.
+#if UGL_PLATFORM_SUPPORTS_GLES
 
-#define OPENGL_ES	1
-
-typedef GLfloat GLdouble;
-
-#include "OpenGL.h"
+#include "OpenGLBase.h"
 #include "OpenGLUtil.h"		// for VERIFY_GL
 
-#ifdef GL_AMD_debug_output
-	#undef GL_AMD_debug_output
-#endif
-
-
-
-/** Unreal tokens that maps to different OpenGL tokens by platform. */
-#undef UGL_ABGR8
-#define UGL_ABGR8				GL_UNSIGNED_BYTE
-#undef UGL_ANY_SAMPLES_PASSED
-#define UGL_ANY_SAMPLES_PASSED	GL_ANY_SAMPLES_PASSED
-#undef UGL_CLAMP_TO_BORDER
-#define UGL_CLAMP_TO_BORDER		GL_CLAMP_TO_EDGE
-#undef UGL_TIME_ELAPSED
-#define UGL_TIME_ELAPSED		GL_TIME_ELAPSED_EXT
-
-/** Official OpenGL definitions */
-#ifndef GL_FILL
-#define GL_FILL 0x1B02
-#endif
-#ifndef GL_GEOMETRY_SHADER
-#define GL_GEOMETRY_SHADER 0x8DD9
-#endif
-#ifndef GL_SAMPLER_1D
-#define GL_SAMPLER_1D 0x8B5D
-#endif
-#ifndef GL_SAMPLER_1D_SHADOW
-#define GL_SAMPLER_1D_SHADOW 0x8B61
-#endif
-#ifndef GL_DOUBLE
-#define GL_DOUBLE 0x140A
-#endif
-#ifndef GL_BGRA
-#define GL_BGRA	GL_BGRA_EXT 
-#endif
-#ifndef GL_TEXTURE_BUFFER
-#define GL_TEXTURE_BUFFER 0x8C2A
-#endif
-#ifndef GL_RGBA16
-/* GL_EXT_texture_norm16 */
-#define GL_RGBA16 GL_RGBA16_EXT
-#endif
-#ifndef GL_R16
-/* GL_EXT_texture_norm16 */
-#define GL_R16 GL_R16_EXT
-#endif
-#ifndef GL_TEXTURE_BUFFER_OFFSET_ALIGNMENT
-#define GL_TEXTURE_BUFFER_OFFSET_ALIGNMENT 0x919F
-#endif
-#ifndef GL_POLYGON_OFFSET_LINE
-#define GL_POLYGON_OFFSET_LINE 0x2A02
-#endif
-#ifndef GL_POLYGON_OFFSET_POINT
-#define GL_POLYGON_OFFSET_POINT 0x2A01
-#endif
-#ifndef GL_TEXTURE_LOD_BIAS
-#define GL_TEXTURE_LOD_BIAS 0x8501
-#endif
-#ifndef GL_FRAMEBUFFER_SRGB
-#define GL_FRAMEBUFFER_SRGB 0x8DB9
-#endif
-#ifndef GL_SAMPLES_PASSED
-#define GL_SAMPLES_PASSED 0x8914
-#endif
-#ifndef GL_POINT
-#define GL_POINT 0x1B00
-#endif
-#ifndef GL_LINE
-#define GL_LINE 0x1B01
-#endif
-#ifndef GL_TEXTURE_1D
-#define GL_TEXTURE_1D			0x0DE0
-#endif
-#ifndef GL_TEXTURE_1D_ARRAY
-#define GL_TEXTURE_1D_ARRAY		0x8C18
-#endif
-#ifndef GL_TEXTURE_RECTANGLE
-#define GL_TEXTURE_RECTANGLE	0x84F5
-#endif
-#ifndef GL_MAX_TEXTURE_BUFFER_SIZE
-#define GL_MAX_TEXTURE_BUFFER_SIZE 0x8C2B
-#endif
-#ifndef GL_DEPTH_CLAMP
-#define GL_DEPTH_CLAMP 0x864F
-#endif
-
-/** For the shader stage bits that don't exist just use 0 */
-#define GL_GEOMETRY_SHADER_BIT				0x00000000
-#define GL_TESS_CONTROL_SHADER_BIT			0x00000000
-#define GL_TESS_EVALUATION_SHADER_BIT		0x00000000
-
-// Normalize debug macros due to naming differences across GL versions
-#if defined(GL_KHR_debug) && GL_KHR_debug
-#define GL_DEBUG_SOURCE_OTHER_ARB GL_DEBUG_SOURCE_OTHER_KHR
-#define GL_DEBUG_SOURCE_API_ARB GL_DEBUG_SOURCE_API_KHR
-#define GL_DEBUG_TYPE_ERROR_ARB GL_DEBUG_TYPE_ERROR_KHR
-#define GL_DEBUG_TYPE_OTHER_ARB GL_DEBUG_TYPE_OTHER_KHR
-#define GL_DEBUG_TYPE_MARKER GL_DEBUG_TYPE_MARKER_KHR
-#define GL_DEBUG_TYPE_PUSH_GROUP GL_DEBUG_TYPE_PUSH_GROUP_KHR
-#define GL_DEBUG_TYPE_POP_GROUP GL_DEBUG_TYPE_POP_GROUP_KHR
-#define GL_DEBUG_SEVERITY_HIGH_ARB GL_DEBUG_SEVERITY_HIGH_KHR
-#define GL_DEBUG_SEVERITY_LOW_ARB GL_DEBUG_SEVERITY_LOW_KHR
-#define GL_DEBUG_SEVERITY_NOTIFICATION GL_DEBUG_SEVERITY_NOTIFICATION_KHR
-#endif
-
-// FIXME: include gl32.h
-typedef void (GL_APIENTRYP PFNGLFRAMEBUFFERTEXTUREPROC) (GLenum target, GLenum attachment, GLuint texture, GLint level);
-
-namespace GLFuncPointers
-{
-	// GL_EXT_multisampled_render_to_texture
-	extern PFNGLFRAMEBUFFERTEXTURE2DMULTISAMPLEEXTPROC	glFramebufferTexture2DMultisampleEXT;
-	extern PFNGLRENDERBUFFERSTORAGEMULTISAMPLEEXTPROC	glRenderbufferStorageMultisampleEXT;
-
-	// GL_EXT_debug_marker
-	extern PFNGLPUSHGROUPMARKEREXTPROC		glPushGroupMarkerEXT;
-	extern PFNGLPOPGROUPMARKEREXTPROC 		glPopGroupMarkerEXT;
-
-	// GL_EXT_debug_label
-	extern PFNGLLABELOBJECTEXTPROC			glLabelObjectEXT;
-	extern PFNGLGETOBJECTLABELEXTPROC		glGetObjectLabelEXT;
-	//GL_EXT_buffer_storage
-	extern PFNGLBUFFERSTORAGEEXTPROC		glBufferStorageEXT;
-
-	extern PFNGLDEBUGMESSAGECONTROLKHRPROC	glDebugMessageControlKHR;
-	extern PFNGLDEBUGMESSAGEINSERTKHRPROC	glDebugMessageInsertKHR;
-	extern PFNGLDEBUGMESSAGECALLBACKKHRPROC	glDebugMessageCallbackKHR;
-	extern PFNGLGETDEBUGMESSAGELOGKHRPROC	glDebugMessageLogKHR;
-	extern PFNGLGETPOINTERVKHRPROC			glGetPointervKHR;
-	extern PFNGLPUSHDEBUGGROUPKHRPROC		glPushDebugGroupKHR;
-	extern PFNGLPOPDEBUGGROUPKHRPROC		glPopDebugGroupKHR;
-	extern PFNGLOBJECTLABELKHRPROC			glObjectLabelKHR;
-	extern PFNGLGETOBJECTLABELKHRPROC		glGetObjectLabelKHR;
-	extern PFNGLOBJECTPTRLABELKHRPROC		glObjectPtrLabelKHR;
-	extern PFNGLGETOBJECTPTRLABELKHRPROC	glGetObjectPtrLabelKHR;
-
-	// ES 3.2
-	extern PFNGLTEXBUFFEREXTPROC				glTexBufferEXT;
-	extern PFNGLTEXBUFFERRANGEEXTPROC			glTexBufferRangeEXT;
-	extern PFNGLCOPYIMAGESUBDATAEXTPROC			glCopyImageSubData;
-	extern PFNGLENABLEIEXTPROC					glEnableiEXT;
-	extern PFNGLDISABLEIEXTPROC					glDisableiEXT;
-	extern PFNGLBLENDEQUATIONIEXTPROC			glBlendEquationiEXT;
-	extern PFNGLBLENDEQUATIONSEPARATEIEXTPROC	glBlendEquationSeparateiEXT;
-	extern PFNGLBLENDFUNCIEXTPROC				glBlendFunciEXT;
-	extern PFNGLBLENDFUNCSEPARATEIEXTPROC		glBlendFuncSeparateiEXT;
-	extern PFNGLCOLORMASKIEXTPROC				glColorMaskiEXT;
-	extern PFNGLFRAMEBUFFERTEXTUREPROC			glFramebufferTexture;
-
-	// Mobile multi-view
-	extern PFNGLFRAMEBUFFERTEXTUREMULTIVIEWOVRPROC glFramebufferTextureMultiviewOVR;
-	extern PFNGLFRAMEBUFFERTEXTUREMULTISAMPLEMULTIVIEWOVRPROC glFramebufferTextureMultisampleMultiviewOVR;
-};
-
-using namespace GLFuncPointers;
-
 struct FOpenGLES : public FOpenGLBase
 {
 	static FORCEINLINE bool IsES31Usable()
@@ -901,4 +741,4 @@ public:
 	static GLint MaxPixelUAVUnits;
 };
 
-#endif //desktop
+#endif // UGL_PLATFORM_SUPPORTS_GLES
diff --git a/Engine/Source/Runtime/OpenGLDrv/Private/OpenGLFunctions.h b/Engine/Source/Runtime/OpenGLDrv/Private/OpenGLFunctions.h
new file mode 100644
--- /dev/null
+++ b/Engine/Source/Runtime/OpenGLDrv/Private/OpenGLFunctions.h
@@ -0,0 +1,7 @@
+// Copyright Epic Games, Inc. All Rights Reserved.
+
+#pragma once
+
+#include "HAL/PreprocessorHelpers.h"
+
+#include COMPILED_PLATFORM_HEADER(OpenGLFunctions.h)
diff --git a/Engine/Source/Runtime/OpenGLDrv/Private/OpenGLPlatform.h b/Engine/Source/Runtime/OpenGLDrv/Private/OpenGLPlatform.h
new file mode 100644
--- /dev/null
+++ b/Engine/Source/Runtime/OpenGLDrv/Private/OpenGLPlatform.h
@@ -0,0 +1,245 @@
+// Copyright Epic Games, Inc. All Rights Reserved.
+
+#pragma once
+
+#include "HAL/PreprocessorHelpers.h"
+
+#include COMPILED_PLATFORM_HEADER(OpenGLPlatform.h)
+
+using UGLsync = GLsync;
+
+/** Official OpenGL definitions */
+#ifndef GL_HALF_FLOAT
+#define GL_HALF_FLOAT 0x140B
+#endif
+
+#ifndef GL_RGBA16F
+#define GL_RGBA16F    0x881A
+#endif
+
+/** OpenGL extensions */
+
+// http://www.opengl.org/registry/specs/EXT/texture_compression_s3tc.txt
+#if !defined(GL_EXT_texture_compression_s3tc)
+#define GL_COMPRESSED_RGB_S3TC_DXT1_EXT			0x83F0
+#define GL_COMPRESSED_RGBA_S3TC_DXT1_EXT		0x83F1
+#define GL_COMPRESSED_RGBA_S3TC_DXT3_EXT		0x83F2
+#define GL_COMPRESSED_RGBA_S3TC_DXT5_EXT		0x83F3
+#endif
+
+// http://www.opengl.org/registry/specs/EXT/texture_sRGB.txt
+#if !defined(GL_EXT_texture_sRGB)
+#define GL_SRGB_EXT                       0x8C40
+#define GL_SRGB8_EXT                      0x8C41
+#define GL_SRGB_ALPHA_EXT                 0x8C42
+#define GL_SRGB8_ALPHA8_EXT               0x8C43
+#define GL_SLUMINANCE_ALPHA_EXT           0x8C44
+#define GL_SLUMINANCE8_ALPHA8_EXT         0x8C45
+#define GL_SLUMINANCE_EXT                 0x8C46
+#define GL_SLUMINANCE8_EXT                0x8C47
+#define GL_COMPRESSED_SRGB_EXT            0x8C48
+#define GL_COMPRESSED_SRGB_ALPHA_EXT      0x8C49
+#define GL_COMPRESSED_SLUMINANCE_EXT      0x8C4A
+#define GL_COMPRESSED_SLUMINANCE_ALPHA_EXT 0x8C4B
+#define GL_COMPRESSED_SRGB_S3TC_DXT1_EXT  0x8C4C
+#define GL_COMPRESSED_SRGB_ALPHA_S3TC_DXT1_EXT 0x8C4D
+#define GL_COMPRESSED_SRGB_ALPHA_S3TC_DXT3_EXT 0x8C4E
+#define GL_COMPRESSED_SRGB_ALPHA_S3TC_DXT5_EXT 0x8C4F
+#endif
+
+// http://www.opengl.org/registry/specs/ARB/texture_compression_rgtc.txt
+#if !defined(GL_ARB_texture_compression_rgtc)
+#define GL_COMPRESSED_RED_RGTC1           0x8DBB
+#define GL_COMPRESSED_SIGNED_RED_RGTC1    0x8DBC
+#define GL_COMPRESSED_RG_RGTC2            0x8DBD
+#define GL_COMPRESSED_SIGNED_RG_RGTC2     0x8DBE
+#endif
+
+/* http://www.khronos.org/registry/gles/extensions/NV/NV_sRGB_formats.txt */
+#if !defined(GL_NV_sRGB_formats)
+#define GL_SLUMINANCE_NV                                        0x8C46
+#define GL_SLUMINANCE_ALPHA_NV                                  0x8C44
+#define GL_SRGB8_NV                                             0x8C41
+#define GL_SLUMINANCE8_NV                                       0x8C47
+#define GL_SLUMINANCE8_ALPHA8_NV                                0x8C45
+#define GL_COMPRESSED_SRGB_S3TC_DXT1_NV                         0x8C4C
+#define GL_COMPRESSED_SRGB_ALPHA_S3TC_DXT1_NV                   0x8C4D
+#define GL_COMPRESSED_SRGB_ALPHA_S3TC_DXT3_NV                   0x8C4E
+#define GL_COMPRESSED_SRGB_ALPHA_S3TC_DXT5_NV                   0x8C4F
+#endif
+
+// http://www.opengl.org/registry/specs/KHR/texture_compression_astc_ldr.txt
+#if !defined(GL_KHR_texture_compression_astc_ldr)
+#define GL_COMPRESSED_RGBA_ASTC_4x4_KHR            0x93B0
+#define GL_COMPRESSED_RGBA_ASTC_5x4_KHR            0x93B1
+#define GL_COMPRESSED_RGBA_ASTC_5x5_KHR            0x93B2
+#define GL_COMPRESSED_RGBA_ASTC_6x5_KHR            0x93B3
+#define GL_COMPRESSED_RGBA_ASTC_6x6_KHR            0x93B4
+#define GL_COMPRESSED_RGBA_ASTC_8x5_KHR            0x93B5
+#define GL_COMPRESSED_RGBA_ASTC_8x6_KHR            0x93B6
+#define GL_COMPRESSED_RGBA_ASTC_8x8_KHR            0x93B7
+#define GL_COMPRESSED_RGBA_ASTC_10x5_KHR           0x93B8
+#define GL_COMPRESSED_RGBA_ASTC_10x6_KHR           0x93B9
+#define GL_COMPRESSED_RGBA_ASTC_10x8_KHR           0x93BA
+#define GL_COMPRESSED_RGBA_ASTC_10x10_KHR          0x93BB
+#define GL_COMPRESSED_RGBA_ASTC_12x10_KHR          0x93BC
+#define GL_COMPRESSED_RGBA_ASTC_12x12_KHR          0x93BD
+
+#define GL_COMPRESSED_SRGB8_ALPHA8_ASTC_4x4_KHR    0x93D0
+#define GL_COMPRESSED_SRGB8_ALPHA8_ASTC_5x4_KHR    0x93D1
+#define GL_COMPRESSED_SRGB8_ALPHA8_ASTC_5x5_KHR    0x93D2
+#define GL_COMPRESSED_SRGB8_ALPHA8_ASTC_6x5_KHR    0x93D3
+#define GL_COMPRESSED_SRGB8_ALPHA8_ASTC_6x6_KHR    0x93D4
+#define GL_COMPRESSED_SRGB8_ALPHA8_ASTC_8x5_KHR    0x93D5
+#define GL_COMPRESSED_SRGB8_ALPHA8_ASTC_8x6_KHR    0x93D6
+#define GL_COMPRESSED_SRGB8_ALPHA8_ASTC_8x8_KHR    0x93D7
+#define GL_COMPRESSED_SRGB8_ALPHA8_ASTC_10x5_KHR   0x93D8
+#define GL_COMPRESSED_SRGB8_ALPHA8_ASTC_10x6_KHR   0x93D9
+#define GL_COMPRESSED_SRGB8_ALPHA8_ASTC_10x8_KHR   0x93DA
+#define GL_COMPRESSED_SRGB8_ALPHA8_ASTC_10x10_KHR  0x93DB
+#define GL_COMPRESSED_SRGB8_ALPHA8_ASTC_12x10_KHR  0x93DC
+#define GL_COMPRESSED_SRGB8_ALPHA8_ASTC_12x12_KHR  0x93DD
+#endif
+
+#if !defined(GL_TESS_EVALUATION_SHADER)
+#define GL_TESS_EVALUATION_SHADER					0x8E87
+#endif
+#if !defined(GL_TESS_CONTROL_SHADER)
+#define GL_TESS_CONTROL_SHADER						0x8E88
+#endif
+#if !defined(GL_PATCHES)
+#define GL_PATCHES									0x000E
+#endif
+#if !defined(GL_PATCH_VERTICES)
+#define GL_PATCH_VERTICES							0x8E72
+#endif
+#if !defined(GL_MAX_GEOMETRY_TEXTURE_IMAGE_UNITS)
+#define GL_MAX_GEOMETRY_TEXTURE_IMAGE_UNITS			0x8C29
+#endif
+#if !defined(GL_MAX_TESS_CONTROL_TEXTURE_IMAGE_UNITS)
+#define GL_MAX_TESS_CONTROL_TEXTURE_IMAGE_UNITS		0x8E81
+#endif
+#if !defined(GL_MAX_TESS_EVALUATION_TEXTURE_IMAGE_UNITS)
+#define GL_MAX_TESS_EVALUATION_TEXTURE_IMAGE_UNITS	0x8E82
+#endif
+#if !defined(GL_READ_WRITE)
+#define GL_READ_WRITE								0x88BA
+#endif
+#if !defined(GL_ALL_BARRIER_BITS)
+#define GL_ALL_BARRIER_BITS							0xFFFFFFFF
+#endif
+#if !defined(GL_TEXTURE_CUBE_MAP_ARRAY)
+#define GL_TEXTURE_CUBE_MAP_ARRAY					0x9009
+#endif
+#if !defined(GL_UNIFORM_BLOCK_REFERENCED_BY_TESS_CONTROL_SHADER)
+#define GL_UNIFORM_BLOCK_REFERENCED_BY_TESS_CONTROL_SHADER		0x84F0
+#endif
+#if !defined(GL_UNIFORM_BLOCK_REFERENCED_BY_TESS_EVALUATION_SHADER)
+#define GL_UNIFORM_BLOCK_REFERENCED_BY_TESS_EVALUATION_SHADER	0x84F1
+#endif
+#if !defined(GL_UNIFORM_BLOCK_REFERENCED_BY_COMPUTE_SHADER)
+#define GL_UNIFORM_BLOCK_REFERENCED_BY_COMPUTE_SHADER			0x90EC
+#endif
+#ifndef GL_ARB_seamless_cube_map
+#define GL_TEXTURE_CUBE_MAP_SEAMLESS				0x884F
+#endif
+#ifndef GL_TIME_ELAPSED_EXT
+#define GL_TIME_ELAPSED_EXT							0x88BF
+#endif
+#ifndef GL_TIMESTAMP_EXT
+#define GL_TIMESTAMP_EXT							0x8E28
+#endif
+#ifndef GL_DISPATCH_INDIRECT_BUFFER
+#define GL_DISPATCH_INDIRECT_BUFFER					0x90EE
+#endif
+#ifndef GL_DRAW_INDIRECT_BUFFER
+#define GL_DRAW_INDIRECT_BUFFER						0x8F3F
+#endif
+#ifndef GL_MAP_WRITE_BIT
+#define GL_MAP_WRITE_BIT							0x0002
+#endif
+#ifndef GL_MAP_PERSISTENT_BIT
+#define GL_MAP_PERSISTENT_BIT						0x0040
+#endif
+#ifndef GL_MAP_COHERENT_BIT
+#define GL_MAP_COHERENT_BIT							0x0080
+#endif
+#ifndef GL_DEPTH_BOUNDS_TEST_EXT
+#define GL_DEPTH_BOUNDS_TEST_EXT					0x8890
+#endif
+#ifndef GL_DEPTH_STENCIL_TEXTURE_MODE
+#define GL_DEPTH_STENCIL_TEXTURE_MODE				0x90EA
+#endif
+#ifndef GL_TEXTURE_SWIZZLE_R
+#define GL_TEXTURE_SWIZZLE_R						0x8E42
+#endif
+#ifndef GL_TEXTURE_SWIZZLE_G
+#define GL_TEXTURE_SWIZZLE_G						0x8E43
+#endif
+#ifndef GL_TEXTURE_SWIZZLE_B
+#define GL_TEXTURE_SWIZZLE_B						0x8E44
+#endif
+#ifndef GL_TEXTURE_SWIZZLE_A
+#define GL_TEXTURE_SWIZZLE_A						0x8E45
+#endif
+#ifndef GL_RED
+#define GL_RED										0x1903
+#endif
+#ifndef GL_BLUE
+#define GL_BLUE										0x1905
+#endif
+#ifndef GL_STENCIL_INDEX
+#define GL_STENCIL_INDEX							0x1901
+#endif
+#ifndef GL_RGBA_INTEGER
+#define GL_RGBA_INTEGER								0x8D99
+#endif
+
+#ifndef GL_ARB_compute_shader
+#define GL_COMPUTE_SHADER                 0x91B9
+#define GL_MAX_COMPUTE_UNIFORM_BLOCKS     0x91BB
+#define GL_MAX_COMPUTE_TEXTURE_IMAGE_UNITS 0x91BC
+#define GL_MAX_COMPUTE_IMAGE_UNIFORMS     0x91BD
+#define GL_MAX_COMPUTE_SHARED_MEMORY_SIZE 0x8262
+#define GL_MAX_COMPUTE_UNIFORM_COMPONENTS 0x8263
+#define GL_MAX_COMPUTE_ATOMIC_COUNTER_BUFFERS 0x8264
+#define GL_MAX_COMPUTE_ATOMIC_COUNTERS    0x8265
+#define GL_MAX_COMBINED_COMPUTE_UNIFORM_COMPONENTS 0x8266
+#define GL_MAX_COMPUTE_LOCAL_INVOCATIONS  0x90EB
+#define GL_MAX_COMPUTE_WORK_GROUP_COUNT   0x91BE
+#define GL_MAX_COMPUTE_WORK_GROUP_SIZE    0x91BF
+#define GL_COMPUTE_LOCAL_WORK_SIZE        0x8267
+#define GL_UNIFORM_BLOCK_REFERENCED_BY_COMPUTE_SHADER 0x90EC
+#define GL_ATOMIC_COUNTER_BUFFER_REFERENCED_BY_COMPUTE_SHADER 0x90ED
+#define GL_DISPATCH_INDIRECT_BUFFER       0x90EE
+#define GL_DISPATCH_INDIRECT_BUFFER_BINDING 0x90EF
+#define GL_COMPUTE_SHADER_BIT             0x00000020
+#endif
+
+#ifndef GL_GPU_DISJOINT_EXT
+#define GL_GPU_DISJOINT_EXT							0x8FBB
+#endif
+
+#ifndef GL_APPLE_client_storage
+#define GL_UNPACK_CLIENT_STORAGE_APPLE				0x85B2
+#endif
+
+#ifndef GL_APPLE_texture_range
+#define GL_TEXTURE_RANGE_LENGTH_APPLE     0x85B7
+#define GL_TEXTURE_RANGE_POINTER_APPLE    0x85B8
+#define GL_TEXTURE_STORAGE_HINT_APPLE     0x85BC
+#define GL_TEXTURE_MINIMIZE_STORAGE_APPLE 0x85B6
+#define GL_STORAGE_PRIVATE_APPLE          0x85BD
+#define GL_STORAGE_CACHED_APPLE           0x85BE
+#define GL_STORAGE_SHARED_APPLE           0x85BF
+#endif
+
+#ifndef TEXTURE_ASTC_DECODE_PRECISION_EXT
+#define TEXTURE_ASTC_DECODE_PRECISION_EXT 0x8F69
+#endif
+
+#ifndef GL_EXT_shader_pixel_local_storage
+#define GL_MAX_SHADER_PIXEL_LOCAL_STORAGE_FAST_SIZE_EXT 0x8F63
+#define GL_MAX_SHADER_PIXEL_LOCAL_STORAGE_SIZE_EXT 0x8F67
+#define GL_SHADER_PIXEL_LOCAL_STORAGE_EXT 0x8F64
+#endif
diff --git a/Engine/Source/Runtime/OpenGLDrv/Private/OpenGLProgramBinaryFileCache.h b/Engine/Source/Runtime/OpenGLDrv/Private/OpenGLProgramBinaryFileCache.h
--- a/Engine/Source/Runtime/OpenGLDrv/Private/OpenGLProgramBinaryFileCache.h
+++ b/Engine/Source/Runtime/OpenGLDrv/Private/OpenGLProgramBinaryFileCache.h
@@ -6,18 +6,7 @@
 
 #pragma once
 
-#include "HAL/UnrealMemory.h"
-#include "Templates/UnrealTemplate.h"
-#include "Containers/Array.h"
-#include "Misc/Crc.h"
-#include "Containers/UnrealString.h"
-#include "Containers/Map.h"
-#include "Misc/SecureHash.h"
-#include "HAL/IConsoleManager.h"
-#include "RHI.h"
-#include "ShaderCore.h"
-#include "CrossCompilerCommon.h"
-#include "ShaderCodeLibrary.h"
+#include "OpenGLShaderResources.h"
 #include "ShaderPipelineCache.h"
 
 class FOpenGLLinkedProgram;
diff --git a/Engine/Source/Runtime/OpenGLDrv/Private/OpenGLResources.h b/Engine/Source/Runtime/OpenGLDrv/Private/OpenGLResources.h
--- a/Engine/Source/Runtime/OpenGLDrv/Private/OpenGLResources.h
+++ b/Engine/Source/Runtime/OpenGLDrv/Private/OpenGLResources.h
@@ -8,6 +8,10 @@
 
 // HEADER_UNIT_SKIP - Not included directly
 
+#include "OpenGLUtil.h"
+#include "OpenGLPlatform.h"
+#include "OpenGL.h"
+
 #include "CoreTypes.h"
 #include "Misc/AssertionMacros.h"
 #include "HAL/UnrealMemory.h"
diff --git a/Engine/Source/Runtime/OpenGLDrv/Private/OpenGLState.cpp b/Engine/Source/Runtime/OpenGLDrv/Private/OpenGLState.cpp
--- a/Engine/Source/Runtime/OpenGLDrv/Private/OpenGLState.cpp
+++ b/Engine/Source/Runtime/OpenGLDrv/Private/OpenGLState.cpp
@@ -4,7 +4,7 @@
 	OpenGLState.cpp: OpenGL state implementation.
 =============================================================================*/
 
-#include "CoreMinimal.h"
+#include "OpenGLState.h"
 #include "Serialization/MemoryWriter.h"
 #include "RHI.h"
 #include "RHIUtilities.h"
diff --git a/Engine/Source/Runtime/OpenGLDrv/Private/OpenGLUtil.cpp b/Engine/Source/Runtime/OpenGLDrv/Private/OpenGLUtil.cpp
--- a/Engine/Source/Runtime/OpenGLDrv/Private/OpenGLUtil.cpp
+++ b/Engine/Source/Runtime/OpenGLDrv/Private/OpenGLUtil.cpp
@@ -4,7 +4,7 @@
 	OpenGLUtil.h: OpenGL RHI utility implementation.
 =============================================================================*/
 
-#include "CoreMinimal.h"
+#include "OpenGLUtil.h"
 #include "Stats/Stats.h"
 #include "OpenGLDrv.h"
 #include "OpenGLDrvPrivate.h"
diff --git a/Engine/Source/Runtime/OpenGLDrv/Private/Windows/OpenGLWindows.cpp b/Engine/Source/Runtime/OpenGLDrv/Private/Windows/OpenGLWindows.cpp
deleted file mode 100644
diff --git a/Engine/Source/Runtime/OpenGLDrv/Private/Windows/OpenGLWindows.h b/Engine/Source/Runtime/OpenGLDrv/Private/Windows/OpenGLWindows.h
deleted file mode 100644
diff --git a/Engine/Source/Runtime/OpenGLDrv/Private/Windows/WindowsOpenGL.h b/Engine/Source/Runtime/OpenGLDrv/Private/Windows/WindowsOpenGL.h
new file mode 100644
--- /dev/null
+++ b/Engine/Source/Runtime/OpenGLDrv/Private/Windows/WindowsOpenGL.h
@@ -0,0 +1,290 @@
+// Copyright Epic Games, Inc. All Rights Reserved.
+
+#pragma once
+
+#include "OpenGLThirdParty.h"
+#include "WindowsOpenGLPlatform.h"
+
+// Set to 1 to enable creating an ES 3.1 context and use ES 3.1 shaders on Windows
+#define EMULATE_ES31 0
+
+#if !EMULATE_ES31
+#include "OpenGL4.h"
+
+// RenderDoc defines
+#define GL_DEBUG_TOOL_EXT                 0x6789
+#define GL_DEBUG_TOOL_NAME_EXT            0x678A
+#define GL_DEBUG_TOOL_PURPOSE_EXT         0x678B
+
+struct FWindowsOpenGL : public FOpenGL4
+{
+	static FORCEINLINE void InitDebugContext()
+	{
+		extern bool GRunningUnderRenderDoc;
+		bDebugContext = glIsEnabled(GL_DEBUG_OUTPUT) != GL_FALSE || GRunningUnderRenderDoc;
+	}
+
+	static FORCEINLINE void LabelObject(GLenum Type, GLuint Object, const ANSICHAR* Name)
+	{
+		if (glObjectLabel && bDebugContext)
+		{
+			glObjectLabel(Type, Object, -1, Name);
+		}
+	}
+
+	static FORCEINLINE void PushGroupMarker(const ANSICHAR* Name)
+	{
+		if (glPushDebugGroup && bDebugContext)
+		{
+			glPushDebugGroup(GL_DEBUG_SOURCE_APPLICATION, 1, FCStringAnsi::Strlen(Name), Name);
+		}
+	}
+
+	static FORCEINLINE void PopGroupMarker()
+	{
+		if (glPopDebugGroup && bDebugContext)
+		{
+			glPopDebugGroup();
+		}
+	}
+
+	static FORCEINLINE bool TexStorage2D(GLenum Target, GLint Levels, GLint InternalFormat, GLsizei Width, GLsizei Height, GLenum Format, GLenum Type, ETextureCreateFlags Flags)
+	{
+		if (glTexStorage2D != NULL)
+		{
+			glTexStorage2D(Target, Levels, InternalFormat, Width, Height);
+			return true;
+		}
+		else
+		{
+			return false;
+		}
+	}
+
+	static FORCEINLINE bool TexStorage2DMultisample(GLenum Target, GLsizei Samples, GLint InternalFormat, GLsizei Width, GLsizei Height, GLboolean FixedSampleLocations)
+	{
+		if (glTexStorage2DMultisample != NULL)
+		{
+			glTexStorage2DMultisample(Target, Samples, InternalFormat, Width, Height, FixedSampleLocations);
+			return true;
+		}
+		else
+		{
+			return false;
+		}
+	}
+
+	static FORCEINLINE void TexStorage3D(GLenum Target, GLint Levels, GLint InternalFormat, GLsizei Width, GLsizei Height, GLsizei Depth, GLenum Format, GLenum Type)
+	{
+		if (glTexStorage3D)
+		{
+			glTexStorage3D(Target, Levels, InternalFormat, Width, Height, Depth);
+		}
+		else
+		{
+			const bool bArrayTexture = Target == GL_TEXTURE_2D_ARRAY || Target == GL_TEXTURE_CUBE_MAP_ARRAY;
+
+			for (uint32 MipIndex = 0; MipIndex < uint32(Levels); MipIndex++)
+			{
+				glTexImage3D(
+					Target,
+					MipIndex,
+					InternalFormat,
+					FMath::Max<uint32>(1, (Width >> MipIndex)),
+					FMath::Max<uint32>(1, (Height >> MipIndex)),
+					(bArrayTexture) ? Depth : FMath::Max<uint32>(1, (Depth >> MipIndex)),
+					0,
+					Format,
+					Type,
+					NULL
+				);
+			}
+		}
+	}
+
+	static FORCEINLINE void CopyImageSubData(GLuint SrcName, GLenum SrcTarget, GLint SrcLevel, GLint SrcX, GLint SrcY, GLint SrcZ, GLuint DstName, GLenum DstTarget, GLint DstLevel, GLint DstX, GLint DstY, GLint DstZ, GLsizei Width, GLsizei Height, GLsizei Depth)
+	{
+		glCopyImageSubData(SrcName, SrcTarget, SrcLevel, SrcX, SrcY, SrcZ, DstName, DstTarget, DstLevel, DstX, DstY, DstZ, Width, Height, Depth);
+	}
+
+	static FORCEINLINE bool SupportsBufferStorage()
+	{
+		return glBufferStorage != NULL;
+	}
+
+	static FORCEINLINE bool SupportsDepthBoundsTest()
+	{
+		return glDepthBoundsEXT != NULL;
+	}
+
+	static FORCEINLINE bool SupportsBindlessTexture()
+	{
+		return glGetTextureSamplerHandleARB != NULL;
+	}
+
+	static FORCEINLINE void BufferStorage(GLenum Target, GLsizeiptr Size, const void* Data, GLbitfield Flags)
+	{
+		glBufferStorage(Target, Size, Data, Flags);
+	}
+
+	static FORCEINLINE void DepthBounds(GLfloat Min, GLfloat Max)
+	{
+		glDepthBoundsEXT(Min, Max);
+	}
+
+	static FORCEINLINE GLuint64 GetTextureSamplerHandle(GLuint Texture, GLuint Sampler)
+	{
+		return glGetTextureSamplerHandleARB(Texture, Sampler);
+	}
+
+	static FORCEINLINE GLuint64 GetTextureHandle(GLuint Texture)
+	{
+		return glGetTextureHandleARB(Texture);
+	}
+
+	static FORCEINLINE void MakeTextureHandleResident(GLuint64 TextureHandle)
+	{
+		glMakeTextureHandleResidentARB(TextureHandle);
+	}
+
+	static FORCEINLINE void MakeTextureHandleNonResident(GLuint64 TextureHandle)
+	{
+		glMakeTextureHandleNonResidentARB(TextureHandle);
+	}
+
+	static FORCEINLINE void UniformHandleui64(GLint Location, GLuint64 Value)
+	{
+		glUniformHandleui64ARB(Location, Value);
+	}
+
+	static FORCEINLINE bool SupportsProgramBinary()
+	{
+		return glProgramBinary != nullptr;
+	}
+
+	static FORCEINLINE void GetProgramBinary(GLuint Program, GLsizei BufSize, GLsizei* Length, GLenum* BinaryFormat, void* Binary)
+	{
+		glGetProgramBinary(Program, BufSize, Length, BinaryFormat, Binary);
+	}
+
+	static FORCEINLINE void ProgramBinary(GLuint Program, GLenum BinaryFormat, const void* Binary, GLsizei Length)
+	{
+		glProgramBinary(Program, BinaryFormat, Binary, Length);
+	}
+};
+
+#else
+
+//fix-up naming differences between OpenGL and OpenGL ES
+#define glMapBufferOES glMapBuffer
+#define glUnmapBufferOES glUnmapBuffer
+#define GL_CLAMP_TO_BORDER_EXT GL_CLAMP_TO_BORDER
+#define GL_WRITE_ONLY_OES GL_WRITE_ONLY
+#define GL_ANY_SAMPLES_PASSED_EXT GL_ANY_SAMPLES_PASSED
+#define GL_MAX_TESS_CONTROL_UNIFORM_COMPONENTS_EXT		GL_MAX_TESS_CONTROL_UNIFORM_COMPONENTS
+#define GL_MAX_TESS_EVALUATION_UNIFORM_COMPONENTS_EXT	GL_MAX_TESS_EVALUATION_UNIFORM_COMPONENTS
+#define GL_MAX_TESS_CONTROL_TEXTURE_IMAGE_UNITS_EXT		GL_MAX_TESS_CONTROL_TEXTURE_IMAGE_UNITS
+#define GL_MAX_TESS_EVALUATION_TEXTURE_IMAGE_UNITS_EXT	GL_MAX_TESS_EVALUATION_TEXTURE_IMAGE_UNITS
+#define GL_DEBUG_SOURCE_API_KHR				GL_DEBUG_SOURCE_API
+#define GL_DEBUG_SOURCE_OTHER_KHR			GL_DEBUG_SOURCE_OTHER
+#define GL_DEBUG_SOURCE_API_KHR				GL_DEBUG_SOURCE_API
+#define GL_DEBUG_TYPE_ERROR_KHR				GL_DEBUG_TYPE_ERROR
+#define GL_DEBUG_TYPE_OTHER_KHR				GL_DEBUG_TYPE_OTHER
+#define GL_DEBUG_TYPE_ERROR_KHR				GL_DEBUG_TYPE_ERROR
+#define GL_DEBUG_TYPE_MARKER_KHR			GL_DEBUG_TYPE_MARKER
+#define GL_DEBUG_TYPE_POP_GROUP_KHR			GL_DEBUG_TYPE_POP_GROUP
+#define GL_DEBUG_TYPE_MARKER_KHR			GL_DEBUG_TYPE_MARKER
+#define GL_DEBUG_SEVERITY_HIGH_KHR			GL_DEBUG_SEVERITY_HIGH
+#define GL_DEBUG_SEVERITY_LOW_KHR			GL_DEBUG_SEVERITY_LOW
+#define GL_DEBUG_SEVERITY_HIGH_KHR			GL_DEBUG_SEVERITY_HIGH
+#define GL_DEBUG_SEVERITY_NOTIFICATION_KHR	GL_DEBUG_SEVERITY_NOTIFICATION
+#define GL_DEBUG_TYPE_ERROR_KHR				GL_DEBUG_TYPE_ERROR
+#define GL_DEBUG_SEVERITY_HIGH_KHR			GL_DEBUG_SEVERITY_HIGH
+
+#include "OpenGLES31.h"
+
+struct FWindowsOpenGL : public FOpenGLESDeferred
+{
+	static FORCEINLINE EShaderPlatform GetShaderPlatform()
+	{
+		return SP_OPENGL_PCES3_1;
+	}
+
+	static FORCEINLINE void InitDebugContext()
+	{
+		bDebugContext = glIsEnabled(GL_DEBUG_OUTPUT) != GL_FALSE;
+	}
+
+	static FORCEINLINE void LabelObject(GLenum Type, GLuint Object, const ANSICHAR* Name)
+	{
+		if (glObjectLabelKHR && bDebugContext)
+		{
+			glObjectLabelKHR(Type, Object, -1, Name);
+		}
+	}
+
+	static FORCEINLINE void PushGroupMarker(const ANSICHAR* Name)
+	{
+		if (glPushDebugGroupKHR && bDebugContext)
+		{
+			glPushDebugGroupKHR(GL_DEBUG_SOURCE_APPLICATION, 1, -1, Name);
+		}
+	}
+
+	static FORCEINLINE void PopGroupMarker()
+	{
+		if (glPopDebugGroupKHR && bDebugContext)
+		{
+			glPopDebugGroupKHR();
+		}
+	}
+
+	static FORCEINLINE bool TexStorage2D(GLenum Target, GLint Levels, GLint InternalFormat, GLsizei Width, GLsizei Height, GLenum Format, GLenum Type, ETextureCreateFlags Flags)
+	{
+		if (glTexStorage2D != NULL)
+		{
+			glTexStorage2D(Target, Levels, InternalFormat, Width, Height);
+			return true;
+		}
+		else
+		{
+			return false;
+		}
+	}
+
+	static FORCEINLINE void TexStorage3D(GLenum Target, GLint Levels, GLint InternalFormat, GLsizei Width, GLsizei Height, GLsizei Depth, GLenum Format, GLenum Type)
+	{
+		if (glTexStorage3D)
+		{
+			glTexStorage3D(Target, Levels, InternalFormat, Width, Height, Depth);
+		}
+		else
+		{
+			const bool bArrayTexture = Target == GL_TEXTURE_2D_ARRAY || Target == GL_TEXTURE_CUBE_MAP_ARRAY;
+
+			for (uint32 MipIndex = 0; MipIndex < uint32(Levels); MipIndex++)
+			{
+				glTexImage3D(
+					Target,
+					MipIndex,
+					InternalFormat,
+					FMath::Max<uint32>(1, (Width >> MipIndex)),
+					FMath::Max<uint32>(1, (Height >> MipIndex)),
+					(bArrayTexture) ? Depth : FMath::Max<uint32>(1, (Depth >> MipIndex)),
+					0,
+					Format,
+					Type,
+					NULL
+				);
+			}
+		}
+	}
+
+	static FORCEINLINE void CopyImageSubData(GLuint SrcName, GLenum SrcTarget, GLint SrcLevel, GLint SrcX, GLint SrcY, GLint SrcZ, GLuint DstName, GLenum DstTarget, GLint DstLevel, GLint DstX, GLint DstY, GLint DstZ, GLsizei Width, GLsizei Height, GLsizei Depth)
+	{
+		glCopyImageSubData(SrcName, SrcTarget, SrcLevel, SrcX, SrcY, SrcZ, DstName, DstTarget, DstLevel, DstX, DstY, DstZ, Width, Height, Depth);
+	}
+};
+#endif
+
+using FOpenGL = FWindowsOpenGL;
diff --git a/Engine/Source/Runtime/OpenGLDrv/Private/Windows/WindowsOpenGLFunctions.cpp b/Engine/Source/Runtime/OpenGLDrv/Private/Windows/WindowsOpenGLFunctions.cpp
new file mode 100644
--- /dev/null
+++ b/Engine/Source/Runtime/OpenGLDrv/Private/Windows/WindowsOpenGLFunctions.cpp
@@ -0,0 +1,9 @@
+// Copyright Epic Games, Inc. All Rights Reserved.
+
+#include "WindowsOpenGLFunctions.h"
+
+#define DEFINE_GL_ENTRYPOINTS(Type,Func) Type Func = NULL;
+	ENUM_GL_ENTRYPOINTS_ALL(DEFINE_GL_ENTRYPOINTS);
+#undef DEFINE_GL_ENTRYPOINTS
+
+PFNWGLCREATECONTEXTATTRIBSARBPROC wglCreateContextAttribsARB = NULL;
diff --git a/Engine/Source/Runtime/OpenGLDrv/Private/Windows/WindowsOpenGLFunctions.h b/Engine/Source/Runtime/OpenGLDrv/Private/Windows/WindowsOpenGLFunctions.h
new file mode 100644
--- /dev/null
+++ b/Engine/Source/Runtime/OpenGLDrv/Private/Windows/WindowsOpenGLFunctions.h
@@ -0,0 +1,397 @@
+// Copyright Epic Games, Inc. All Rights Reserved.
+
+#pragma once
+
+#include "OpenGLPlatform.h"
+
+#define UGL_PLATFORM_SUPPORTS_GLES 0
+#define UGL_PLATFORM_SUPPORTS_GL3  1
+#define UGL_PLATFORM_SUPPORTS_GL4  1
+
+/** List all OpenGL entry points used by Unreal that must be loaded from opengl32.dll */
+#define ENUM_GL_ENTRYPOINTS_DLL(EnumMacro) \
+	EnumMacro(PFNGLBINDTEXTUREPROC,glBindTexture) \
+	EnumMacro(PFNGLBLENDFUNCPROC,glBlendFunc) \
+	EnumMacro(PFNGLCOLORMASKPROC,glColorMask) \
+	EnumMacro(PFNGLCOPYTEXIMAGE1DPROC,glCopyTexImage1D) \
+	EnumMacro(PFNGLCOPYTEXIMAGE2DPROC,glCopyTexImage2D) \
+	EnumMacro(PFNGLCOPYTEXSUBIMAGE1DPROC,glCopyTexSubImage1D) \
+	EnumMacro(PFNGLCOPYTEXSUBIMAGE2DPROC,glCopyTexSubImage2D) \
+	EnumMacro(PFNGLCULLFACEPROC,glCullFace) \
+	EnumMacro(PFNGLDELETETEXTURESPROC,glDeleteTextures) \
+	EnumMacro(PFNGLDEPTHFUNCPROC,glDepthFunc) \
+	EnumMacro(PFNGLDEPTHMASKPROC,glDepthMask) \
+	EnumMacro(PFNGLDEPTHRANGEPROC,glDepthRange) \
+	EnumMacro(PFNGLDISABLEPROC,glDisable) \
+	EnumMacro(PFNGLDRAWARRAYSPROC,glDrawArrays) \
+	EnumMacro(PFNGLDRAWBUFFERPROC,glDrawBuffer) \
+	EnumMacro(PFNGLDRAWELEMENTSPROC,glDrawElements) \
+	EnumMacro(PFNGLENABLEPROC,glEnable) \
+	EnumMacro(PFNGLFINISHPROC,glFinish) \
+	EnumMacro(PFNGLFLUSHPROC,glFlush) \
+	EnumMacro(PFNGLFRONTFACEPROC,glFrontFace) \
+	EnumMacro(PFNGLGENTEXTURESPROC,glGenTextures) \
+	EnumMacro(PFNGLGETBOOLEANVPROC,glGetBooleanv) \
+	EnumMacro(PFNGLGETDOUBLEVPROC,glGetDoublev) \
+	EnumMacro(PFNGLGETERRORPROC,glGetError) \
+	EnumMacro(PFNGLGETFLOATVPROC,glGetFloatv) \
+	EnumMacro(PFNGLGETINTEGERVPROC,glGetIntegerv) \
+	EnumMacro(PFNGLGETPOINTERVPROC,glGetPointerv) \
+	EnumMacro(PFNGLGETSTRINGPROC,glGetString) \
+	EnumMacro(PFNGLGETTEXIMAGEPROC,glGetTexImage) \
+	EnumMacro(PFNGLGETTEXLEVELPARAMETERFVPROC,glGetTexLevelParameterfv) \
+	EnumMacro(PFNGLGETTEXLEVELPARAMETERIVPROC,glGetTexLevelParameteriv) \
+	EnumMacro(PFNGLGETTEXPARAMETERFVPROC,glGetTexParameterfv) \
+	EnumMacro(PFNGLGETTEXPARAMETERIVPROC,glGetTexParameteriv) \
+	EnumMacro(PFNGLHINTPROC,glHint) \
+	EnumMacro(PFNGLISENABLEDPROC,glIsEnabled) \
+	EnumMacro(PFNGLISTEXTUREPROC,glIsTexture) \
+	EnumMacro(PFNGLLINEWIDTHPROC,glLineWidth) \
+	EnumMacro(PFNGLLOGICOPPROC,glLogicOp) \
+	EnumMacro(PFNGLPIXELSTOREFPROC,glPixelStoref) \
+	EnumMacro(PFNGLPIXELSTOREIPROC,glPixelStorei) \
+	EnumMacro(PFNGLPOINTSIZEPROC,glPointSize) \
+	EnumMacro(PFNGLPOLYGONMODEPROC,glPolygonMode) \
+	EnumMacro(PFNGLPOLYGONOFFSETPROC,glPolygonOffset) \
+	EnumMacro(PFNGLREADBUFFERPROC,glReadBuffer) \
+	EnumMacro(PFNGLREADPIXELSPROC,glReadPixels) \
+	EnumMacro(PFNGLSCISSORPROC,glScissor) \
+	EnumMacro(PFNGLCLEARDEPTHPROC,glClearDepth) \
+	EnumMacro(PFNGLCLEARCOLORPROC,glClearColor) \
+	EnumMacro(PFNGLCLEARSTENCILPROC,glClearStencil) \
+	EnumMacro(PFNGLCLEARPROC,glClear) \
+	EnumMacro(PFNGLSTENCILFUNCPROC,glStencilFunc) \
+	EnumMacro(PFNGLSTENCILMASKPROC,glStencilMask) \
+	EnumMacro(PFNGLSTENCILOPPROC,glStencilOp) \
+	EnumMacro(PFNGLTEXIMAGE1DPROC,glTexImage1D) \
+	EnumMacro(PFNGLTEXIMAGE2DPROC,glTexImage2D) \
+	EnumMacro(PFNGLTEXPARAMETERFPROC,glTexParameterf) \
+	EnumMacro(PFNGLTEXPARAMETERFVPROC,glTexParameterfv) \
+	EnumMacro(PFNGLTEXPARAMETERIPROC,glTexParameteri) \
+	EnumMacro(PFNGLTEXPARAMETERIVPROC,glTexParameteriv) \
+	EnumMacro(PFNGLTEXSUBIMAGE1DPROC,glTexSubImage1D) \
+	EnumMacro(PFNGLTEXSUBIMAGE2DPROC,glTexSubImage2D) \
+	EnumMacro(PFNGLVIEWPORTPROC,glViewport)
+
+/** List all OpenGL entry points used by Unreal. */
+#define ENUM_GL_ENTRYPOINTS(EnumMacro) \
+	EnumMacro(PFNGLBINDSAMPLERPROC,glBindSampler) \
+	EnumMacro(PFNGLDELETESAMPLERSPROC,glDeleteSamplers) \
+	EnumMacro(PFNGLGENSAMPLERSPROC,glGenSamplers) \
+	EnumMacro(PFNGLSAMPLERPARAMETERIPROC,glSamplerParameteri) \
+	EnumMacro(PFNGLCLEARBUFFERFVPROC,glClearBufferfv) \
+	EnumMacro(PFNGLCLEARBUFFERIVPROC,glClearBufferiv) \
+	EnumMacro(PFNGLCLEARBUFFERUIVPROC,glClearBufferuiv) \
+	EnumMacro(PFNGLCLEARBUFFERFIPROC,glClearBufferfi) \
+	EnumMacro(PFNGLCOLORMASKIPROC,glColorMaski) \
+	EnumMacro(PFNGLDISABLEIPROC,glDisablei) \
+	EnumMacro(PFNGLENABLEIPROC,glEnablei) \
+	EnumMacro(PFNGLGETBOOLEANI_VPROC,glGetBooleani_v) \
+	EnumMacro(PFNGLGETINTEGERI_VPROC,glGetIntegeri_v) \
+	EnumMacro(PFNGLISENABLEDIPROC,glIsEnabledi) \
+	EnumMacro(PFNGLBLENDCOLORPROC,glBlendColor) \
+	EnumMacro(PFNGLBLENDEQUATIONPROC,glBlendEquation) \
+	EnumMacro(PFNGLDRAWRANGEELEMENTSPROC,glDrawRangeElements) \
+	EnumMacro(PFNGLTEXIMAGE3DPROC,glTexImage3D) \
+	EnumMacro(PFNGLTEXSUBIMAGE3DPROC,glTexSubImage3D) \
+	EnumMacro(PFNGLCOPYTEXSUBIMAGE3DPROC,glCopyTexSubImage3D) \
+	EnumMacro(PFNGLACTIVETEXTUREPROC,glActiveTexture) \
+	EnumMacro(PFNGLSAMPLECOVERAGEPROC,glSampleCoverage) \
+	EnumMacro(PFNGLCOMPRESSEDTEXIMAGE3DPROC,glCompressedTexImage3D) \
+	EnumMacro(PFNGLCOMPRESSEDTEXIMAGE2DPROC,glCompressedTexImage2D) \
+	EnumMacro(PFNGLCOMPRESSEDTEXIMAGE1DPROC,glCompressedTexImage1D) \
+	EnumMacro(PFNGLCOMPRESSEDTEXSUBIMAGE3DPROC,glCompressedTexSubImage3D) \
+	EnumMacro(PFNGLCOMPRESSEDTEXSUBIMAGE2DPROC,glCompressedTexSubImage2D) \
+	EnumMacro(PFNGLCOMPRESSEDTEXSUBIMAGE1DPROC,glCompressedTexSubImage1D) \
+	EnumMacro(PFNGLGETCOMPRESSEDTEXIMAGEPROC,glGetCompressedTexImage) \
+	EnumMacro(PFNGLBLENDFUNCSEPARATEPROC,glBlendFuncSeparate) \
+	EnumMacro(PFNGLMULTIDRAWARRAYSPROC,glMultiDrawArrays) \
+	EnumMacro(PFNGLMULTIDRAWELEMENTSPROC,glMultiDrawElements) \
+	EnumMacro(PFNGLPOINTPARAMETERFPROC,glPointParameterf) \
+	EnumMacro(PFNGLPOINTPARAMETERFVPROC,glPointParameterfv) \
+	EnumMacro(PFNGLPOINTPARAMETERIPROC,glPointParameteri) \
+	EnumMacro(PFNGLPOINTPARAMETERIVPROC,glPointParameteriv) \
+	EnumMacro(PFNGLGENQUERIESPROC,glGenQueries) \
+	EnumMacro(PFNGLDELETEQUERIESPROC,glDeleteQueries) \
+	EnumMacro(PFNGLISQUERYPROC,glIsQuery) \
+	EnumMacro(PFNGLBEGINQUERYPROC,glBeginQuery) \
+	EnumMacro(PFNGLENDQUERYPROC,glEndQuery) \
+	EnumMacro(PFNGLGETQUERYIVPROC,glGetQueryiv) \
+	EnumMacro(PFNGLGETQUERYOBJECTIVPROC,glGetQueryObjectiv) \
+	EnumMacro(PFNGLGETQUERYOBJECTUIVPROC,glGetQueryObjectuiv) \
+	EnumMacro(PFNGLBINDBUFFERPROC,glBindBuffer) \
+	EnumMacro(PFNGLBINDBUFFERBASEPROC,glBindBufferBase) \
+	EnumMacro(PFNGLDELETEBUFFERSPROC,glDeleteBuffers) \
+	EnumMacro(PFNGLGENBUFFERSPROC,glGenBuffers) \
+	EnumMacro(PFNGLISBUFFERPROC,glIsBuffer) \
+	EnumMacro(PFNGLBUFFERDATAPROC,glBufferData) \
+	EnumMacro(PFNGLBUFFERSUBDATAPROC,glBufferSubData) \
+	EnumMacro(PFNGLGETBUFFERSUBDATAPROC,glGetBufferSubData) \
+	EnumMacro(PFNGLMAPBUFFERPROC,glMapBuffer) \
+	EnumMacro(PFNGLUNMAPBUFFERPROC,glUnmapBuffer) \
+	EnumMacro(PFNGLGETBUFFERPARAMETERIVPROC,glGetBufferParameteriv) \
+	EnumMacro(PFNGLGETBUFFERPOINTERVPROC,glGetBufferPointerv) \
+	EnumMacro(PFNGLBLENDEQUATIONSEPARATEPROC,glBlendEquationSeparate) \
+	EnumMacro(PFNGLDRAWBUFFERSPROC,glDrawBuffers) \
+	EnumMacro(PFNGLSTENCILOPSEPARATEPROC,glStencilOpSeparate) \
+	EnumMacro(PFNGLSTENCILFUNCSEPARATEPROC,glStencilFuncSeparate) \
+	EnumMacro(PFNGLSTENCILMASKSEPARATEPROC,glStencilMaskSeparate) \
+	EnumMacro(PFNGLATTACHSHADERPROC,glAttachShader) \
+	EnumMacro(PFNGLBINDATTRIBLOCATIONPROC,glBindAttribLocation) \
+	EnumMacro(PFNGLBINDFRAGDATALOCATIONPROC,glBindFragDataLocation) \
+	EnumMacro(PFNGLCOMPILESHADERPROC,glCompileShader) \
+	EnumMacro(PFNGLCREATEPROGRAMPROC,glCreateProgram) \
+	EnumMacro(PFNGLCREATESHADERPROC,glCreateShader) \
+	EnumMacro(PFNGLDELETEPROGRAMPROC,glDeleteProgram) \
+	EnumMacro(PFNGLDELETESHADERPROC,glDeleteShader) \
+	EnumMacro(PFNGLDETACHSHADERPROC,glDetachShader) \
+	EnumMacro(PFNGLDISABLEVERTEXATTRIBARRAYPROC,glDisableVertexAttribArray) \
+	EnumMacro(PFNGLENABLEVERTEXATTRIBARRAYPROC,glEnableVertexAttribArray) \
+	EnumMacro(PFNGLGETACTIVEATTRIBPROC,glGetActiveAttrib) \
+	EnumMacro(PFNGLGETACTIVEUNIFORMPROC,glGetActiveUniform) \
+	EnumMacro(PFNGLGETATTACHEDSHADERSPROC,glGetAttachedShaders) \
+	EnumMacro(PFNGLGETATTRIBLOCATIONPROC,glGetAttribLocation) \
+	EnumMacro(PFNGLGETPROGRAMIVPROC,glGetProgramiv) \
+	EnumMacro(PFNGLGETPROGRAMINFOLOGPROC,glGetProgramInfoLog) \
+	EnumMacro(PFNGLGETSHADERIVPROC,glGetShaderiv) \
+	EnumMacro(PFNGLGETSHADERINFOLOGPROC,glGetShaderInfoLog) \
+	EnumMacro(PFNGLGETSHADERSOURCEPROC,glGetShaderSource) \
+	EnumMacro(PFNGLGETUNIFORMLOCATIONPROC,glGetUniformLocation) \
+	EnumMacro(PFNGLGETUNIFORMBLOCKINDEXPROC,glGetUniformBlockIndex) \
+	EnumMacro(PFNGLGETUNIFORMFVPROC,glGetUniformfv) \
+	EnumMacro(PFNGLGETUNIFORMIVPROC,glGetUniformiv) \
+	EnumMacro(PFNGLGETVERTEXATTRIBDVPROC,glGetVertexAttribdv) \
+	EnumMacro(PFNGLGETVERTEXATTRIBFVPROC,glGetVertexAttribfv) \
+	EnumMacro(PFNGLGETVERTEXATTRIBIVPROC,glGetVertexAttribiv) \
+	EnumMacro(PFNGLGETVERTEXATTRIBPOINTERVPROC,glGetVertexAttribPointerv) \
+	EnumMacro(PFNGLISPROGRAMPROC,glIsProgram) \
+	EnumMacro(PFNGLISSHADERPROC,glIsShader) \
+	EnumMacro(PFNGLLINKPROGRAMPROC,glLinkProgram) \
+	EnumMacro(PFNGLSHADERSOURCEPROC,glShaderSource) \
+	EnumMacro(PFNGLUSEPROGRAMPROC,glUseProgram) \
+	EnumMacro(PFNGLUNIFORM1FPROC,glUniform1f) \
+	EnumMacro(PFNGLUNIFORM2FPROC,glUniform2f) \
+	EnumMacro(PFNGLUNIFORM3FPROC,glUniform3f) \
+	EnumMacro(PFNGLUNIFORM4FPROC,glUniform4f) \
+	EnumMacro(PFNGLUNIFORM1IPROC,glUniform1i) \
+	EnumMacro(PFNGLUNIFORM2IPROC,glUniform2i) \
+	EnumMacro(PFNGLUNIFORM3IPROC,glUniform3i) \
+	EnumMacro(PFNGLUNIFORM4IPROC,glUniform4i) \
+	EnumMacro(PFNGLUNIFORM1FVPROC,glUniform1fv) \
+	EnumMacro(PFNGLUNIFORM2FVPROC,glUniform2fv) \
+	EnumMacro(PFNGLUNIFORM3FVPROC,glUniform3fv) \
+	EnumMacro(PFNGLUNIFORM4FVPROC,glUniform4fv) \
+	EnumMacro(PFNGLUNIFORM1IVPROC,glUniform1iv) \
+	EnumMacro(PFNGLUNIFORM2IVPROC,glUniform2iv) \
+	EnumMacro(PFNGLUNIFORM3IVPROC,glUniform3iv) \
+	EnumMacro(PFNGLUNIFORM4IVPROC,glUniform4iv) \
+	EnumMacro(PFNGLUNIFORM1UIVPROC,glUniform1uiv) \
+	EnumMacro(PFNGLUNIFORM2UIVPROC,glUniform2uiv) \
+	EnumMacro(PFNGLUNIFORM3UIVPROC,glUniform3uiv) \
+	EnumMacro(PFNGLUNIFORM4UIVPROC,glUniform4uiv) \
+	EnumMacro(PFNGLUNIFORMBLOCKBINDINGPROC,glUniformBlockBinding) \
+	EnumMacro(PFNGLUNIFORMMATRIX2FVPROC,glUniformMatrix2fv) \
+	EnumMacro(PFNGLUNIFORMMATRIX3FVPROC,glUniformMatrix3fv) \
+	EnumMacro(PFNGLUNIFORMMATRIX4FVPROC,glUniformMatrix4fv) \
+	EnumMacro(PFNGLVALIDATEPROGRAMPROC,glValidateProgram) \
+	EnumMacro(PFNGLVERTEXATTRIB1DPROC,glVertexAttrib1d) \
+	EnumMacro(PFNGLVERTEXATTRIB1DVPROC,glVertexAttrib1dv) \
+	EnumMacro(PFNGLVERTEXATTRIB1FPROC,glVertexAttrib1f) \
+	EnumMacro(PFNGLVERTEXATTRIB1FVPROC,glVertexAttrib1fv) \
+	EnumMacro(PFNGLVERTEXATTRIB1SPROC,glVertexAttrib1s) \
+	EnumMacro(PFNGLVERTEXATTRIB1SVPROC,glVertexAttrib1sv) \
+	EnumMacro(PFNGLVERTEXATTRIB2DPROC,glVertexAttrib2d) \
+	EnumMacro(PFNGLVERTEXATTRIB2DVPROC,glVertexAttrib2dv) \
+	EnumMacro(PFNGLVERTEXATTRIB2FPROC,glVertexAttrib2f) \
+	EnumMacro(PFNGLVERTEXATTRIB2FVPROC,glVertexAttrib2fv) \
+	EnumMacro(PFNGLVERTEXATTRIB2SPROC,glVertexAttrib2s) \
+	EnumMacro(PFNGLVERTEXATTRIB2SVPROC,glVertexAttrib2sv) \
+	EnumMacro(PFNGLVERTEXATTRIB3DPROC,glVertexAttrib3d) \
+	EnumMacro(PFNGLVERTEXATTRIB3DVPROC,glVertexAttrib3dv) \
+	EnumMacro(PFNGLVERTEXATTRIB3FPROC,glVertexAttrib3f) \
+	EnumMacro(PFNGLVERTEXATTRIB3FVPROC,glVertexAttrib3fv) \
+	EnumMacro(PFNGLVERTEXATTRIB3SPROC,glVertexAttrib3s) \
+	EnumMacro(PFNGLVERTEXATTRIB3SVPROC,glVertexAttrib3sv) \
+	EnumMacro(PFNGLVERTEXATTRIB4NBVPROC,glVertexAttrib4Nbv) \
+	EnumMacro(PFNGLVERTEXATTRIB4NIVPROC,glVertexAttrib4Niv) \
+	EnumMacro(PFNGLVERTEXATTRIB4NSVPROC,glVertexAttrib4Nsv) \
+	EnumMacro(PFNGLVERTEXATTRIB4NUBPROC,glVertexAttrib4Nub) \
+	EnumMacro(PFNGLVERTEXATTRIB4NUBVPROC,glVertexAttrib4Nubv) \
+	EnumMacro(PFNGLVERTEXATTRIB4NUIVPROC,glVertexAttrib4Nuiv) \
+	EnumMacro(PFNGLVERTEXATTRIB4NUSVPROC,glVertexAttrib4Nusv) \
+	EnumMacro(PFNGLVERTEXATTRIB4BVPROC,glVertexAttrib4bv) \
+	EnumMacro(PFNGLVERTEXATTRIB4DPROC,glVertexAttrib4d) \
+	EnumMacro(PFNGLVERTEXATTRIB4DVPROC,glVertexAttrib4dv) \
+	EnumMacro(PFNGLVERTEXATTRIB4FPROC,glVertexAttrib4f) \
+	EnumMacro(PFNGLVERTEXATTRIB4FVPROC,glVertexAttrib4fv) \
+	EnumMacro(PFNGLVERTEXATTRIB4IVPROC,glVertexAttrib4iv) \
+	EnumMacro(PFNGLVERTEXATTRIB4SPROC,glVertexAttrib4s) \
+	EnumMacro(PFNGLVERTEXATTRIB4SVPROC,glVertexAttrib4sv) \
+	EnumMacro(PFNGLVERTEXATTRIB4UBVPROC,glVertexAttrib4ubv) \
+	EnumMacro(PFNGLVERTEXATTRIB4UIVPROC,glVertexAttrib4uiv) \
+	EnumMacro(PFNGLVERTEXATTRIB4USVPROC,glVertexAttrib4usv) \
+	EnumMacro(PFNGLVERTEXATTRIBI4IVPROC,glVertexAttribI4iv) \
+	EnumMacro(PFNGLVERTEXATTRIBI4UIVPROC,glVertexAttribI4uiv) \
+	EnumMacro(PFNGLVERTEXATTRIBI4SVPROC,glVertexAttribI4sv) \
+	EnumMacro(PFNGLVERTEXATTRIBI4USVPROC,glVertexAttribI4usv) \
+	EnumMacro(PFNGLVERTEXATTRIBI4BVPROC,glVertexAttribI4bv) \
+	EnumMacro(PFNGLVERTEXATTRIBI4UBVPROC,glVertexAttribI4ubv) \
+	EnumMacro(PFNGLVERTEXATTRIBPOINTERPROC,glVertexAttribPointer) \
+	EnumMacro(PFNGLVERTEXATTRIBIPOINTERPROC,glVertexAttribIPointer) \
+	EnumMacro(PFNGLUNIFORMMATRIX2X3FVPROC,glUniformMatrix2x3fv) \
+	EnumMacro(PFNGLUNIFORMMATRIX3X2FVPROC,glUniformMatrix3x2fv) \
+	EnumMacro(PFNGLUNIFORMMATRIX2X4FVPROC,glUniformMatrix2x4fv) \
+	EnumMacro(PFNGLUNIFORMMATRIX4X2FVPROC,glUniformMatrix4x2fv) \
+	EnumMacro(PFNGLUNIFORMMATRIX3X4FVPROC,glUniformMatrix3x4fv) \
+	EnumMacro(PFNGLUNIFORMMATRIX4X3FVPROC,glUniformMatrix4x3fv) \
+	EnumMacro(PFNGLISRENDERBUFFERPROC,glIsRenderbuffer) \
+	EnumMacro(PFNGLBINDRENDERBUFFERPROC,glBindRenderbuffer) \
+	EnumMacro(PFNGLDELETERENDERBUFFERSPROC,glDeleteRenderbuffers) \
+	EnumMacro(PFNGLGENRENDERBUFFERSPROC,glGenRenderbuffers) \
+	EnumMacro(PFNGLRENDERBUFFERSTORAGEPROC,glRenderbufferStorage) \
+	EnumMacro(PFNGLGETRENDERBUFFERPARAMETERIVPROC,glGetRenderbufferParameteriv) \
+	EnumMacro(PFNGLISFRAMEBUFFERPROC,glIsFramebuffer) \
+	EnumMacro(PFNGLBINDFRAMEBUFFERPROC,glBindFramebuffer) \
+	EnumMacro(PFNGLDELETEFRAMEBUFFERSPROC,glDeleteFramebuffers) \
+	EnumMacro(PFNGLGENFRAMEBUFFERSPROC,glGenFramebuffers) \
+	EnumMacro(PFNGLCHECKFRAMEBUFFERSTATUSPROC,glCheckFramebufferStatus) \
+	EnumMacro(PFNGLFRAMEBUFFERTEXTURE1DPROC,glFramebufferTexture1D) \
+	EnumMacro(PFNGLFRAMEBUFFERTEXTURE2DPROC,glFramebufferTexture2D) \
+	EnumMacro(PFNGLFRAMEBUFFERTEXTURE3DPROC,glFramebufferTexture3D) \
+	EnumMacro(PFNGLFRAMEBUFFERTEXTUREPROC,glFramebufferTexture) \
+	EnumMacro(PFNGLFRAMEBUFFERRENDERBUFFERPROC,glFramebufferRenderbuffer) \
+	EnumMacro(PFNGLGETFRAMEBUFFERATTACHMENTPARAMETERIVPROC,glGetFramebufferAttachmentParameteriv) \
+	EnumMacro(PFNGLGENERATEMIPMAPPROC,glGenerateMipmap) \
+	EnumMacro(PFNGLBLITFRAMEBUFFERPROC,glBlitFramebuffer) \
+	EnumMacro(PFNGLRENDERBUFFERSTORAGEMULTISAMPLEPROC,glRenderbufferStorageMultisample) \
+	EnumMacro(PFNGLFRAMEBUFFERTEXTURELAYERPROC,glFramebufferTextureLayer) \
+	EnumMacro(PFNGLMAPBUFFERRANGEPROC,glMapBufferRange) \
+	EnumMacro(PFNGLFLUSHMAPPEDBUFFERRANGEPROC,glFlushMappedBufferRange) \
+	EnumMacro(PFNGLVERTEXATTRIBDIVISORPROC,glVertexAttribDivisor) \
+	EnumMacro(PFNGLDRAWARRAYSINSTANCEDPROC,glDrawArraysInstanced) \
+	EnumMacro(PFNGLDRAWELEMENTSINSTANCEDPROC,glDrawElementsInstanced) \
+	EnumMacro(PFNGLGETSTRINGIPROC,glGetStringi) \
+	EnumMacro(PFNGLGENVERTEXARRAYSPROC,glGenVertexArrays) \
+	EnumMacro(PFNGLDELETEVERTEXARRAYSPROC,glDeleteVertexArrays) \
+	EnumMacro(PFNGLBINDVERTEXARRAYPROC,glBindVertexArray) \
+	EnumMacro(PFNGLCOPYBUFFERSUBDATAPROC,glCopyBufferSubData) \
+	EnumMacro(PFNGLTEXBUFFERPROC,glTexBuffer) \
+	EnumMacro(PFNGLTEXBUFFERRANGEPROC,glTexBufferRange) \
+	EnumMacro(PFNGLTEXIMAGE2DMULTISAMPLEPROC,glTexImage2DMultisample) \
+	EnumMacro(PFNGLQUERYCOUNTERPROC, glQueryCounter)\
+	EnumMacro(PFNGLISSYNCPROC, glIsSync)\
+	EnumMacro(PFNGLDELETESYNCPROC, glDeleteSync)\
+	EnumMacro(PFNGLGETQUERYOBJECTUI64VPROC, glGetQueryObjectui64v)\
+	EnumMacro(PFNGLFENCESYNCPROC, glFenceSync)\
+	EnumMacro(PFNGLGETSYNCIVPROC, glGetSynciv)\
+	EnumMacro(PFNGLCLIENTWAITSYNCPROC, glClientWaitSync)\
+	EnumMacro(PFNGLBINDBUFFERRANGEPROC, glBindBufferRange) \
+	EnumMacro(PFNGLGETPROGRAMRESOURCEINDEXPROC, glGetProgramResourceIndex) \
+
+#define ENUM_GL_ENTRYPOINTS_OPTIONAL(EnumMacro) \
+	EnumMacro(PFNGLCLIPCONTROLPROC,glClipControl) \
+	EnumMacro(PFNGLDEBUGMESSAGECALLBACKARBPROC,glDebugMessageCallbackARB) \
+	EnumMacro(PFNGLDEBUGMESSAGECONTROLARBPROC,glDebugMessageControlARB) \
+	EnumMacro(PFNGLDEBUGMESSAGECALLBACKAMDPROC,glDebugMessageCallbackAMD) \
+	EnumMacro(PFNGLDEBUGMESSAGEENABLEAMDPROC,glDebugMessageEnableAMD) \
+	EnumMacro(PFNGLGETACTIVEUNIFORMSIVPROC,glGetActiveUniformsiv) \
+	EnumMacro(PFNGLGETVERTEXATTRIBIUIVPROC,glGetVertexAttribIuiv) \
+	EnumMacro(PFNGLGETACTIVEUNIFORMNAMEPROC,glGetActiveUniformName) \
+	EnumMacro(PFNGLGETUNIFORMUIVPROC,glGetUniformuiv) \
+	EnumMacro(PFNGLGETACTIVEUNIFORMBLOCKIVPROC,glGetActiveUniformBlockiv) \
+	EnumMacro(PFNGLGETBUFFERPARAMETERI64VPROC,glGetBufferParameteri64v) \
+	EnumMacro(PFNGLGETACTIVEUNIFORMBLOCKNAMEPROC,glGetActiveUniformBlockName) \
+	EnumMacro(PFNGLGETSAMPLERPARAMETERFVPROC,glGetSamplerParameterfv) \
+	EnumMacro(PFNGLGETSAMPLERPARAMETERIVPROC,glGetSamplerParameteriv) \
+	EnumMacro(PFNGLDISPATCHCOMPUTEPROC, glDispatchCompute) \
+	EnumMacro(PFNGLDISPATCHCOMPUTEINDIRECTPROC, glDispatchComputeIndirect) \
+	EnumMacro(PFNGLBINDIMAGETEXTUREPROC, glBindImageTexture) \
+	EnumMacro(PFNGLMEMORYBARRIERPROC, glMemoryBarrier) \
+	EnumMacro(PFNGLBLENDEQUATIONIPROC, glBlendEquationi) \
+	EnumMacro(PFNGLBLENDEQUATIONSEPARATEIPROC, glBlendEquationSeparatei) \
+	EnumMacro(PFNGLBLENDFUNCIPROC, glBlendFunci) \
+	EnumMacro(PFNGLBLENDFUNCSEPARATEIPROC, glBlendFuncSeparatei)\
+	EnumMacro(PFNGLBLENDEQUATIONSEPARATEIARBPROC,glBlendEquationSeparateiARB)\
+	EnumMacro(PFNGLBLENDEQUATIONIARBPROC,glBlendEquationiARB)\
+	EnumMacro(PFNGLBLENDFUNCSEPARATEIARBPROC,glBlendFuncSeparateiARB)\
+	EnumMacro(PFNGLBLENDFUNCIARBPROC,glBlendFunciARB)\
+	EnumMacro(PFNGLCLEARBUFFERDATAPROC, glClearBufferData)\
+	EnumMacro(PFNGLCLEARBUFFERSUBDATAPROC, glClearBufferSubData)\
+	EnumMacro(PFNGLPUSHDEBUGGROUPPROC, glPushDebugGroup)\
+	EnumMacro(PFNGLPOPDEBUGGROUPPROC, glPopDebugGroup)\
+	EnumMacro(PFNGLOBJECTLABELPROC, glObjectLabel)\
+	EnumMacro(PFNGLOBJECTPTRLABELPROC, glObjectPtrLabel)\
+	EnumMacro(PFNGLPATCHPARAMETERIPROC, glPatchParameteri)\
+	EnumMacro(PFNGLBINDVERTEXBUFFERPROC, glBindVertexBuffer)\
+	EnumMacro(PFNGLVERTEXATTRIBFORMATPROC, glVertexAttribFormat)\
+	EnumMacro(PFNGLVERTEXATTRIBIFORMATPROC, glVertexAttribIFormat)\
+	EnumMacro(PFNGLVERTEXATTRIBBINDINGPROC, glVertexAttribBinding)\
+	EnumMacro(PFNGLVERTEXBINDINGDIVISORPROC, glVertexBindingDivisor)\
+	EnumMacro(PFNGLCOPYIMAGESUBDATAPROC, glCopyImageSubData)\
+	EnumMacro(PFNGLTEXSTORAGE1DPROC, glTexStorage1D)\
+	EnumMacro(PFNGLTEXSTORAGE2DPROC, glTexStorage2D)\
+	EnumMacro(PFNGLTEXSTORAGE3DPROC, glTexStorage3D)\
+	EnumMacro(PFNGLBUFFERSTORAGEPROC, glBufferStorage)\
+	EnumMacro(PFNGLTEXTUREVIEWPROC, glTextureView)\
+	EnumMacro(PFNGLTEXSTORAGE2DMULTISAMPLEPROC, glTexStorage2DMultisample)\
+	EnumMacro(PFNGLDRAWELEMENTSINDIRECTPROC, glDrawElementsIndirect)\
+	EnumMacro(PFNGLDRAWARRAYSINDIRECTPROC, glDrawArraysIndirect)\
+	EnumMacro(PFNGLDEPTHBOUNDSEXTPROC, glDepthBoundsEXT)\
+	EnumMacro(PFNGLGETTEXTUREHANDLENVPROC, glGetTextureHandleARB)\
+	EnumMacro(PFNGLGETTEXTURESAMPLERHANDLENVPROC, glGetTextureSamplerHandleARB)\
+	EnumMacro(PFNGLMAKETEXTUREHANDLERESIDENTNVPROC, glMakeTextureHandleResidentARB)\
+	EnumMacro(PFNGLUNIFORMHANDLEUI64NVPROC, glUniformHandleui64ARB)\
+	EnumMacro(PFNGLMAKETEXTUREHANDLENONRESIDENTNVPROC, glMakeTextureHandleNonResidentARB)\
+	EnumMacro(PFNGLPUSHDEBUGGROUPPROC, glPushDebugGroupKHR)\
+	EnumMacro(PFNGLPOPDEBUGGROUPPROC, glPopDebugGroupKHR)\
+	EnumMacro(PFNGLOBJECTLABELPROC, glObjectLabelKHR)\
+	EnumMacro(PFNGLOBJECTPTRLABELPROC, glObjectPtrLabelKHR)\
+	EnumMacro(PFNGLDEBUGMESSAGECALLBACKARBPROC,glDebugMessageCallbackKHR) \
+	EnumMacro(PFNGLDEBUGMESSAGECONTROLARBPROC,glDebugMessageControlKHR) \
+	EnumMacro(PFNGLPATCHPARAMETERIPROC, glPatchParameteriEXT)\
+	EnumMacro(PFNGLTEXTUREVIEWPROC, glTextureViewEXT)\
+	EnumMacro(PFNGLBLENDEQUATIONIPROC, glBlendEquationiEXT) \
+	EnumMacro(PFNGLBLENDEQUATIONSEPARATEIPROC, glBlendEquationSeparateiEXT) \
+	EnumMacro(PFNGLBLENDFUNCIPROC, glBlendFunciEXT) \
+	EnumMacro(PFNGLBLENDFUNCSEPARATEIPROC, glBlendFuncSeparateiEXT)\
+	EnumMacro(PFNGLCOLORMASKIPROC,glColorMaskiEXT) \
+	EnumMacro(PFNGLDISABLEIPROC,glDisableiEXT) \
+	EnumMacro(PFNGLENABLEIPROC,glEnableiEXT) \
+	EnumMacro(PFNGLFRAMEBUFFERTEXTUREPROC,glFramebufferTextureEXT) \
+	EnumMacro(PFNGLTEXBUFFERPROC,glTexBufferEXT) \
+	EnumMacro(PFNGLTEXBUFFERRANGEPROC,glTexBufferRangeEXT) \
+	EnumMacro(PFNGLDEPTHRANGEFPROC,glDepthRangef) \
+	EnumMacro(PFNGLCLEARDEPTHFPROC,glClearDepthf) \
+	EnumMacro(PFNGLGETSHADERPRECISIONFORMATPROC, glGetShaderPrecisionFormat) \
+	EnumMacro(PFNGLPROGRAMPARAMETERIPROC, glProgramParameteri) \
+	EnumMacro(PFNGLUSEPROGRAMSTAGESPROC, glUseProgramStages) \
+	EnumMacro(PFNGLBINDPROGRAMPIPELINEPROC, glBindProgramPipeline) \
+	EnumMacro(PFNGLDELETEPROGRAMPIPELINESPROC, glDeleteProgramPipelines) \
+	EnumMacro(PFNGLGENPROGRAMPIPELINESPROC, glGenProgramPipelines) \
+	EnumMacro(PFNGLPROGRAMUNIFORM1IPROC, glProgramUniform1i) \
+	EnumMacro(PFNGLPROGRAMUNIFORM4IVPROC, glProgramUniform4iv) \
+	EnumMacro(PFNGLPROGRAMUNIFORM4FVPROC, glProgramUniform4fv) \
+	EnumMacro(PFNGLPROGRAMUNIFORM4UIVPROC, glProgramUniform4uiv) \
+	EnumMacro(PFNGLGETPROGRAMPIPELINEIVPROC, glGetProgramPipelineiv) \
+	EnumMacro(PFNGLVALIDATEPROGRAMPIPELINEPROC, glValidateProgramPipeline) \
+	EnumMacro(PFNGLGETPROGRAMPIPELINEINFOLOGPROC, glGetProgramPipelineInfoLog) \
+	EnumMacro(PFNGLISPROGRAMPIPELINEPROC, glIsProgramPipeline) \
+	EnumMacro(PFNGLGETPROGRAMBINARYPROC, glGetProgramBinary) \
+	EnumMacro(PFNGLPROGRAMBINARYPROC, glProgramBinary)
+
+/** List of all OpenGL entry points. */
+#define ENUM_GL_ENTRYPOINTS_ALL(EnumMacro) \
+	ENUM_GL_ENTRYPOINTS_DLL(EnumMacro) \
+	ENUM_GL_ENTRYPOINTS(EnumMacro) \
+	ENUM_GL_ENTRYPOINTS_OPTIONAL(EnumMacro)
+
+/** Declare all GL functions. */
+#define DECLARE_GL_ENTRYPOINTS(Type,Func) extern Type OPENGLDRV_API Func;
+ENUM_GL_ENTRYPOINTS_ALL(DECLARE_GL_ENTRYPOINTS);
+#undef DECLARE_GL_ENTRYPOINTS
+
+/** This function is handled separately because it is used to get a real context. */
+extern PFNWGLCREATECONTEXTATTRIBSARBPROC wglCreateContextAttribsARB;
+
+// set in OpenGLDevice.cpp
+extern PFNWGLSWAPINTERVALEXTPROC wglSwapIntervalEXT_ProcAddress;
diff --git a/Engine/Source/Runtime/OpenGLDrv/Private/Windows/WindowsOpenGLPlatform.cpp b/Engine/Source/Runtime/OpenGLDrv/Private/Windows/WindowsOpenGLPlatform.cpp
new file mode 100644
--- /dev/null
+++ b/Engine/Source/Runtime/OpenGLDrv/Private/Windows/WindowsOpenGLPlatform.cpp
@@ -0,0 +1,1009 @@
+// Copyright Epic Games, Inc. All Rights Reserved.
+
+#include "WindowsOpenGLPlatform.h"
+
+#include "Misc/ScopeLock.h"
+#include "OpenGLDrvPrivate.h"
+#include "OpenGLUtil.h"
+#include "RHI.h"
+#include "RHIUtilities.h"
+
+bool GRunningUnderRenderDoc = false;
+
+/*------------------------------------------------------------------------------
+	OpenGL context management.
+------------------------------------------------------------------------------*/
+
+static void ContextMakeCurrent(HDC DC, HGLRC RC)
+{
+	BOOL Result = wglMakeCurrent(DC, RC);
+	if (!Result)
+	{
+		Result = wglMakeCurrent(nullptr, nullptr);
+	}
+	check(Result);
+}
+
+static HGLRC GetCurrentContext()
+{
+	return wglGetCurrentContext();
+}
+
+
+/** Platform specific OpenGL context. */
+struct FPlatformOpenGLContext
+{
+	HWND WindowHandle;
+	HDC DeviceContext;
+	HGLRC OpenGLContext;
+	bool bReleaseWindowOnDestroy;
+	int32 SyncInterval;
+	GLuint	ViewportFramebuffer;
+	GLuint	VertexArrayObject;	// one has to be generated and set for each context (OpenGL 3.2 Core requirements)
+	GLuint	BackBufferResource;
+	GLenum	BackBufferTarget;
+};
+
+class FScopeContext
+{
+public:
+	FScopeContext(FPlatformOpenGLContext* Context)
+	{
+		check(Context);
+		PrevDC = wglGetCurrentDC();
+		PrevContext = GetCurrentContext();
+		bSameDC = (PrevDC == Context->DeviceContext);
+		bSameDCAndContext = (PrevContext == Context->OpenGLContext) && bSameDC;
+		if (!bSameDCAndContext)
+		{
+			//			if (PrevContext)
+			//			{
+			//				glFlush();
+			//			}
+						// no need to glFlush() on Windows, it does flush by itself before switching contexts
+			ContextMakeCurrent(Context->DeviceContext, Context->OpenGLContext);
+		}
+	}
+
+	~FScopeContext(void)
+	{
+		if (!bSameDCAndContext)
+		{
+			//			glFlush();	// not needed on Windows, it does flush by itself before switching contexts
+			if (PrevContext)
+			{
+				ContextMakeCurrent(PrevDC, PrevContext);
+			}
+			else
+			{
+				ContextMakeCurrent(NULL, NULL);
+			}
+		}
+	}
+
+	bool ContextWasAlreadyActive() const
+	{
+		return bSameDCAndContext;
+	}
+
+	bool ContextsShareSameDC() const
+	{
+		return bSameDC;
+	}
+
+private:
+	HDC					PrevDC;
+	HGLRC				PrevContext;
+	bool				bSameDCAndContext;
+	bool				bSameDC;
+};
+
+void DeleteQueriesForCurrentContext(HGLRC Context);
+
+/**
+ * A dummy wndproc.
+ */
+static LRESULT CALLBACK PlatformDummyGLWndproc(HWND hWnd, uint32 Message, WPARAM wParam, LPARAM lParam)
+{
+	return DefWindowProc(hWnd, Message, wParam, lParam);
+}
+
+/**
+ * Initialize a pixel format descriptor for the given window handle.
+ */
+static void PlatformInitPixelFormatForDevice(HDC DeviceContext, bool bTryIsDummyContext)
+{
+	// Pixel format descriptor for the context.
+	PIXELFORMATDESCRIPTOR PixelFormatDesc;
+	FMemory::Memzero(PixelFormatDesc);
+	PixelFormatDesc.nSize = sizeof(PIXELFORMATDESCRIPTOR);
+	PixelFormatDesc.nVersion = 1;
+	PixelFormatDesc.dwFlags = PFD_DRAW_TO_WINDOW | PFD_SUPPORT_OPENGL | PFD_DOUBLEBUFFER;
+	PixelFormatDesc.iPixelType = PFD_TYPE_RGBA;
+	PixelFormatDesc.cColorBits = 32;
+	PixelFormatDesc.cDepthBits = 0;
+	PixelFormatDesc.cStencilBits = 0;
+	PixelFormatDesc.iLayerType = PFD_MAIN_PLANE;
+
+	static bool bRequestedQuadBufferStereo = FParse::Param(FCommandLine::Get(), TEXT("quad_buffer_stereo"));
+	if (bRequestedQuadBufferStereo)
+	{
+		PixelFormatDesc.dwFlags = PixelFormatDesc.dwFlags | PFD_STEREO;
+	}
+
+	// Set the pixel format and create the context.
+	int32 PixelFormat = ChoosePixelFormat(DeviceContext, &PixelFormatDesc);
+	if (!PixelFormat || !SetPixelFormat(DeviceContext, PixelFormat, &PixelFormatDesc))
+	{
+		UE_LOG(LogRHI, Fatal, TEXT("Failed to set pixel format for device context."));
+	}
+}
+
+/**
+ * Create a dummy window used to construct OpenGL contexts.
+ */
+static void PlatformCreateDummyGLWindow(FPlatformOpenGLContext* OutContext)
+{
+	const TCHAR* WindowClassName = TEXT("DummyGLWindow");
+
+	// Register a dummy window class.
+	static bool bInitializedWindowClass = false;
+	if (!bInitializedWindowClass)
+	{
+		WNDCLASS wc;
+
+		bInitializedWindowClass = true;
+		FMemory::Memzero(wc);
+		wc.style = CS_OWNDC;
+		wc.lpfnWndProc = PlatformDummyGLWndproc;
+		wc.cbClsExtra = 0;
+		wc.cbWndExtra = 0;
+		wc.hInstance = NULL;
+		wc.hIcon = NULL;
+		wc.hCursor = NULL;
+		wc.hbrBackground = (HBRUSH)(COLOR_MENUTEXT);
+		wc.lpszMenuName = NULL;
+		wc.lpszClassName = WindowClassName;
+		ATOM ClassAtom = ::RegisterClass(&wc);
+		check(ClassAtom);
+	}
+
+	int32 WinX = 0;
+	int32 WinY = 0;
+	FParse::Value(FCommandLine::Get(), TEXT("WinX="), WinX);
+	FParse::Value(FCommandLine::Get(), TEXT("WinY="), WinY);
+
+	// Create a dummy window.
+	OutContext->WindowHandle = CreateWindowEx(
+		WS_EX_WINDOWEDGE,
+		WindowClassName,
+		NULL,
+		WS_POPUP,
+		WinX, WinY, 1, 1,
+		NULL, NULL, NULL, NULL);
+	check(OutContext->WindowHandle);
+	OutContext->bReleaseWindowOnDestroy = true;
+
+	// Get the device context.
+	OutContext->DeviceContext = GetDC(OutContext->WindowHandle);
+	check(OutContext->DeviceContext);
+	PlatformInitPixelFormatForDevice(OutContext->DeviceContext, true);
+}
+
+/**
+ * Set OpenGL Context version to fixed version
+ */
+
+static void GetOpenGLVersionForCoreProfile(int& OutMajorVersion, int& OutMinorVersion)
+{
+	// Always initialize GL context with version 4.3, it's the only GL desktop version we support now
+	OutMajorVersion = 4;
+	OutMinorVersion = 3;
+}
+
+/**
+ * Enable/Disable debug context from the commandline
+ */
+static bool PlatformOpenGLDebugCtx()
+{
+	return IsOGLDebugOutputEnabled();
+}
+
+
+/**
+ * Create a core profile OpenGL context.
+ */
+static void PlatformCreateOpenGLContextCore(FPlatformOpenGLContext* OutContext, int MajorVersion, int MinorVersion, HGLRC InParentContext)
+{
+	check(wglCreateContextAttribsARB);
+	check(OutContext);
+	check(OutContext->DeviceContext);
+
+	OutContext->SyncInterval = -1;	// invalid value to enforce setup on first buffer swap
+	OutContext->ViewportFramebuffer = 0;
+
+	int DebugFlag = 0;
+
+	if (PlatformOpenGLDebugCtx())
+	{
+		DebugFlag = WGL_CONTEXT_DEBUG_BIT_ARB;
+	}
+
+#if !EMULATE_ES31
+	int AttribList[] =
+	{
+		WGL_CONTEXT_MAJOR_VERSION_ARB, MajorVersion,
+		WGL_CONTEXT_MINOR_VERSION_ARB, MinorVersion,
+		WGL_CONTEXT_FLAGS_ARB, WGL_CONTEXT_FORWARD_COMPATIBLE_BIT_ARB | DebugFlag,
+		WGL_CONTEXT_PROFILE_MASK_ARB, WGL_CONTEXT_CORE_PROFILE_BIT_ARB,
+		0
+	};
+#else
+	int AttribList[] =
+	{
+		WGL_CONTEXT_MAJOR_VERSION_ARB, 3,
+		WGL_CONTEXT_MINOR_VERSION_ARB, 1,
+		WGL_CONTEXT_FLAGS_ARB, DebugFlag,
+		WGL_CONTEXT_PROFILE_MASK_ARB, WGL_CONTEXT_ES2_PROFILE_BIT_EXT,
+		0
+	};
+#endif
+
+	OutContext->OpenGLContext = wglCreateContextAttribsARB(OutContext->DeviceContext, InParentContext, AttribList);
+}
+
+void PlatformReleaseOpenGLContext(FPlatformOpenGLDevice* Device, FPlatformOpenGLContext* Context);
+
+extern void OnQueryInvalidation(void);
+
+/** Platform specific OpenGL device. */
+struct FPlatformOpenGLDevice
+{
+	FPlatformOpenGLContext	SharedContext;
+	FPlatformOpenGLContext	RenderingContext;
+	TArray<FPlatformOpenGLContext*>	ViewportContexts;
+	bool					TargetDirty;
+
+	/** Guards against operating on viewport contexts from more than one thread at the same time. */
+	FCriticalSection* ContextUsageGuard;
+
+	FPlatformOpenGLDevice()
+		: TargetDirty(true)
+	{
+		extern void InitDebugContext();
+		ContextUsageGuard = new FCriticalSection;
+
+		int MajorVersion = 0;
+		int MinorVersion = 0;
+		GetOpenGLVersionForCoreProfile(MajorVersion, MinorVersion);
+
+		// Need to call this before we set the debug callback, otherwise if we're not running under RD, the debug extension will assert (invalid enum)
+		GRunningUnderRenderDoc = glIsEnabled(GL_DEBUG_TOOL_EXT) != GL_FALSE;
+
+		PlatformCreateDummyGLWindow(&SharedContext);
+		PlatformCreateOpenGLContextCore(&SharedContext, MajorVersion, MinorVersion, NULL);
+		check(SharedContext.OpenGLContext)
+		{
+			FScopeContext ScopeContext(&SharedContext);
+			InitDebugContext();
+			glGenVertexArrays(1, &SharedContext.VertexArrayObject);
+			glBindVertexArray(SharedContext.VertexArrayObject);
+			InitDefaultGLContextState();
+			glGenFramebuffers(1, &SharedContext.ViewportFramebuffer);
+		}
+
+		PlatformCreateDummyGLWindow(&RenderingContext);
+		PlatformCreateOpenGLContextCore(&RenderingContext, MajorVersion, MinorVersion, SharedContext.OpenGLContext);
+		check(RenderingContext.OpenGLContext);
+		{
+			FScopeContext ScopeContext(&RenderingContext);
+			InitDebugContext();
+			glGenVertexArrays(1, &RenderingContext.VertexArrayObject);
+			glBindVertexArray(RenderingContext.VertexArrayObject);
+			InitDefaultGLContextState();
+			glGenFramebuffers(1, &RenderingContext.ViewportFramebuffer);
+		}
+
+		ContextMakeCurrent(SharedContext.DeviceContext, SharedContext.OpenGLContext);
+	}
+
+	~FPlatformOpenGLDevice()
+	{
+		check(ViewportContexts.Num() == 0);
+
+		ContextMakeCurrent(NULL, NULL);
+
+		OnQueryInvalidation();
+		PlatformReleaseOpenGLContext(this, &RenderingContext);
+		PlatformReleaseOpenGLContext(this, &SharedContext);
+
+		delete ContextUsageGuard;
+	}
+};
+
+FPlatformOpenGLDevice* PlatformCreateOpenGLDevice()
+{
+	return new FPlatformOpenGLDevice;
+}
+
+bool PlatformCanEnableGPUCapture()
+{
+	return GRunningUnderRenderDoc;
+}
+
+void PlatformDestroyOpenGLDevice(FPlatformOpenGLDevice* Device)
+{
+	delete Device;
+}
+
+/**
+ * Create an OpenGL context.
+ */
+FPlatformOpenGLContext* PlatformCreateOpenGLContext(FPlatformOpenGLDevice* Device, void* InWindowHandle)
+{
+	check(InWindowHandle);
+
+	Device->TargetDirty = true;
+
+	FPlatformOpenGLContext* Context = new FPlatformOpenGLContext;
+	Context->WindowHandle = (HWND)InWindowHandle;
+	Context->bReleaseWindowOnDestroy = false;
+	Context->DeviceContext = GetDC(Context->WindowHandle);
+	check(Context->DeviceContext);
+	PlatformInitPixelFormatForDevice(Context->DeviceContext, false);
+
+	int MajorVersion = 0;
+	int MinorVersion = 0;
+	GetOpenGLVersionForCoreProfile(MajorVersion, MinorVersion);
+
+	PlatformCreateOpenGLContextCore(Context, MajorVersion, MinorVersion, Device->SharedContext.OpenGLContext);
+	check(Context->OpenGLContext);
+	{
+		FScopeContext Scope(Context);
+		InitDefaultGLContextState();
+		glGenFramebuffers(1, &Context->ViewportFramebuffer);
+	}
+
+	Device->ViewportContexts.Add(Context);
+	return Context;
+}
+
+/**
+ * Release an OpenGL context.
+ */
+void PlatformReleaseOpenGLContext(FPlatformOpenGLDevice* Device, FPlatformOpenGLContext* Context)
+{
+	check(Context && Context->OpenGLContext);
+
+	Device->ViewportContexts.RemoveSingle(Context);
+	Device->TargetDirty = true;
+
+	bool bActiveContextWillBeReleased = false;
+	bool bSharedDC = false;
+
+	{
+		FScopeLock ScopeLock(Device->ContextUsageGuard);
+		{
+			FScopeContext ScopeContext(Context);
+
+			bActiveContextWillBeReleased = ScopeContext.ContextWasAlreadyActive();
+			bSharedDC = ScopeContext.ContextsShareSameDC();
+
+			DeleteQueriesForCurrentContext(Context->OpenGLContext);
+			glBindVertexArray(0);
+			glDeleteVertexArrays(1, &Context->VertexArrayObject);
+
+			if (Context->ViewportFramebuffer)
+			{
+				glDeleteFramebuffers(1, &Context->ViewportFramebuffer);	// this can be done from any context shared with ours, as long as it's not nil.
+				Context->ViewportFramebuffer = 0;
+			}
+		}
+
+		wglDeleteContext(Context->OpenGLContext);
+		Context->OpenGLContext = NULL;
+	}
+
+	check(Context->DeviceContext);
+
+	if (bActiveContextWillBeReleased)
+	{
+		wglMakeCurrent(NULL, NULL);
+	}
+	ReleaseDC(Context->WindowHandle, Context->DeviceContext);
+	Context->DeviceContext = NULL;
+
+	if (PlatformOpenGLCurrentContext(Device) == CONTEXT_Rendering && bSharedDC)
+	{
+		// The rendering context has been made current using the DC of the now destroyed context. Since this DC has been released the current context will be invalid.
+		// To properly set the rendering context we must make current here with it's own DC.
+		ContextMakeCurrent(Device->RenderingContext.DeviceContext, Device->RenderingContext.OpenGLContext);
+	}
+
+	check(Context->WindowHandle);
+
+	if (Context->bReleaseWindowOnDestroy)
+	{
+		DestroyWindow(Context->WindowHandle);
+	}
+	Context->WindowHandle = NULL;
+}
+
+/**
+ * Destroy an OpenGL context.
+ */
+void PlatformDestroyOpenGLContext(FPlatformOpenGLDevice* Device, FPlatformOpenGLContext* Context)
+{
+	PlatformReleaseOpenGLContext(Device, Context);
+	delete Context;
+}
+
+/**
+ * Main function for transferring data to on-screen buffers.
+ * On Windows it temporarily switches OpenGL context, on Mac only context's output view.
+ */
+bool PlatformBlitToViewport(FPlatformOpenGLDevice* Device, const FOpenGLViewport& Viewport, uint32 BackbufferSizeX, uint32 BackbufferSizeY, bool bPresent, bool bLockToVsync)
+{
+	int32 SyncInterval = RHIGetSyncInterval();
+
+	FPlatformOpenGLContext* const Context = Viewport.GetGLContext();
+
+	check(Context && Context->DeviceContext);
+
+	if (FOpenGL::IsAndroidGLESCompatibilityModeEnabled())
+	{
+		glDisable(GL_FRAMEBUFFER_SRGB);
+
+		int32 RealSyncInterval = bLockToVsync ? SyncInterval : 0;
+		if (wglSwapIntervalEXT_ProcAddress && Context->SyncInterval != RealSyncInterval)
+		{
+			wglSwapIntervalEXT_ProcAddress(RealSyncInterval);
+			Context->SyncInterval = RealSyncInterval;
+		}
+
+		::SwapBuffers(Context->DeviceContext);
+		return true;
+	}
+
+	FScopeLock ScopeLock(Device->ContextUsageGuard);
+	{
+		FPlatformOpenGLContext TempContext = *Context;
+		if (Device->ViewportContexts.Num() == 1)
+		{
+			TempContext.OpenGLContext = Device->RenderingContext.OpenGLContext;
+			TempContext.ViewportFramebuffer = Device->RenderingContext.ViewportFramebuffer;
+		}
+		FScopeContext ScopeContext(&TempContext);
+
+		GLuint vfb = TempContext.ViewportFramebuffer;
+		if (Viewport.GetCustomPresent())
+		{
+			Device->TargetDirty = false;
+			glDisable(GL_FRAMEBUFFER_SRGB);
+			bool bShouldPresent = Viewport.GetCustomPresent()->Present(SyncInterval);
+			glEnable(GL_FRAMEBUFFER_SRGB);
+			if (!bShouldPresent)
+			{
+				return false;
+			}
+			else
+			{
+				Device->TargetDirty = true;
+			}
+		}
+
+		if (Device->ViewportContexts.Num() == 1 && Device->TargetDirty)
+		{
+			glBindFramebuffer(GL_FRAMEBUFFER, TempContext.ViewportFramebuffer);
+			glFramebufferTexture2D(GL_FRAMEBUFFER, GL_COLOR_ATTACHMENT0, Context->BackBufferTarget, Context->BackBufferResource, 0);
+
+			FOpenGL::CheckFrameBuffer();
+			Device->TargetDirty = false;
+		}
+
+		glDisable(GL_FRAMEBUFFER_SRGB);
+		glBindFramebuffer(GL_DRAW_FRAMEBUFFER, 0);
+		FOpenGL::DrawBuffer(GL_BACK);
+		glBindFramebuffer(GL_READ_FRAMEBUFFER, TempContext.ViewportFramebuffer);
+		glReadBuffer(GL_COLOR_ATTACHMENT0);
+
+		glBlitFramebuffer(
+			0, 0, BackbufferSizeX, BackbufferSizeY,
+			0, BackbufferSizeY, BackbufferSizeX, 0,
+			GL_COLOR_BUFFER_BIT,
+			GL_NEAREST
+		);
+		glEnable(GL_FRAMEBUFFER_SRGB);
+
+		if (bPresent)
+		{
+			int32 RealSyncInterval = bLockToVsync ? SyncInterval : 0;
+			if (wglSwapIntervalEXT_ProcAddress && Context->SyncInterval != RealSyncInterval)
+			{
+				wglSwapIntervalEXT_ProcAddress(RealSyncInterval);
+				Context->SyncInterval = RealSyncInterval;
+			}
+
+			::SwapBuffers(Context->DeviceContext);
+		}
+	}
+	return true;
+}
+
+void PlatformFlushIfNeeded()
+{
+	glFinish();
+}
+
+void PlatformRebindResources(FPlatformOpenGLDevice* Device)
+{
+	// @todo: Figure out if we need to rebind frame & renderbuffers after switching contexts
+}
+
+FPlatformOpenGLContext* PlatformGetOpenGLRenderingContext(FPlatformOpenGLDevice* Device)
+{
+	return &Device->RenderingContext;
+}
+
+void PlatformRenderingContextSetup(FPlatformOpenGLDevice* Device)
+{
+	check(Device && Device->RenderingContext.DeviceContext && Device->RenderingContext.OpenGLContext);
+
+	if (GetCurrentContext())
+	{
+		glFlush();
+	}
+	if (Device->ViewportContexts.Num() == 1)
+	{
+		// use the HDC of the window, to reduce context swap overhead
+		ContextMakeCurrent(Device->ViewportContexts[0]->DeviceContext, Device->RenderingContext.OpenGLContext);
+	}
+	else
+	{
+		ContextMakeCurrent(Device->RenderingContext.DeviceContext, Device->RenderingContext.OpenGLContext);
+	}
+}
+
+void PlatformSharedContextSetup(FPlatformOpenGLDevice* Device)
+{
+	check(Device && Device->SharedContext.DeviceContext && Device->SharedContext.OpenGLContext);
+
+	// no need to glFlush() on Windows, it does flush by itself before switching contexts
+	ContextMakeCurrent(Device->SharedContext.DeviceContext, Device->SharedContext.OpenGLContext);
+
+}
+
+void PlatformNULLContextSetup()
+{
+	if (wglGetCurrentDC())
+	{
+		// no need to glFlush() on Windows, it does flush by itself before switching contexts
+		ContextMakeCurrent(NULL, NULL);
+	}
+}
+
+/**
+ * Resize the GL context.
+ */
+void PlatformResizeGLContext(FPlatformOpenGLDevice* Device, FPlatformOpenGLContext* Context, uint32 SizeX, uint32 SizeY, bool bFullscreen, bool bWasFullscreen, GLenum BackBufferTarget, GLuint BackBufferResource)
+{
+	FScopeLock ScopeLock(Device->ContextUsageGuard);
+	{
+		uint32 WindowStyle = WS_CAPTION | WS_SYSMENU;
+		uint32 WindowStyleEx = 0;
+		HWND InsertAfter = HWND_NOTOPMOST;
+
+		if (bFullscreen)
+		{
+			// Get the monitor info from the window handle.
+			HMONITOR hMonitor = MonitorFromWindow(Context->WindowHandle, MONITOR_DEFAULTTOPRIMARY);
+			MONITORINFOEX MonitorInfo;
+			memset(&MonitorInfo, 0, sizeof(MONITORINFOEX));
+			MonitorInfo.cbSize = sizeof(MONITORINFOEX);
+			GetMonitorInfo(hMonitor, &MonitorInfo);
+
+			DEVMODE Mode;
+			Mode.dmSize = sizeof(DEVMODE);
+			Mode.dmBitsPerPel = 32;
+			Mode.dmPelsWidth = SizeX;
+			Mode.dmPelsHeight = SizeY;
+			Mode.dmFields = DM_BITSPERPEL | DM_PELSWIDTH | DM_PELSHEIGHT;
+
+			// Turn on fullscreen mode for the current monitor
+			ChangeDisplaySettingsEx(MonitorInfo.szDevice, &Mode, NULL, CDS_FULLSCREEN, NULL);
+
+			WindowStyle = WS_POPUP;
+			WindowStyleEx = WS_EX_APPWINDOW | WS_EX_TOPMOST;
+			InsertAfter = HWND_TOPMOST;
+		}
+		else if (bWasFullscreen)
+		{
+			ChangeDisplaySettings(NULL, 0);
+		}
+
+		Device->TargetDirty = true;
+		Context->BackBufferResource = BackBufferResource;
+		Context->BackBufferTarget = BackBufferTarget;
+
+		//SetWindowLong(Context->WindowHandle, GWL_STYLE, WindowStyle);
+		//SetWindowLong(Context->WindowHandle, GWL_EXSTYLE, WindowStyleEx);
+
+		if (!FOpenGL::IsAndroidGLESCompatibilityModeEnabled())
+		{
+			FScopeContext ScopeContext(Context);
+
+			glBindFramebuffer(GL_FRAMEBUFFER, Context->ViewportFramebuffer);
+			glFramebufferTexture2D(GL_FRAMEBUFFER, GL_COLOR_ATTACHMENT0, BackBufferTarget, BackBufferResource, 0);
+#if UE_BUILD_DEBUG
+			glReadBuffer(GL_COLOR_ATTACHMENT0);
+			FOpenGL::DrawBuffer(GL_COLOR_ATTACHMENT0);
+#endif
+			FOpenGL::CheckFrameBuffer();
+
+			glViewport(0, 0, SizeX, SizeY);
+			static GLfloat ZeroColor[4] = { 0.0f, 0.0f, 0.0f, 0.0f };
+			glClearBufferfv(GL_COLOR, 0, ZeroColor);
+		}
+	}
+}
+
+void PlatformGetSupportedResolution(uint32& Width, uint32& Height)
+{
+	uint32 InitializedMode = false;
+	uint32 BestWidth = 0;
+	uint32 BestHeight = 0;
+	uint32 ModeIndex = 0;
+	DEVMODE DisplayMode;
+	FMemory::Memzero(&DisplayMode, sizeof(DEVMODE));
+
+	while (EnumDisplaySettings(NULL, ModeIndex++, &DisplayMode))
+	{
+		bool IsEqualOrBetterWidth = FMath::Abs((int32)DisplayMode.dmPelsWidth - (int32)Width) <= FMath::Abs((int32)BestWidth - (int32)Width);
+		bool IsEqualOrBetterHeight = FMath::Abs((int32)DisplayMode.dmPelsHeight - (int32)Height) <= FMath::Abs((int32)BestHeight - (int32)Height);
+		if (!InitializedMode || (IsEqualOrBetterWidth && IsEqualOrBetterHeight))
+		{
+			BestWidth = DisplayMode.dmPelsWidth;
+			BestHeight = DisplayMode.dmPelsHeight;
+			InitializedMode = true;
+		}
+	}
+	check(InitializedMode);
+	Width = BestWidth;
+	Height = BestHeight;
+}
+
+bool PlatformGetAvailableResolutions(FScreenResolutionArray& Resolutions, bool bIgnoreRefreshRate)
+{
+	int32 MinAllowableResolutionX = 0;
+	int32 MinAllowableResolutionY = 0;
+	int32 MaxAllowableResolutionX = 10480;
+	int32 MaxAllowableResolutionY = 10480;
+	int32 MinAllowableRefreshRate = 0;
+	int32 MaxAllowableRefreshRate = 10480;
+
+	if (MaxAllowableResolutionX == 0) //-V547
+	{
+		MaxAllowableResolutionX = 10480;
+	}
+	if (MaxAllowableResolutionY == 0) //-V547
+	{
+		MaxAllowableResolutionY = 10480;
+	}
+	if (MaxAllowableRefreshRate == 0) //-V547
+	{
+		MaxAllowableRefreshRate = 10480;
+	}
+
+	uint32 ModeIndex = 0;
+	DEVMODE DisplayMode;
+	FMemory::Memzero(&DisplayMode, sizeof(DEVMODE));
+
+	while (EnumDisplaySettings(NULL, ModeIndex++, &DisplayMode))
+	{
+		if (((int32)DisplayMode.dmPelsWidth >= MinAllowableResolutionX) &&
+			((int32)DisplayMode.dmPelsWidth <= MaxAllowableResolutionX) &&
+			((int32)DisplayMode.dmPelsHeight >= MinAllowableResolutionY) &&
+			((int32)DisplayMode.dmPelsHeight <= MaxAllowableResolutionY)
+			)
+		{
+			bool bAddIt = true;
+			if (bIgnoreRefreshRate == false)
+			{
+				if (((int32)DisplayMode.dmDisplayFrequency < MinAllowableRefreshRate) ||
+					((int32)DisplayMode.dmDisplayFrequency > MaxAllowableRefreshRate)
+					)
+				{
+					continue;
+				}
+			}
+			else
+			{
+				// See if it is in the list already
+				for (int32 CheckIndex = 0; CheckIndex < Resolutions.Num(); CheckIndex++)
+				{
+					FScreenResolutionRHI& CheckResolution = Resolutions[CheckIndex];
+					if ((CheckResolution.Width == DisplayMode.dmPelsWidth) &&
+						(CheckResolution.Height == DisplayMode.dmPelsHeight))
+					{
+						// Already in the list...
+						bAddIt = false;
+						break;
+					}
+				}
+			}
+
+			if (bAddIt)
+			{
+				// Add the mode to the list
+				int32 Temp2Index = Resolutions.AddZeroed();
+				FScreenResolutionRHI& ScreenResolution = Resolutions[Temp2Index];
+
+				ScreenResolution.Width = DisplayMode.dmPelsWidth;
+				ScreenResolution.Height = DisplayMode.dmPelsHeight;
+				ScreenResolution.RefreshRate = DisplayMode.dmDisplayFrequency;
+			}
+		}
+	}
+
+	return true;
+}
+
+void PlatformRestoreDesktopDisplayMode()
+{
+	ChangeDisplaySettings(NULL, 0);
+}
+
+bool PlatformInitOpenGL()
+{
+	static bool bInitialized = false;
+	static bool bOpenGLSupported = false;
+	if (!bInitialized)
+	{
+		// Disable warning C4191: 'type cast' : unsafe conversion from 'PROC' to 'XXX' while getting GL entry points.
+#pragma warning(push)
+#pragma warning(disable:4191)
+
+// Create a dummy context so that wglCreateContextAttribsARB can be initialized.
+		FPlatformOpenGLContext DummyContext;
+		PlatformCreateDummyGLWindow(&DummyContext);
+		DummyContext.OpenGLContext = wglCreateContext(DummyContext.DeviceContext);
+		check(DummyContext.OpenGLContext);
+		ContextMakeCurrent(DummyContext.DeviceContext, DummyContext.OpenGLContext);
+		wglCreateContextAttribsARB = (PFNWGLCREATECONTEXTATTRIBSARBPROC)wglGetProcAddress("wglCreateContextAttribsARB");
+		if (wglCreateContextAttribsARB)
+		{
+			int MajorVersion = 0;
+			int MinorVersion = 0;
+
+			ContextMakeCurrent(NULL, NULL);
+			wglDeleteContext(DummyContext.OpenGLContext);
+			GetOpenGLVersionForCoreProfile(MajorVersion, MinorVersion);
+			PlatformCreateOpenGLContextCore(&DummyContext, MajorVersion, MinorVersion, NULL);
+			if (DummyContext.OpenGLContext)
+			{
+				bOpenGLSupported = true;
+				ContextMakeCurrent(DummyContext.DeviceContext, DummyContext.OpenGLContext);
+			}
+			else
+			{
+				UE_LOG(LogRHI, Error, TEXT("OpenGL %d.%d not supported by driver"), MajorVersion, MinorVersion);
+			}
+		}
+
+		if (bOpenGLSupported)
+		{
+			// Retrieve the OpenGL DLL.
+			void* OpenGLDLL = FPlatformProcess::GetDllHandle(TEXT("opengl32.dll"));
+			if (!OpenGLDLL)
+			{
+				UE_LOG(LogRHI, Fatal, TEXT("Couldn't load opengl32.dll"));
+			}
+
+			// Initialize entry points required by Unreal from opengl32.dll
+#define GET_GL_ENTRYPOINTS_DLL(Type,Func) Func = (Type)FPlatformProcess::GetDllExport(OpenGLDLL,TEXT(#Func));
+			ENUM_GL_ENTRYPOINTS_DLL(GET_GL_ENTRYPOINTS_DLL);
+#undef GET_GL_ENTRYPOINTS_DLL
+
+			// Release the OpenGL DLL.
+			FPlatformProcess::FreeDllHandle(OpenGLDLL);
+
+			// Initialize all entry points required by Unreal.
+#define GET_GL_ENTRYPOINTS(Type,Func) Func = (Type)wglGetProcAddress(#Func);
+			ENUM_GL_ENTRYPOINTS(GET_GL_ENTRYPOINTS);
+			ENUM_GL_ENTRYPOINTS_OPTIONAL(GET_GL_ENTRYPOINTS);
+#undef GET_GL_ENTRYPOINTS
+
+			// Restore warning C4191.
+#pragma warning(pop)
+
+// Check that all of the entry points have been initialized.
+			bool bFoundAllEntryPoints = true;
+#define CHECK_GL_ENTRYPOINTS(Type,Func) if (Func == NULL) { bFoundAllEntryPoints = false; UE_LOG(LogRHI, Warning, TEXT("Failed to find entry point for %s"), TEXT(#Func)); }
+			ENUM_GL_ENTRYPOINTS_DLL(CHECK_GL_ENTRYPOINTS);
+			ENUM_GL_ENTRYPOINTS(CHECK_GL_ENTRYPOINTS);
+#undef CHECK_GL_ENTRYPOINTS
+			checkf(bFoundAllEntryPoints, TEXT("Failed to find all OpenGL entry points."));
+		}
+
+		// The dummy context can now be released.
+		if (DummyContext.OpenGLContext)
+		{
+			ContextMakeCurrent(NULL, NULL);
+			wglDeleteContext(DummyContext.OpenGLContext);
+		}
+		ReleaseDC(DummyContext.WindowHandle, DummyContext.DeviceContext);
+		check(DummyContext.bReleaseWindowOnDestroy);
+		DestroyWindow(DummyContext.WindowHandle);
+
+	}
+	return bOpenGLSupported;
+}
+
+bool PlatformOpenGLContextValid()
+{
+	return(GetCurrentContext() != NULL);
+}
+
+int32 PlatformGlGetError()
+{
+	return glGetError();
+}
+
+EOpenGLCurrentContext PlatformOpenGLCurrentContext(FPlatformOpenGLDevice* Device)
+{
+	HGLRC Context = GetCurrentContext();
+
+	if (Context == Device->RenderingContext.OpenGLContext)	// most common case
+	{
+		return CONTEXT_Rendering;
+	}
+	else if (Context == Device->SharedContext.OpenGLContext)
+	{
+		return CONTEXT_Shared;
+	}
+	else if (Context)
+	{
+		return CONTEXT_Other;
+	}
+	else
+	{
+		return CONTEXT_Invalid;
+	}
+}
+
+void* PlatformOpenGLCurrentContextHandle(FPlatformOpenGLDevice* Device)
+{
+	return GetCurrentContext();
+}
+
+void PlatformGetBackbufferDimensions(uint32& OutWidth, uint32& OutHeight)
+{
+	OutWidth = OutHeight = 0;
+	HDC DeviceContext = wglGetCurrentDC();
+	if (DeviceContext)
+	{
+		OutWidth = GetDeviceCaps(DeviceContext, HORZRES);
+		OutHeight = GetDeviceCaps(DeviceContext, VERTRES);
+	}
+}
+
+// =============================================================
+
+struct FOpenGLReleasedQuery
+{
+	HGLRC	Context;
+	GLuint	Query;
+};
+
+static TArray<FOpenGLReleasedQuery>	ReleasedQueries;
+static FCriticalSection* ReleasedQueriesGuard;
+
+void PlatformGetNewRenderQuery(GLuint* OutQuery, uint64* OutQueryContext)
+{
+	if (!ReleasedQueriesGuard)
+	{
+		ReleasedQueriesGuard = new FCriticalSection;
+	}
+
+	{
+		FScopeLock Lock(ReleasedQueriesGuard);
+
+#ifdef UE_BUILD_DEBUG
+		check(OutQuery && OutQueryContext);
+#endif
+
+		HGLRC Context = GetCurrentContext();
+		check(Context);
+
+		GLuint NewQuery = 0;
+
+		// Check for possible query reuse
+		const int32 ArraySize = ReleasedQueries.Num();
+		for (int32 Index = 0; Index < ArraySize; ++Index)
+		{
+			if (ReleasedQueries[Index].Context == Context)
+			{
+				NewQuery = ReleasedQueries[Index].Query;
+				ReleasedQueries.RemoveAtSwap(Index);
+				break;
+			}
+		}
+
+		if (!NewQuery)
+		{
+			FOpenGL::GenQueries(1, &NewQuery);
+		}
+
+		*OutQuery = NewQuery;
+		*OutQueryContext = (uint64)Context;
+	}
+}
+
+void PlatformReleaseRenderQuery(GLuint Query, uint64 QueryContext)
+{
+	HGLRC Context = GetCurrentContext();
+	if ((uint64)Context == QueryContext)
+	{
+		FOpenGL::DeleteQueries(1, &Query);
+	}
+	else
+	{
+		FScopeLock Lock(ReleasedQueriesGuard);
+#ifdef UE_BUILD_DEBUG
+		check(Query && QueryContext && ReleasedQueriesGuard);
+#endif
+		FOpenGLReleasedQuery ReleasedQuery;
+		ReleasedQuery.Context = (HGLRC)QueryContext;
+		ReleasedQuery.Query = Query;
+		ReleasedQueries.Add(ReleasedQuery);
+	}
+}
+
+void DeleteQueriesForCurrentContext(HGLRC Context)
+{
+	if (!ReleasedQueriesGuard)
+	{
+		ReleasedQueriesGuard = new FCriticalSection;
+	}
+
+	{
+		FScopeLock Lock(ReleasedQueriesGuard);
+		for (int32 Index = 0; Index < ReleasedQueries.Num(); ++Index)
+		{
+			if (ReleasedQueries[Index].Context == Context)
+			{
+				FOpenGL::DeleteQueries(1, &ReleasedQueries[Index].Query);
+				ReleasedQueries.RemoveAtSwap(Index);
+				--Index;
+			}
+		}
+	}
+}
+
+bool PlatformContextIsCurrent(uint64 QueryContext)
+{
+	return (uint64)GetCurrentContext() == QueryContext;
+}
+
+FOpenGLTexture* PlatformCreateBuiltinBackBuffer(FOpenGLDynamicRHI* OpenGLRHI, uint32 SizeX, uint32 SizeY)
+{
+	if (FOpenGL::IsAndroidGLESCompatibilityModeEnabled())
+	{
+		const FRHITextureCreateDesc Desc =
+			FRHITextureCreateDesc::Create2D(TEXT("PlatformCreateBuiltinBackBuffer"), SizeX, SizeY, PF_B8G8R8A8)
+			.SetClearValue(FClearValueBinding::Transparent)
+			.SetFlags(ETextureCreateFlags::RenderTargetable | ETextureCreateFlags::Presentable | ETextureCreateFlags::ResolveTargetable)
+			.DetermineInititialState();
+
+		return new FOpenGLTexture(FRHICommandListImmediate::Get(), Desc);
+	}
+
+	return nullptr;
+}
+
+void* PlatformGetWindow(FPlatformOpenGLContext* Context, void** AddParam)
+{
+	check(Context && Context->WindowHandle);
+
+	return (void*)&Context->WindowHandle;
+}
diff --git a/Engine/Source/Runtime/OpenGLDrv/Private/Windows/WindowsOpenGLPlatform.h b/Engine/Source/Runtime/OpenGLDrv/Private/Windows/WindowsOpenGLPlatform.h
new file mode 100644
--- /dev/null
+++ b/Engine/Source/Runtime/OpenGLDrv/Private/Windows/WindowsOpenGLPlatform.h
@@ -0,0 +1,18 @@
+// Copyright Epic Games, Inc. All Rights Reserved.
+
+#pragma once
+
+#include "OpenGLThirdParty.h"
+
+#define UGL_PLATFORM_SUPPORTS_GLES 0
+#define UGL_PLATFORM_SUPPORTS_GL3  1
+#define UGL_PLATFORM_SUPPORTS_GL4  1
+
+/** Unreal tokens that maps to different OpenGL tokens by platform. */
+#define UGL_DRAW_FRAMEBUFFER	GL_DRAW_FRAMEBUFFER
+#define UGL_READ_FRAMEBUFFER	GL_READ_FRAMEBUFFER
+#define UGL_ABGR8				GL_UNSIGNED_INT_8_8_8_8_REV
+#define UGL_ANY_SAMPLES_PASSED	GL_ANY_SAMPLES_PASSED
+#define UGL_SAMPLES_PASSED		GL_SAMPLES_PASSED
+#define UGL_TIME_ELAPSED		GL_TIME_ELAPSED
+#define UGL_CLAMP_TO_BORDER		GL_CLAMP_TO_BORDER
diff --git a/Engine/Source/Runtime/OpenGLDrv/Public/Android/AndroidOpenGLThirdParty.h b/Engine/Source/Runtime/OpenGLDrv/Public/Android/AndroidOpenGLThirdParty.h
new file mode 100644
--- /dev/null
+++ b/Engine/Source/Runtime/OpenGLDrv/Public/Android/AndroidOpenGLThirdParty.h
@@ -0,0 +1,13 @@
+// Copyright Epic Games, Inc. All Rights Reserved.
+
+#pragma once
+
+#include "HAL/Platform.h"
+
+THIRD_PARTY_INCLUDES_START
+	#include <EGL/egl.h>
+	#include <EGL/eglext.h>
+	#include <GLES2/gl2.h>
+	#include <GLES2/gl2ext.h>
+	#include <GLES3/gl31.h>
+THIRD_PARTY_INCLUDES_END
diff --git a/Engine/Source/Runtime/OpenGLDrv/Public/Linux/LinuxOpenGLThirdParty.h b/Engine/Source/Runtime/OpenGLDrv/Public/Linux/LinuxOpenGLThirdParty.h
new file mode 100644
--- /dev/null
+++ b/Engine/Source/Runtime/OpenGLDrv/Public/Linux/LinuxOpenGLThirdParty.h
@@ -0,0 +1,10 @@
+// Copyright Epic Games, Inc. All Rights Reserved.
+
+#pragma once
+
+#include "HAL/Platform.h"
+
+THIRD_PARTY_INCLUDES_START
+	#include <GL/glcorearb.h>
+	#include <GL/glext.h>
+THIRD_PARTY_INCLUDES_END
diff --git a/Engine/Source/Runtime/OpenGLDrv/Public/Mac/MacOpenGLThirdParty.h b/Engine/Source/Runtime/OpenGLDrv/Public/Mac/MacOpenGLThirdParty.h
new file mode 100644
--- /dev/null
+++ b/Engine/Source/Runtime/OpenGLDrv/Public/Mac/MacOpenGLThirdParty.h
@@ -0,0 +1,17 @@
+// Copyright Epic Games, Inc. All Rights Reserved.
+
+#pragma once
+
+#include <OpenGL/OpenGL.h>
+#include <OpenGL/gl3.h>
+#include <OpenGL/gl3ext.h>
+
+#ifndef GL_COMPUTE_SHADER
+	#define GL_COMPUTE_SHADER 0x91B9
+#endif
+#ifndef GL_TESS_EVALUATION_SHADER
+	#define GL_TESS_EVALUATION_SHADER 0x8E87
+#endif
+#ifndef GL_TESS_CONTROL_SHADER
+	#define GL_TESS_CONTROL_SHADER 0x8E88
+#endif
diff --git a/Engine/Source/Runtime/OpenGLDrv/Public/OpenGLShaderResources.h b/Engine/Source/Runtime/OpenGLDrv/Public/OpenGLShaderResources.h
--- a/Engine/Source/Runtime/OpenGLDrv/Public/OpenGLShaderResources.h
+++ b/Engine/Source/Runtime/OpenGLDrv/Public/OpenGLShaderResources.h
@@ -6,8 +6,6 @@
 
 #pragma once
 
-// HEADER_UNIT_SKIP - Not included directly
-
 #include "Containers/Array.h"
 #include "CrossCompilerCommon.h"
 #include "OpenGLThirdParty.h"
diff --git a/Engine/Source/Runtime/OpenGLDrv/Public/OpenGLThirdParty.h b/Engine/Source/Runtime/OpenGLDrv/Public/OpenGLThirdParty.h
--- a/Engine/Source/Runtime/OpenGLDrv/Public/OpenGLThirdParty.h
+++ b/Engine/Source/Runtime/OpenGLDrv/Public/OpenGLThirdParty.h
@@ -2,26 +2,6 @@
 
 #pragma once
 
-#if PLATFORM_WINDOWS
-	#include "Windows/AllowWindowsPlatformTypes.h"
-	THIRD_PARTY_INCLUDES_START
-		#include <GL/glcorearb.h>
-		#include <GL/glext.h>
-		#include <GL/wglext.h>
-	THIRD_PARTY_INCLUDES_END
-	#include "Windows/HideWindowsPlatformTypes.h"
-#elif PLATFORM_LINUX
-	THIRD_PARTY_INCLUDES_START
-		#include <GL/glcorearb.h>
-		#include <GL/glext.h>
-	THIRD_PARTY_INCLUDES_END
-#elif PLATFORM_ANDROID
-	#include "Android/AndroidPlatform.h"
-	THIRD_PARTY_INCLUDES_START
-		#include <EGL/egl.h>
-		#include <EGL/eglext.h>
-		#include <GLES3/gl31.h>
-	THIRD_PARTY_INCLUDES_END
-#endif
+#include "HAL/PreprocessorHelpers.h"
 
-using UGLsync = GLsync;
+#include COMPILED_PLATFORM_HEADER(OpenGLThirdParty.h)
diff --git a/Engine/Source/Runtime/OpenGLDrv/Public/OpenGLUtil.h b/Engine/Source/Runtime/OpenGLDrv/Public/OpenGLUtil.h
--- a/Engine/Source/Runtime/OpenGLDrv/Public/OpenGLUtil.h
+++ b/Engine/Source/Runtime/OpenGLDrv/Public/OpenGLUtil.h
@@ -29,7 +29,7 @@ int32 GetOGLDebugOutputLevel();
 #define ENABLE_VERIFY_GL_THREAD (UE_BUILD_DEBUG)
 
 /** Set to 1 to verify that the the engine side uniform buffer layout matches the driver side of the GLSL shader*/
-#define ENABLE_UNIFORM_BUFFER_LAYOUT_VERIFICATION ( 0 & UE_BUILD_DEBUG & (OPENGL_GL3 | OPENGL_GL4))
+#define ENABLE_UNIFORM_BUFFER_LAYOUT_VERIFICATION ( 0 & UE_BUILD_DEBUG & (!PLATFORM_ANDROID))
 
 /** Set to 1 to additinally dump uniform buffer layout at shader link time, this assumes ENABLE_UNIFORM_BUFFER_LAYOUT_VERIFICATION == 1 */
 #define ENABLE_UNIFORM_BUFFER_LAYOUT_DUMP 0
diff --git a/Engine/Source/Runtime/OpenGLDrv/Public/Windows/WindowsOpenGLThirdParty.h b/Engine/Source/Runtime/OpenGLDrv/Public/Windows/WindowsOpenGLThirdParty.h
new file mode 100644
--- /dev/null
+++ b/Engine/Source/Runtime/OpenGLDrv/Public/Windows/WindowsOpenGLThirdParty.h
@@ -0,0 +1,13 @@
+// Copyright Epic Games, Inc. All Rights Reserved.
+
+#pragma once
+
+#include "HAL/Platform.h"
+
+#include "Windows/AllowWindowsPlatformTypes.h"
+THIRD_PARTY_INCLUDES_START
+	#include <GL/glcorearb.h>
+	#include <GL/glext.h>
+	#include <GL/wglext.h>
+THIRD_PARTY_INCLUDES_END
+#include "Windows/HideWindowsPlatformTypes.h"
