Deprecate:

diff --git a/Engine/Source/Runtime/D3D12RHI/Public/ID3D12DynamicRHI.h b/Engine/Source/Runtime/D3D12RHI/Public/ID3D12DynamicRHI.h
--- a/Engine/Source/Runtime/D3D12RHI/Public/ID3D12DynamicRHI.h
+++ b/Engine/Source/Runtime/D3D12RHI/Public/ID3D12DynamicRHI.h
@@ -31,9 +31,9 @@ struct ID3D12DynamicRHI : public FDynamicRHI
 	virtual ID3D12GraphicsCommandList* RHIGetGraphicsCommandList(uint32 InDeviceIndex) const = 0;
 	virtual DXGI_FORMAT                RHIGetSwapChainFormat(EPixelFormat InFormat) const = 0;
 
-	virtual FTexture2DRHIRef      RHICreateTexture2DFromResource(EPixelFormat Format, ETextureCreateFlags TexCreateFlags, const FClearValueBinding& ClearValueBinding, ID3D12Resource* Resource) = 0;
-	virtual FTexture2DArrayRHIRef RHICreateTexture2DArrayFromResource(EPixelFormat Format, ETextureCreateFlags TexCreateFlags, const FClearValueBinding& ClearValueBinding, ID3D12Resource* Resource) = 0;
-	virtual FTextureCubeRHIRef    RHICreateTextureCubeFromResource(EPixelFormat Format, ETextureCreateFlags TexCreateFlags, const FClearValueBinding& ClearValueBinding, ID3D12Resource* Resource) = 0;
+	virtual FTextureRHIRef        RHICreateTexture2DFromResource(EPixelFormat Format, ETextureCreateFlags TexCreateFlags, const FClearValueBinding& ClearValueBinding, ID3D12Resource* Resource) = 0;
+	virtual FTextureRHIRef        RHICreateTexture2DArrayFromResource(EPixelFormat Format, ETextureCreateFlags TexCreateFlags, const FClearValueBinding& ClearValueBinding, ID3D12Resource* Resource) = 0;
+	virtual FTextureRHIRef        RHICreateTextureCubeFromResource(EPixelFormat Format, ETextureCreateFlags TexCreateFlags, const FClearValueBinding& ClearValueBinding, ID3D12Resource* Resource) = 0;
 
 	virtual ID3D12Resource*       RHIGetResource(FRHIBuffer* InBuffer) const = 0;
 	virtual uint32                RHIGetResourceDeviceIndex(FRHIBuffer* InBuffer) const = 0;
diff --git a/Engine/Source/Runtime/Engine/Classes/Components/SkyLightComponent.h b/Engine/Source/Runtime/Engine/Classes/Components/SkyLightComponent.h
--- a/Engine/Source/Runtime/Engine/Classes/Components/SkyLightComponent.h
+++ b/Engine/Source/Runtime/Engine/Classes/Components/SkyLightComponent.h
@@ -75,7 +75,7 @@ private:
 	int32 Size;
 	int32 NumMips;
 	EPixelFormat Format;
-	FTextureCubeRHIRef TextureCubeRHI;
+	FTextureRHIRef TextureCubeRHI;
 	int32 NumRefs;
 };
 
diff --git a/Engine/Source/Runtime/Engine/Classes/Engine/TextureRenderTarget.h b/Engine/Source/Runtime/Engine/Classes/Engine/TextureRenderTarget.h
--- a/Engine/Source/Runtime/Engine/Classes/Engine/TextureRenderTarget.h
+++ b/Engine/Source/Runtime/Engine/Classes/Engine/TextureRenderTarget.h
@@ -29,7 +29,7 @@ class UTextureRenderTarget : public UTexture
 	/** If true, there will be two copies in memory - one for the texture and one for the render target. If false, they will share memory if possible. This is useful for scene capture textures that are used in the scene. */
 	uint32 bNeedsTwoCopies:1;
 
-	/** If true, it will be possible to create a FUnorderedAccessViewRHIRef using RHICreateUnorderedAccessView and the internal FTexture2DRHIRef. */
+	/** If true, it will be possible to create a FUnorderedAccessViewRHIRef using RHICreateUnorderedAccessView and the internal FTextureRHIRef. */
 	uint32 bCanCreateUAV : 1;
 
 	/**
diff --git a/Engine/Source/Runtime/Engine/Private/Components/ReflectionCaptureComponent.cpp b/Engine/Source/Runtime/Engine/Private/Components/ReflectionCaptureComponent.cpp
--- a/Engine/Source/Runtime/Engine/Private/Components/ReflectionCaptureComponent.cpp
+++ b/Engine/Source/Runtime/Engine/Private/Components/ReflectionCaptureComponent.cpp
@@ -684,7 +684,7 @@ private:
 	int32 Size;
 	int32 NumMips;
 	EPixelFormat Format;
-	FTextureCubeRHIRef TextureCubeRHI;
+	FTextureRHIRef TextureCubeRHI;
 
 	TArray<uint8> SourceData;
 };
diff --git a/Engine/Source/Runtime/Engine/Private/Particles/FXSystemPrivate.h b/Engine/Source/Runtime/Engine/Private/Particles/FXSystemPrivate.h
--- a/Engine/Source/Runtime/Engine/Private/Particles/FXSystemPrivate.h
+++ b/Engine/Source/Runtime/Engine/Private/Particles/FXSystemPrivate.h
@@ -124,11 +124,11 @@ public:
 	/** Buffer containing unsorted particle indices. */
 	FRHIShaderResourceView* UnsortedParticleIndicesSRV;
 	/** Texture containing positions for all particles. */
-	FRHITexture2D* PositionTextureRHI;
+	FRHITexture* PositionTextureRHI;
 	/** Texture containing velocities for all particles. */
-	FRHITexture2D* VelocityTextureRHI;
+	FRHITexture* VelocityTextureRHI;
 	/** Texture containint attributes for all particles. */
-	FRHITexture2D* AttributesTextureRHI;
+	FRHITexture* AttributesTextureRHI;
 	/** LWC tile offset, will be 0,0,0 for localspace emitters. */
 	FVector3f LWCTile;
 	/** Tile page offset factors associated with the GPU particle simulation resources. */
@@ -361,8 +361,8 @@ private:
 		EParticleSimulatePhase::Type Phase,
 		const TUniformBufferRef<FViewUniformShaderParameters>& ViewUniformBuffer,
 		const FGlobalDistanceFieldParameterData* GlobalDistanceFieldParameterData,
-		FRHITexture2D* SceneDepthTexture,
-		FRHITexture2D* GBufferATexture
+		FRHITexture* SceneDepthTexture,
+		FRHITexture* GBufferATexture
 	);
 
 	/*-------------------------------------------------------------------------
diff --git a/Engine/Source/Runtime/Engine/Private/Particles/ParticleCurveTexture.cpp b/Engine/Source/Runtime/Engine/Private/Particles/ParticleCurveTexture.cpp
--- a/Engine/Source/Runtime/Engine/Private/Particles/ParticleCurveTexture.cpp
+++ b/Engine/Source/Runtime/Engine/Private/Particles/ParticleCurveTexture.cpp
@@ -165,7 +165,7 @@ TGlobalResource<FParticleCurveInjectionVertexDeclaration> GParticleCurveInjectio
 */
 static void InjectCurves(
 	FRHICommandListImmediate& RHICmdList,
-	FRHITexture2D* CurveTextureRHI,
+	FRHITexture* CurveTextureRHI,
 	TArray<FCurveSamples>& InPendingCurves)
 {
 	static bool bFirstCall = true;
diff --git a/Engine/Source/Runtime/Engine/Private/Particles/ParticleCurveTexture.h b/Engine/Source/Runtime/Engine/Private/Particles/ParticleCurveTexture.h
--- a/Engine/Source/Runtime/Engine/Private/Particles/ParticleCurveTexture.h
+++ b/Engine/Source/Runtime/Engine/Private/Particles/ParticleCurveTexture.h
@@ -158,7 +158,7 @@ public:
 	/**
 	 * Retrieves the curve texture from which shaders can sample.
 	 */
-	FRHITexture2D* GetCurveTexture() { return CurveTextureRHI; }
+	FRHITexture* GetCurveTexture() { return CurveTextureRHI; }
 
 	/**
 	 * Submits pending curves to the GPU.
diff --git a/Engine/Source/Runtime/Engine/Private/Particles/ParticleGpuSimulation.cpp b/Engine/Source/Runtime/Engine/Private/Particles/ParticleGpuSimulation.cpp
--- a/Engine/Source/Runtime/Engine/Private/Particles/ParticleGpuSimulation.cpp
+++ b/Engine/Source/Runtime/Engine/Private/Particles/ParticleGpuSimulation.cpp
@@ -323,9 +323,9 @@ public:
 	const TCHAR* PaticleVelocityTextureName = TEXT("FParticleStateVelocity");
 
 	/** Contains the positions of all simulating particles. */
-	FTexture2DRHIRef PositionTextureRHI;
+	FTextureRHIRef PositionTextureRHI;
 	/** Contains the velocity of all simulating particles. */
-	FTexture2DRHIRef VelocityTextureRHI;
+	FTextureRHIRef VelocityTextureRHI;
 
 	bool bTexturesCleared;
 	int32 ParticleStateIndex = 0;
@@ -386,8 +386,8 @@ public:
 		RenderTargetSize.Y = SizeY;
 
 		// allocate the new textures
-		FTexture2DRHIRef NewPositionTextureRHI = CreateStateTextureRHI(PaticlePositonTextureName, PF_A32B32G32R32F);
-		FTexture2DRHIRef NewVelocityTextureRHI = CreateStateTextureRHI(PaticleVelocityTextureName, PF_FloatRGBA);
+		FTextureRHIRef NewPositionTextureRHI = CreateStateTextureRHI(PaticlePositonTextureName, PF_A32B32G32R32F);
+		FTextureRHIRef NewVelocityTextureRHI = CreateStateTextureRHI(PaticleVelocityTextureName, PF_FloatRGBA);
 		
 		RHICmdList.Transition(FRHITransitionInfo(NewPositionTextureRHI, ERHIAccess::SRVMask, ERHIAccess::CopyDest));
 		RHICmdList.Transition(FRHITransitionInfo(PositionTextureRHI, ERHIAccess::SRVMask, ERHIAccess::CopySrc));
@@ -476,7 +476,7 @@ public:
 		RenderTargetSize.Y = SizeY;
 
 		// allocate the new textures
-		FTexture2DRHIRef NewAttributesTextureRHI = CreateAttributesTextureRHI();
+		FTextureRHIRef NewAttributesTextureRHI = CreateAttributesTextureRHI();
 
 		RHICmdList.Transition(FRHITransitionInfo(NewAttributesTextureRHI, ERHIAccess::SRVMask, ERHIAccess::CopyDest));
 		RHICmdList.Transition(FRHITransitionInfo(TextureRHI, ERHIAccess::SRVMask, ERHIAccess::CopySrc));
@@ -1438,7 +1438,7 @@ struct FSimulationCommandGPU
 	/** Parameters to sample the local vector field for this simulation. */
 	FVectorFieldUniformBufferRef VectorFieldsUniformBuffer;
 	/** Vector field volume textures for this simulation. */
-	FRHITexture3D* VectorFieldTexturesRHI[MAX_VECTOR_FIELDS];
+	FRHITexture* VectorFieldTexturesRHI[MAX_VECTOR_FIELDS];
 	/** The number of tiles to simulate. */
 	int32 UnalignedTileCount;
 
@@ -1451,7 +1451,7 @@ struct FSimulationCommandGPU
 		, VectorFieldsUniformBuffer(InVectorFieldsUniformBuffer)
 		, UnalignedTileCount(InTileCount)
 	{
-		FRHITexture3D* BlackVolumeTextureRHI = (FRHITexture3D*)(FRHITexture*)GBlackVolumeTexture->TextureRHI;
+		FRHITexture* BlackVolumeTextureRHI = (FRHITexture*)(FRHITexture*)GBlackVolumeTexture->TextureRHI;
 		for (int32 i = 0; i < MAX_VECTOR_FIELDS; ++i)
 		{
 			VectorFieldTexturesRHI[i] = BlackVolumeTextureRHI;
@@ -2037,7 +2037,7 @@ static void VisualizeGPUSimulation(
 	int32 VisualizationMode,
 	FRenderTarget* RenderTarget,
 	const FParticleStateTextures& StateTextures,
-	FRHITexture2D* CurveTextureRHI
+	FRHITexture* CurveTextureRHI
 	)
 {
 	check(IsInRenderingThread());
@@ -2250,7 +2250,7 @@ static FBox ComputeParticleBounds(
 	FRHICommandListImmediate& RHICmdList,
 	ERHIFeatureLevel::Type FeatureLevel,
 	FRHIShaderResourceView* VertexBufferSRV,
-	FRHITexture2D* PositionTextureRHI,
+	FRHITexture* PositionTextureRHI,
 	int32 ParticleCount,
 	FParticleSimulationResources* ParticleSimulationResources)
 {
@@ -4835,7 +4835,7 @@ void FFXSystem::SimulateGPUParticles(
 		if (GPUSimulations.Num() > 0)
 		{
 			FVectorFieldUniformParameters VectorFieldParameters;
-			FRHITexture3D* BlackVolumeTextureRHI = (FRHITexture3D*)(FRHITexture*)GBlackVolumeTexture->TextureRHI;
+			FRHITexture* BlackVolumeTextureRHI = (FRHITexture*)(FRHITexture*)GBlackVolumeTexture->TextureRHI;
 			for (int32 Index = 0; Index < MAX_VECTOR_FIELDS; ++Index)
 			{
 				VectorFieldParameters.WorldToVolume[Index] = FMatrix44f::Identity;
diff --git a/Engine/Source/Runtime/Engine/Private/Particles/ParticleSortingGPU.cpp b/Engine/Source/Runtime/Engine/Private/Particles/ParticleSortingGPU.cpp
--- a/Engine/Source/Runtime/Engine/Private/Particles/ParticleSortingGPU.cpp
+++ b/Engine/Source/Runtime/Engine/Private/Particles/ParticleSortingGPU.cpp
@@ -81,7 +81,7 @@ public:
 		const FParticleKeyGenParameters& KeyGenParameters,
 		FRHIUnorderedAccessView* OutKeysUAV,
 		FRHIUnorderedAccessView* OutIndicesUAV,
-		FRHITexture2D* PositionTextureRHI,
+		FRHITexture* PositionTextureRHI,
 		FRHIShaderResourceView* InIndicesSRV
 		)
 	{
@@ -126,7 +126,7 @@ int32 GenerateParticleSortKeys(
 	FRHICommandListImmediate& RHICmdList,
 	FRHIUnorderedAccessView* KeyBufferUAV,
 	FRHIUnorderedAccessView* SortedVertexBufferUAV,
-	FRHITexture2D* PositionTextureRHI,
+	FRHITexture* PositionTextureRHI,
 	const TArray<FParticleSimulationSortInfo>& SimulationsToSort,
 	ERHIFeatureLevel::Type FeatureLevel,
 	int32 BatchId
diff --git a/Engine/Source/Runtime/Engine/Private/Particles/ParticleSortingGPU.h b/Engine/Source/Runtime/Engine/Private/Particles/ParticleSortingGPU.h
--- a/Engine/Source/Runtime/Engine/Private/Particles/ParticleSortingGPU.h
+++ b/Engine/Source/Runtime/Engine/Private/Particles/ParticleSortingGPU.h
@@ -39,7 +39,7 @@ int32 GenerateParticleSortKeys(
 	FRHICommandListImmediate& RHICmdList,
 	FRHIUnorderedAccessView* KeyBufferUAV,
 	FRHIUnorderedAccessView* SortedVertexBufferUAV,
-	FRHITexture2D* PositionTextureRHI,
+	FRHITexture* PositionTextureRHI,
 	const TArray<FParticleSimulationSortInfo>& SimulationsToSort,
 	ERHIFeatureLevel::Type FeatureLevel,
 	int32 BatchId = INDEX_NONE);
diff --git a/Engine/Source/Runtime/Engine/Private/Slate/DebugCanvas.cpp b/Engine/Source/Runtime/Engine/Private/Slate/DebugCanvas.cpp
--- a/Engine/Source/Runtime/Engine/Private/Slate/DebugCanvas.cpp
+++ b/Engine/Source/Runtime/Engine/Private/Slate/DebugCanvas.cpp
@@ -244,7 +244,7 @@ void FDebugCanvasDrawer::Draw_RenderThread(FRHICommandListImmediate& RHICmdList,
 	QUICK_SCOPE_CYCLE_COUNTER(Stat_DrawDebugCanvas);
 	if( RenderThreadCanvas.IsValid() )
 	{
-		FRHITexture* RT = InWindowBackBuffer != nullptr ? ((FTexture2DRHIRef*)InWindowBackBuffer)->GetReference() : nullptr;
+		FRHITexture* RT = InWindowBackBuffer != nullptr ? ((FTextureRHIRef*)InWindowBackBuffer)->GetReference() : nullptr;
 		FTextureRHIRef HMDSwapchain = nullptr;	// Attention: RT could point to HMDSwapchain later.
 		if (RenderThreadCanvas->IsUsingInternalTexture())
 		{
diff --git a/Engine/Source/Runtime/Engine/Private/Slate/SceneViewport.cpp b/Engine/Source/Runtime/Engine/Private/Slate/SceneViewport.cpp
--- a/Engine/Source/Runtime/Engine/Private/Slate/SceneViewport.cpp
+++ b/Engine/Source/Runtime/Engine/Private/Slate/SceneViewport.cpp
@@ -1674,7 +1674,7 @@ void FSceneViewport::EnqueueEndRenderFrame(const bool bLockToVsync, const bool b
 	}
 }
 
-const FTexture2DRHIRef& FSceneViewport::GetRenderTargetTexture() const
+const FTextureRHIRef& FSceneViewport::GetRenderTargetTexture() const
 {
 	if (IsInRenderingThread())
 	{
@@ -1692,7 +1692,7 @@ FSlateShaderResource* FSceneViewport::GetViewportRenderTargetTexture()
 	return (BufferedSlateHandles.Num() != 0) ? BufferedSlateHandles[CurrentBufferedTargetIndex] : nullptr;
 }
 
-void FSceneViewport::SetRenderTargetTextureRenderThread(FTexture2DRHIRef& RT)
+void FSceneViewport::SetRenderTargetTextureRenderThread(FTextureRHIRef& RT)
 {
 	check(IsInRenderingThread());
 	RenderTargetTextureRenderThreadRHI = RT;
@@ -2073,8 +2073,8 @@ void FSceneViewport::InitRHI(FRHICommandListBase& RHICmdList)
 	if (UseSeparateRenderTarget())
 	{
 		int32 NumBufferedFrames = 1;
-		TArray<FTexture2DRHIRef> BufferedRTRHI;
-		TArray<FTexture2DRHIRef> BufferedSRVRHI;
+		TArray<FTextureRHIRef> BufferedRTRHI;
+		TArray<FTextureRHIRef> BufferedSRVRHI;
 
 		static const auto CVarDefaultBackBufferPixelFormat = IConsoleManager::Get().FindTConsoleVariableDataInt(TEXT("r.DefaultBackBufferPixelFormat"));
 		EPixelFormat SceneTargetFormat = EDefaultBackBufferPixelFormat::Convert2PixelFormat(EDefaultBackBufferPixelFormat::FromInt(CVarDefaultBackBufferPixelFormat->GetValueOnRenderThread()));
@@ -2133,8 +2133,8 @@ void FSceneViewport::InitRHI(FRHICommandListBase& RHICmdList)
 		}
 		check(BufferedSlateHandles.Num() == BufferedRenderTargetsRHI.Num() && BufferedSlateHandles.Num() == BufferedShaderResourceTexturesRHI.Num());
 
-		FTexture2DRHIRef RTRHI;
-		FTexture2DRHIRef SRVRHI;
+		FTextureRHIRef RTRHI;
+		FTextureRHIRef SRVRHI;
 
 		for (int32 i = 0; i < NumBufferedFrames; ++i)
 		{
diff --git a/Engine/Source/Runtime/Engine/Private/Slate/SlateTextures.cpp b/Engine/Source/Runtime/Engine/Private/Slate/SlateTextures.cpp
--- a/Engine/Source/Runtime/Engine/Private/Slate/SlateTextures.cpp
+++ b/Engine/Source/Runtime/Engine/Private/Slate/SlateTextures.cpp
@@ -7,7 +7,7 @@
 #include "RenderingThread.h"
 #include "RHIUtilities.h"
 
-FSlateTexture2DRHIRef::FSlateTexture2DRHIRef( FTexture2DRHIRef InRef, uint32 InWidth, uint32 InHeight )
+FSlateTexture2DRHIRef::FSlateTexture2DRHIRef( FTextureRHIRef InRef, uint32 InWidth, uint32 InHeight )
 	: TSlateTexture( InRef )
 	, Width( InWidth )
 	, Height( InHeight )
@@ -118,7 +118,7 @@ void FSlateTexture2DRHIRef::Resize( uint32 InWidth, uint32 InHeight )
 	UpdateRHI(FRHICommandListImmediate::Get());
 }
 
-void FSlateTexture2DRHIRef::SetRHIRef( FTexture2DRHIRef InRHIRef, uint32 InWidth, uint32 InHeight )
+void FSlateTexture2DRHIRef::SetRHIRef( FTextureRHIRef InRHIRef, uint32 InWidth, uint32 InHeight )
 {
 	check( IsInRenderingThread() );
 	ShaderResource = InRHIRef;
@@ -262,7 +262,7 @@ void FSlateTexture2DRHIRef::UpdateTextureThreadSafeWithTextureData(FSlateTexture
 		});
 }
 
-void FSlateRenderTargetRHI::SetRHIRef( FTexture2DRHIRef InRenderTargetTexture, uint32 InWidth, uint32 InHeight )
+void FSlateRenderTargetRHI::SetRHIRef( FTextureRHIRef InRenderTargetTexture, uint32 InWidth, uint32 InHeight )
 {
 	check( IsInRenderingThread() );
 	ShaderResource = InRenderTargetTexture;
diff --git a/Engine/Source/Runtime/Engine/Private/Streaming/Texture2DMipAllocator_AsyncCreate.cpp b/Engine/Source/Runtime/Engine/Private/Streaming/Texture2DMipAllocator_AsyncCreate.cpp
--- a/Engine/Source/Runtime/Engine/Private/Streaming/Texture2DMipAllocator_AsyncCreate.cpp
+++ b/Engine/Source/Runtime/Engine/Private/Streaming/Texture2DMipAllocator_AsyncCreate.cpp
@@ -72,7 +72,7 @@ bool FTexture2DMipAllocator_AsyncCreate::AllocateMips(
 // - Render : swap the results
 bool FTexture2DMipAllocator_AsyncCreate::FinalizeMips(const FTextureUpdateContext& Context, const FTextureUpdateSyncOptions& SyncOptions)
 {
-	FRHITexture2D* Texture2DRHI = Context.Resource ? Context.Resource->GetTexture2DRHI() : nullptr;
+	FRHITexture* Texture2DRHI = Context.Resource ? Context.Resource->GetTexture2DRHI() : nullptr;
 	if (!Texture2DRHI)
 	{
 		return false;
diff --git a/Engine/Source/Runtime/Engine/Private/Streaming/Texture2DMipAllocator_AsyncCreate.h b/Engine/Source/Runtime/Engine/Private/Streaming/Texture2DMipAllocator_AsyncCreate.h
--- a/Engine/Source/Runtime/Engine/Private/Streaming/Texture2DMipAllocator_AsyncCreate.h
+++ b/Engine/Source/Runtime/Engine/Private/Streaming/Texture2DMipAllocator_AsyncCreate.h
@@ -41,7 +41,7 @@ protected:
 	void ReleaseAllocatedMipData();
 
 	// The intermediate texture async created in the update process.
-	FTexture2DRHIRef IntermediateTextureRHI;
+	FTextureRHIRef IntermediateTextureRHI;
 
 	// The final resolution x and y of the texture 2d.
 	int32 FinalSizeX = 0;
diff --git a/Engine/Source/Runtime/Engine/Private/Streaming/Texture2DMipAllocator_AsyncReallocate.cpp b/Engine/Source/Runtime/Engine/Private/Streaming/Texture2DMipAllocator_AsyncReallocate.cpp
--- a/Engine/Source/Runtime/Engine/Private/Streaming/Texture2DMipAllocator_AsyncReallocate.cpp
+++ b/Engine/Source/Runtime/Engine/Private/Streaming/Texture2DMipAllocator_AsyncReallocate.cpp
@@ -32,7 +32,7 @@ bool FTexture2DMipAllocator_AsyncReallocate::AllocateMips(
 {
 	check(PendingFirstLODIdx < CurrentFirstLODIdx);
 
-	FRHITexture2D* Texture2DRHI = Context.Resource ? Context.Resource->GetTexture2DRHI() : nullptr;
+	FRHITexture* Texture2DRHI = Context.Resource ? Context.Resource->GetTexture2DRHI() : nullptr;
 	if (!Texture2DRHI)
 	{
 		return false;
diff --git a/Engine/Source/Runtime/Engine/Private/Streaming/Texture2DMipAllocator_AsyncReallocate.h b/Engine/Source/Runtime/Engine/Private/Streaming/Texture2DMipAllocator_AsyncReallocate.h
--- a/Engine/Source/Runtime/Engine/Private/Streaming/Texture2DMipAllocator_AsyncReallocate.h
+++ b/Engine/Source/Runtime/Engine/Private/Streaming/Texture2DMipAllocator_AsyncReallocate.h
@@ -39,7 +39,7 @@ protected:
 	void UnlockNewMips();
 
 	// The intermediate texture created with RHIAsyncReallocateTexture2D.
-	FTexture2DRHIRef IntermediateTextureRHI;
+	FTextureRHIRef IntermediateTextureRHI;
 	// The list of mips that are currently locked.
 	TArray<int32, TInlineAllocator<MAX_TEXTURE_MIP_COUNT>> LockedMipIndices;
 };
diff --git a/Engine/Source/Runtime/Engine/Private/Streaming/VolumeTextureStreaming.h b/Engine/Source/Runtime/Engine/Private/Streaming/VolumeTextureStreaming.h
--- a/Engine/Source/Runtime/Engine/Private/Streaming/VolumeTextureStreaming.h
+++ b/Engine/Source/Runtime/Engine/Private/Streaming/VolumeTextureStreaming.h
@@ -42,7 +42,7 @@ public:
 protected:
 
 	// The intermediate texture async created in the update process.
-	FTexture3DRHIRef IntermediateTextureRHI;
+	FTextureRHIRef IntermediateTextureRHI;
 
 	// The temporary main memory allocations holding the mip data
 	FVolumeTextureBulkData StreamedInMipData;
diff --git a/Engine/Source/Runtime/Engine/Private/Texture2D.cpp b/Engine/Source/Runtime/Engine/Private/Texture2D.cpp
--- a/Engine/Source/Runtime/Engine/Private/Texture2D.cpp
+++ b/Engine/Source/Runtime/Engine/Private/Texture2D.cpp
@@ -1646,7 +1646,7 @@ void FVirtualTexture2DResource::InitializeEditorResources(IVirtualTexture* InVir
 		FRHITextureCreateDesc Desc = FRHITextureCreateDesc::Create2D(*Name, MipWidthInTiles * TileSizeInPixels, MipHeightInTiles * TileSizeInPixels, PixelFormat);
 		Desc.AddFlags(TexCreateFlags);
 
-		FTexture2DRHIRef Texture2DRHI = RHICreateTexture(Desc);
+		FTextureRHIRef Texture2DRHI = RHICreateTexture(Desc);
 
 		FRHICommandListImmediate& RHICommandList = FRHICommandListExecutor::GetImmediateCommandList();
 
@@ -1696,7 +1696,7 @@ void FVirtualTexture2DResource::InitializeEditorResources(IVirtualTexture* InVir
 				.SetFlags(Desc.Flags)
 				.SetInitialState(ERHIAccess::CopyDest);
 
-			FTexture2DRHIRef ResizedTexture2DRHI = RHICreateTexture(ResizedDesc);
+			FTextureRHIRef ResizedTexture2DRHI = RHICreateTexture(ResizedDesc);
 
 			FRHICopyTextureInfo CopyInfo;
 			CopyInfo.Size = FIntVector(MipWidth, MipHeight, 1);
diff --git a/Engine/Source/Runtime/Engine/Private/Texture2DDynamic.cpp b/Engine/Source/Runtime/Engine/Private/Texture2DDynamic.cpp
--- a/Engine/Source/Runtime/Engine/Private/Texture2DDynamic.cpp
+++ b/Engine/Source/Runtime/Engine/Private/Texture2DDynamic.cpp
@@ -85,7 +85,7 @@ void FTexture2DDynamicResource::ReleaseRHI()
 }
 
 /** Returns the Texture2DRHI, which can be used for locking/unlocking the mips. */
-FTexture2DRHIRef FTexture2DDynamicResource::GetTexture2DRHI()
+FTextureRHIRef FTexture2DDynamicResource::GetTexture2DRHI()
 {
 	return Texture2DRHI;
 }
diff --git a/Engine/Source/Runtime/Engine/Private/TextureCube.cpp b/Engine/Source/Runtime/Engine/Private/TextureCube.cpp
--- a/Engine/Source/Runtime/Engine/Private/TextureCube.cpp
+++ b/Engine/Source/Runtime/Engine/Private/TextureCube.cpp
@@ -537,7 +537,7 @@ public:
 	 * Accessor
 	 * @return Texture2DRHI
 	 */
-	FTextureCubeRHIRef GetTextureCubeRHI() const
+	FTextureRHIRef GetTextureCubeRHI() const
 	{
 		return TextureCubeRHI;
 	}
@@ -547,7 +547,7 @@ public:
 	const FTextureCubeResource* GetProxiedResource() const { return ProxiedResource; }
 private:
 	/** A reference to the texture's RHI resource as a cube-map texture. */
-	FTextureCubeRHIRef TextureCubeRHI;
+	FTextureRHIRef TextureCubeRHI;
 
 	/** Local copy/ cache of mip data. Only valid between creation and first call to InitRHI */
 	void* MipData[6][MAX_TEXTURE_MIP_COUNT];
diff --git a/Engine/Source/Runtime/Engine/Private/TextureCubeArray.cpp b/Engine/Source/Runtime/Engine/Private/TextureCubeArray.cpp
--- a/Engine/Source/Runtime/Engine/Private/TextureCubeArray.cpp
+++ b/Engine/Source/Runtime/Engine/Private/TextureCubeArray.cpp
@@ -224,7 +224,7 @@ public:
 		return Owner->GetNumSlices();
 	}
 
-	FTextureCubeRHIRef GetTextureCubeRHI() const
+	FTextureRHIRef GetTextureCubeRHI() const
 	{
 		return TextureCubeRHI;
 	}
@@ -234,7 +234,7 @@ public:
 	const FTextureCubeArrayResource* GetProxiedResource() const { return ProxiedResource; }
 private:
 	/** A reference to the texture's RHI resource as a cube-map texture. */
-	FTextureCubeRHIRef TextureCubeRHI;
+	FTextureRHIRef TextureCubeRHI;
 
 	/** Local copy/ cache of mip data. Only valid between creation and first call to InitRHI */
 	TArray<void*> MipData;
diff --git a/Engine/Source/Runtime/Engine/Private/TileRendering.cpp b/Engine/Source/Runtime/Engine/Private/TileRendering.cpp
--- a/Engine/Source/Runtime/Engine/Private/TileRendering.cpp
+++ b/Engine/Source/Runtime/Engine/Private/TileRendering.cpp
@@ -288,7 +288,7 @@ bool FCanvasTileRendererItem::Render_GameThread(const FCanvas* Canvas, FCanvasRe
 
 		bool bRequiresExplicit128bitRT = false;
 
-		FTexture2DRHIRef CanvasRTTexture = CanvasRenderTarget->GetRenderTargetTexture();
+		FTextureRHIRef CanvasRTTexture = CanvasRenderTarget->GetRenderTargetTexture();
 		if (CanvasRTTexture)
 		{
 			bRequiresExplicit128bitRT = PlatformRequires128bitRT(CanvasRTTexture->GetFormat());
diff --git a/Engine/Source/Runtime/Engine/Private/UserInterface/Canvas.cpp b/Engine/Source/Runtime/Engine/Private/UserInterface/Canvas.cpp
--- a/Engine/Source/Runtime/Engine/Private/UserInterface/Canvas.cpp
+++ b/Engine/Source/Runtime/Engine/Private/UserInterface/Canvas.cpp
@@ -57,9 +57,9 @@ public:
 		return RDGTexture->Desc.Extent;
 	}
 
-	const FTexture2DRHIRef& GetRenderTargetTexture() const override
+	const FTextureRHIRef& GetRenderTargetTexture() const override
 	{
-		static FTexture2DRHIRef NullRef;
+		static FTextureRHIRef NullRef;
 		return NullRef;
 	}
 
diff --git a/Engine/Source/Runtime/Engine/Private/VT/VirtualTextureUploadCache.cpp b/Engine/Source/Runtime/Engine/Private/VT/VirtualTextureUploadCache.cpp
--- a/Engine/Source/Runtime/Engine/Private/VT/VirtualTextureUploadCache.cpp
+++ b/Engine/Source/Runtime/Engine/Private/VT/VirtualTextureUploadCache.cpp
@@ -420,7 +420,7 @@ FVTUploadTileHandle FVirtualTextureUploadCache::PrepareTileForUpload(FRHICommand
 	return FVTUploadTileHandle(Index);
 }
 
-void FVirtualTextureUploadCache::SubmitTile(FRHICommandList& RHICmdList, const FVTUploadTileHandle& InHandle, FRHITexture2D* InDestTexture, int InDestX, int InDestY, int InSkipBorderSize)
+void FVirtualTextureUploadCache::SubmitTile(FRHICommandList& RHICmdList, const FVTUploadTileHandle& InHandle, FRHITexture* InDestTexture, int InDestX, int InDestY, int InSkipBorderSize)
 {
 	checkSlow(IsInParallelRenderingThread());
 
diff --git a/Engine/Source/Runtime/Engine/Private/VT/VirtualTextureUploadCache.h b/Engine/Source/Runtime/Engine/Private/VT/VirtualTextureUploadCache.h
--- a/Engine/Source/Runtime/Engine/Private/VT/VirtualTextureUploadCache.h
+++ b/Engine/Source/Runtime/Engine/Private/VT/VirtualTextureUploadCache.h
@@ -142,7 +142,7 @@ public:
 	 * Mark streamed upload data ready for upload to to the physical virtual texture.
 	 * Depending on the platform the upload might happen here, or be deferred to the Finalize() call.
 	 */
-	void SubmitTile(FRHICommandList& RHICmdList, const FVTUploadTileHandle& InHandle, FRHITexture2D* InDestTexture, int InDestX, int InDestY, int InSkipBorderSize);
+	void SubmitTile(FRHICommandList& RHICmdList, const FVTUploadTileHandle& InHandle, FRHITexture* InDestTexture, int InDestX, int InDestY, int InSkipBorderSize);
 	/** Cancel a tile that was already in flight. */
 	void CancelTile(FRHICommandList& RHICmdList, const FVTUploadTileHandle& InHandle);
 
@@ -165,7 +165,7 @@ private:
 		uint32 PoolIndex = 0;
 		uint32 TileHandle = 0u;
 		uint32 FrameSubmitted = 0u;
-		FRHITexture2D* RHISubmitTexture = nullptr;
+		FRHITexture* RHISubmitTexture = nullptr;
 		int32 SubmitDestX = 0;
 		int32 SubmitDestY = 0;
 		int32 SubmitSkipBorderSize = 0;
@@ -174,7 +174,7 @@ private:
 	/** Staging texture used for tile upload. Only used on platforms that don't have faster upload methods. */
 	struct FStagingTexture
 	{
-		TRefCountPtr<FRHITexture2D> RHITexture;
+		TRefCountPtr<FRHITexture> RHITexture;
 		uint32 WidthInTiles = 0u;
 		uint32 BatchCapacity = 0u;
 		bool bIsCPUWritable;
diff --git a/Engine/Source/Runtime/Engine/Private/VectorField.h b/Engine/Source/Runtime/Engine/Private/VectorField.h
--- a/Engine/Source/Runtime/Engine/Private/VectorField.h
+++ b/Engine/Source/Runtime/Engine/Private/VectorField.h
@@ -39,7 +39,7 @@ class FVectorFieldResource : public FRenderResource
 public:
 
 	/** The volume texture containing the vector field. */
-	FTexture3DRHIRef VolumeTextureRHI;
+	FTextureRHIRef VolumeTextureRHI;
 	/** Size of the vector field (X). */
 	int32 SizeX;
 	/** Size of the vector field (Y). */
diff --git a/Engine/Source/Runtime/Engine/Private/VectorFieldVisualization.cpp b/Engine/Source/Runtime/Engine/Private/VectorFieldVisualization.cpp
--- a/Engine/Source/Runtime/Engine/Private/VectorFieldVisualization.cpp
+++ b/Engine/Source/Runtime/Engine/Private/VectorFieldVisualization.cpp
@@ -146,7 +146,7 @@ struct FVectorFieldVisualizationUserData : public FOneFrameResource
 	FVectorFieldVisualizationBufferRef UniformBuffer;
 
 	/** Texture containing the vector field. */
-	FTexture3DRHIRef VectorFieldTextureRHI;
+	FTextureRHIRef VectorFieldTextureRHI;
 };
 
 void FVectorFieldVisualizationVertexFactoryShaderParameters::GetElementShaderBindings(
diff --git a/Engine/Source/Runtime/Engine/Private/WorldPartition/HLOD/DestructibleHLODComponent.cpp b/Engine/Source/Runtime/Engine/Private/WorldPartition/HLOD/DestructibleHLODComponent.cpp
--- a/Engine/Source/Runtime/Engine/Private/WorldPartition/HLOD/DestructibleHLODComponent.cpp
+++ b/Engine/Source/Runtime/Engine/Private/WorldPartition/HLOD/DestructibleHLODComponent.cpp
@@ -295,7 +295,7 @@ void UWorldPartitionDestructibleHLODMeshComponent::UpdateVisibilityTexture_Rende
 {
 	check(IsInRenderingThread());
 
-	FRHITexture2D* TextureRHI = TextureResource->GetTexture2DRHI();
+	FRHITexture* TextureRHI = TextureResource->GetTexture2DRHI();
 
 	uint32 DestStride = 0;
 	uint8* DestData = reinterpret_cast<uint8*>(RHILockTexture2D(TextureRHI, 0, RLM_WriteOnly, DestStride, false, false));
diff --git a/Engine/Source/Runtime/Engine/Public/DistanceFieldDownsampling.h b/Engine/Source/Runtime/Engine/Public/DistanceFieldDownsampling.h
--- a/Engine/Source/Runtime/Engine/Public/DistanceFieldDownsampling.h
+++ b/Engine/Source/Runtime/Engine/Public/DistanceFieldDownsampling.h
@@ -13,7 +13,7 @@ class FRDGBuilder;
 
 struct FDistanceFieldDownsamplingDataTask
 {
-	FTexture3DRHIRef VolumeTextureRHI;
+	FTextureRHIRef VolumeTextureRHI;
 	FVector TexelSrcSize;
 	FIntVector DstSize;
 	FIntVector OffsetInAtlas;
diff --git a/Engine/Source/Runtime/Engine/Public/DummyViewport.h b/Engine/Source/Runtime/Engine/Public/DummyViewport.h
--- a/Engine/Source/Runtime/Engine/Public/DummyViewport.h
+++ b/Engine/Source/Runtime/Engine/Public/DummyViewport.h
@@ -19,7 +19,7 @@ public:
 	virtual void BeginRenderFrame(FRHICommandListImmediate& RHICmdList) override
 	{
 		check( IsInRenderingThread() );
-		//SetRenderTarget(RHICmdList,  RenderTargetTextureRHI,  FTexture2DRHIRef() );
+		//SetRenderTarget(RHICmdList,  RenderTargetTextureRHI,  FTextureRHIRef() );
 	};
 
 	virtual void EndRenderFrame(FRHICommandListImmediate& RHICmdList, bool bPresent, bool bLockToVsync) override
diff --git a/Engine/Source/Runtime/Engine/Public/IStereoLayers.h b/Engine/Source/Runtime/Engine/Public/IStereoLayers.h
--- a/Engine/Source/Runtime/Engine/Public/IStereoLayers.h
+++ b/Engine/Source/Runtime/Engine/Public/IStereoLayers.h
@@ -264,11 +264,11 @@ public:
 	}
 
 protected:
-	bool				bSplashIsShown = false;
-	bool				bSplashShowMovie = false;
-	FTexture2DRHIRef	SplashTexture;
-	FTexture2DRHIRef	SplashMovie;
-	FVector			    SplashOffset = FVector::ZeroVector;
-	FVector2D			SplashScale = FVector2D(1.0f, 1.0f);
-	uint32				SplashLayerHandle = 0;
+	bool			bSplashIsShown = false;
+	bool			bSplashShowMovie = false;
+	FTextureRHIRef	SplashTexture;
+	FTextureRHIRef	SplashMovie;
+	FVector			SplashOffset = FVector::ZeroVector;
+	FVector2D		SplashScale = FVector2D(1.0f, 1.0f);
+	uint32			SplashLayerHandle = 0;
 };
diff --git a/Engine/Source/Runtime/Engine/Public/Slate/SceneViewport.h b/Engine/Source/Runtime/Engine/Public/Slate/SceneViewport.h
--- a/Engine/Source/Runtime/Engine/Public/Slate/SceneViewport.h
+++ b/Engine/Source/Runtime/Engine/Public/Slate/SceneViewport.h
@@ -58,9 +58,9 @@ public:
 	ENGINE_API virtual void EnqueueEndRenderFrame(const bool bLockToVsync, const bool bShouldPresent) override;
 
 	/** Gets the proper RenderTarget based on the current thread*/
-	ENGINE_API virtual const FTexture2DRHIRef& GetRenderTargetTexture() const;
+	ENGINE_API virtual const FTextureRHIRef& GetRenderTargetTexture() const;
 
-	ENGINE_API virtual void SetRenderTargetTextureRenderThread(FTexture2DRHIRef& RT);
+	ENGINE_API virtual void SetRenderTargetTextureRenderThread(FTextureRHIRef& RT);
 
 	/**
 	 * Captures or uncaptures the joystick
@@ -430,10 +430,10 @@ private:
 	/** Reprojection on some HMD RHI's requires ViewportTargets to be buffered */
 	/** The render target used by Slate to draw the viewport.  Can be null if this viewport renders directly to the backbuffer */
 	TArray<class FSlateRenderTargetRHI*> BufferedSlateHandles;
-	TArray<FTexture2DRHIRef> BufferedRenderTargetsRHI;
-	TArray<FTexture2DRHIRef> BufferedShaderResourceTexturesRHI;
+	TArray<FTextureRHIRef> BufferedRenderTargetsRHI;
+	TArray<FTextureRHIRef> BufferedShaderResourceTexturesRHI;
 
-	FTexture2DRHIRef RenderTargetTextureRenderThreadRHI;
+	FTextureRHIRef RenderTargetTextureRenderThreadRHI;
 	class FSlateRenderTargetRHI* RenderThreadSlateTexture;
 
 	int32 CurrentBufferedTargetIndex;
diff --git a/Engine/Source/Runtime/Engine/Public/Slate/SlateTextures.h b/Engine/Source/Runtime/Engine/Public/Slate/SlateTextures.h
--- a/Engine/Source/Runtime/Engine/Public/Slate/SlateTextures.h
+++ b/Engine/Source/Runtime/Engine/Public/Slate/SlateTextures.h
@@ -17,10 +17,10 @@
 /**
  * Encapsulates a Texture2DRHIRef for use by a Slate rendering implementation                   
  */
-class FSlateTexture2DRHIRef : public TSlateTexture<FTexture2DRHIRef>, public FSlateUpdatableTexture, public FDeferredCleanupInterface, public FRenderResource
+class FSlateTexture2DRHIRef : public TSlateTexture<FTextureRHIRef>, public FSlateUpdatableTexture, public FDeferredCleanupInterface, public FRenderResource
 {
 public:
-	ENGINE_API FSlateTexture2DRHIRef( FTexture2DRHIRef InRef, uint32 InWidth, uint32 InHeight );
+	ENGINE_API FSlateTexture2DRHIRef( FTextureRHIRef InRef, uint32 InWidth, uint32 InHeight );
 	ENGINE_API FSlateTexture2DRHIRef( uint32 InWidth, uint32 InHeight, EPixelFormat InPixelFormat, TSharedPtr<FSlateTextureData, ESPMode::ThreadSafe> InTextureData, ETextureCreateFlags InTexCreateFlags = TexCreate_None, bool bCreateEmptyTexture = false );
 
 	ENGINE_API virtual ~FSlateTexture2DRHIRef();
@@ -49,9 +49,9 @@ public:
 	/** 
 	 *  Sets the RHI Ref to use. 
 	 */
-	ENGINE_API void SetRHIRef( FTexture2DRHIRef InRenderTargetTexture, uint32 InWidth, uint32 InHeight );
+	ENGINE_API void SetRHIRef( FTextureRHIRef InRenderTargetTexture, uint32 InWidth, uint32 InHeight );
 
-	FTexture2DRHIRef GetRHIRef() const { return ShaderResource; }
+	FTextureRHIRef GetRHIRef() const { return ShaderResource; }
 
 	/**
 	 * Sets the bulk data for this texture.  Note: Does not reinitialize the resource,  Can only be used on the render thread
@@ -110,10 +110,10 @@ private:
 /**
  * Encapsulates a render target for use by a Slate rendering implementation                   
  */
-class FSlateRenderTargetRHI : public TSlateTexture<FTexture2DRHIRef>, public FRenderResource
+class FSlateRenderTargetRHI : public TSlateTexture<FTextureRHIRef>, public FRenderResource
 {
 public:
-	FSlateRenderTargetRHI( FTexture2DRHIRef InRenderTargetTexture, uint32 InWidth, uint32 InHeight )
+	FSlateRenderTargetRHI( FTextureRHIRef InRenderTargetTexture, uint32 InWidth, uint32 InHeight )
 		: TSlateTexture( InRenderTargetTexture )
 		, Width( InWidth )
 		, Height( InHeight )
@@ -136,9 +136,9 @@ public:
 	/** 
 	 *  Sets the RHI Ref to use.  Useful for reusing this class for multiple render targets
 	 */
-	ENGINE_API void SetRHIRef( FTexture2DRHIRef InRenderTargetTexture, uint32 InWidth, uint32 InHeight );
+	ENGINE_API void SetRHIRef( FTextureRHIRef InRenderTargetTexture, uint32 InWidth, uint32 InHeight );
 
-	FTexture2DRHIRef GetRHIRef() const { return ShaderResource; }
+	FTextureRHIRef GetRHIRef() const { return ShaderResource; }
 private:
 	/** Width of this texture */
 	uint32 Width;
diff --git a/Engine/Source/Runtime/Engine/Public/StereoRenderTargetManager.h b/Engine/Source/Runtime/Engine/Public/StereoRenderTargetManager.h
--- a/Engine/Source/Runtime/Engine/Public/StereoRenderTargetManager.h
+++ b/Engine/Source/Runtime/Engine/Public/StereoRenderTargetManager.h
@@ -12,7 +12,7 @@ StereoRenderTargetManager.h: Abstract interface returned from IStereoRendering t
 class FRHITexture;
 enum class ETextureCreateFlags : uint64;
 enum EShaderPlatform : uint16;
-typedef TRefCountPtr<FRHITexture> FTexture2DRHIRef;
+typedef TRefCountPtr<FRHITexture> FTextureRHIRef;
 
 /** 
  * The IStereoRenderTargetManager can be returned from IStereoRendering::GetRenderTargetManager() implementations.
@@ -74,7 +74,7 @@ public:
 	 * @return				true, if texture was allocated; false, if the default texture allocation should be used.
 	 */
 	UE_DEPRECATED(5.2, "Implement AllocateRenderTargetTextures to allocate all textures at once")
-	virtual bool AllocateRenderTargetTexture(uint32 Index, uint32 SizeX, uint32 SizeY, uint8 Format, uint32 NumMips, ETextureCreateFlags Flags, ETextureCreateFlags TargetableTextureFlags, FTexture2DRHIRef& OutTargetableTexture, FTexture2DRHIRef& OutShaderResourceTexture, uint32 NumSamples = 1) { return false; }
+	virtual bool AllocateRenderTargetTexture(uint32 Index, uint32 SizeX, uint32 SizeY, uint8 Format, uint32 NumMips, ETextureCreateFlags Flags, ETextureCreateFlags TargetableTextureFlags, FTextureRHIRef& OutTargetableTexture, FTextureRHIRef& OutShaderResourceTexture, uint32 NumSamples = 1) { return false; }
 
 	/**
 	 * Allocates the render target textures, which includes the color textures and optionally other textures like depth.
@@ -82,7 +82,7 @@ public:
 	 *
 	 * @return				true, if textures were allocated; false, if the default texture allocation should be used.
 	 */
-	virtual bool AllocateRenderTargetTextures(uint32 SizeX, uint32 SizeY, uint8 Format, uint32 NumLayers, ETextureCreateFlags Flags, ETextureCreateFlags TargetableTextureFlags, TArray<FTexture2DRHIRef>& OutTargetableTextures, TArray<FTexture2DRHIRef>& OutShaderResourceTextures, uint32 NumSamples = 1) { return false; }
+	virtual bool AllocateRenderTargetTextures(uint32 SizeX, uint32 SizeY, uint8 Format, uint32 NumLayers, ETextureCreateFlags Flags, ETextureCreateFlags TargetableTextureFlags, TArray<FTextureRHIRef>& OutTargetableTextures, TArray<FTextureRHIRef>& OutShaderResourceTextures, uint32 NumSamples = 1) { return false; }
 
 	/**
 	 * Returns pixel format that the device created its swapchain with (which can be different than what was requested in AllocateRenderTargetTexture)
@@ -102,7 +102,7 @@ public:
 	 * @param Index			(in) index of the buffer, changing from 0 to GetNumberOfBufferedFrames()
 	 * @return				true, if texture was allocated; false, if the default texture allocation should be used.
 	 */
-	virtual bool AllocateDepthTexture(uint32 Index, uint32 SizeX, uint32 SizeY, uint8 Format, uint32 NumMips, ETextureCreateFlags Flags, ETextureCreateFlags TargetableTextureFlags, FTexture2DRHIRef& OutTargetableTexture, FTexture2DRHIRef& OutShaderResourceTexture, uint32 NumSamples = 1) { return false; }
+	virtual bool AllocateDepthTexture(uint32 Index, uint32 SizeX, uint32 SizeY, uint8 Format, uint32 NumMips, ETextureCreateFlags Flags, ETextureCreateFlags TargetableTextureFlags, FTextureRHIRef& OutTargetableTexture, FTextureRHIRef& OutShaderResourceTexture, uint32 NumSamples = 1) { return false; }
 
 	/**
 	 * Allocates a shading rate texture.
@@ -111,7 +111,7 @@ public:
 	 * @param Index			(in) index of the buffer, changing from 0 to GetNumberOfBufferedFrames()
 	 * @return				true, if texture was allocated; false, if the default texture allocation should be used.
 	 */
-	virtual bool AllocateShadingRateTexture(uint32 Index, uint32 RenderSizeX, uint32 RenderSizeY, uint8 Format, uint32 NumMips, ETextureCreateFlags Flags, ETextureCreateFlags TargetableTextureFlags, FTexture2DRHIRef& OutTexture, FIntPoint& OutTextureSize) { return false; }
+	virtual bool AllocateShadingRateTexture(uint32 Index, uint32 RenderSizeX, uint32 RenderSizeY, uint8 Format, uint32 NumMips, ETextureCreateFlags Flags, ETextureCreateFlags TargetableTextureFlags, FTextureRHIRef& OutTexture, FIntPoint& OutTextureSize) { return false; }
 
 	/**
 	 * Retrieves HDR information about the stereo device, if any is available.
diff --git a/Engine/Source/Runtime/Engine/Public/TextureResource.h b/Engine/Source/Runtime/Engine/Public/TextureResource.h
--- a/Engine/Source/Runtime/Engine/Public/TextureResource.h
+++ b/Engine/Source/Runtime/Engine/Public/TextureResource.h
@@ -253,7 +253,7 @@ public:
 	ENGINE_API virtual void ReleaseRHI() override;
 
 	/** Returns the Texture2DRHI, which can be used for locking/unlocking the mips. */
-	ENGINE_API FTexture2DRHIRef GetTexture2DRHI();
+	ENGINE_API FTextureRHIRef GetTexture2DRHI();
 
 #if !UE_SERVER
 	ENGINE_API void WriteRawToTexture_RenderThread(TArrayView64<const uint8> RawData);
@@ -263,7 +263,7 @@ private:
 	/** The owner of this resource. */
 	class UTexture2DDynamic* Owner;
 	/** Texture2D reference, used for locking/unlocking the mips. */
-	FTexture2DRHIRef Texture2DRHI;
+	FTextureRHIRef Texture2DRHI;
 };
 
 /**
@@ -521,7 +521,7 @@ private:
 	const class UTextureRenderTarget2D* Owner;
 	/** Texture resource used for rendering with and resolving to */
 	UE_DEPRECATED(5.1, "Texture2DRHI is deprecated. Use TextureRHI instead.")
-	FTexture2DRHIRef Texture2DRHI;
+	FTextureRHIRef Texture2DRHI;
 	/** Optional Unordered Access View for the resource, automatically created if bCanCreateUAV is true */
 	FUnorderedAccessViewRHIRef UnorderedAccessViewRHI;
 	/** the color the texture is cleared to */
diff --git a/Engine/Source/Runtime/Engine/Public/UnrealClient.h b/Engine/Source/Runtime/Engine/Public/UnrealClient.h
--- a/Engine/Source/Runtime/Engine/Public/UnrealClient.h
+++ b/Engine/Source/Runtime/Engine/Public/UnrealClient.h
@@ -764,8 +764,8 @@ protected:
 		virtual void AddReferencedObjects( FReferenceCollector& Collector ) override;
 		virtual FString GetReferencerName() const override;
 
-		const FTexture2DRHIRef& GetHitProxyTexture(void) const		{ return RenderTargetTextureRHI; }
-		const FTexture2DRHIRef& GetHitProxyCPUTexture(void) const		{ return HitProxyCPUTexture; }
+		const FTextureRHIRef& GetHitProxyTexture(void) const		{ return RenderTargetTextureRHI; }
+		const FTextureRHIRef& GetHitProxyCPUTexture(void) const		{ return HitProxyCPUTexture; }
 
 	private:
 
@@ -778,7 +778,7 @@ protected:
 		/** References to the hit proxies cached by the hit proxy map. */
 		TArray<TRefCountPtr<HHitProxy> > HitProxies;
 
-		FTexture2DRHIRef HitProxyCPUTexture;
+		FTextureRHIRef HitProxyCPUTexture;
 	};
 
 	/** The viewport's hit proxy map. */
diff --git a/Engine/Source/Runtime/GameplayMediaEncoder/Private/GameplayMediaEncoder.cpp b/Engine/Source/Runtime/GameplayMediaEncoder/Private/GameplayMediaEncoder.cpp
--- a/Engine/Source/Runtime/GameplayMediaEncoder/Private/GameplayMediaEncoder.cpp
+++ b/Engine/Source/Runtime/GameplayMediaEncoder/Private/GameplayMediaEncoder.cpp
@@ -410,7 +410,7 @@ void FGameplayMediaEncoder::OnNewSubmixBuffer(const USoundSubmix* OwningSubmix,
 	ProcessAudioFrame(AudioData, NumSamples, NumChannels, SampleRate);
 }
 
-void FGameplayMediaEncoder::OnFrameBufferReady(SWindow& SlateWindow, const FTexture2DRHIRef& FrameBuffer)
+void FGameplayMediaEncoder::OnFrameBufferReady(SWindow& SlateWindow, const FTextureRHIRef& FrameBuffer)
 {
 	CSV_SCOPED_TIMING_STAT(GameplayMediaEncoder, OnBackBufferReady);
 	check(IsInRenderingThread());
@@ -541,7 +541,7 @@ void FGameplayMediaEncoder::ProcessAudioFrame(const float* AudioData, int32 NumS
 	AudioClock += FloatBuffer.GetSampleDuration();
 }
 
-void FGameplayMediaEncoder::ProcessVideoFrame(const FTexture2DRHIRef& FrameBuffer)
+void FGameplayMediaEncoder::ProcessVideoFrame(const FTextureRHIRef& FrameBuffer)
 {
 	// Early exit is video encoder is not valid because it is not setup or has been destroyed
 	if(!VideoEncoder.IsValid())
@@ -689,7 +689,7 @@ void FGameplayMediaEncoder::OnEncodedVideoFrame(uint32 LayerIndex, const TShared
 	InputFrame->Release();
 }
 
-void FGameplayMediaEncoder::CopyTexture(const FTexture2DRHIRef& SourceTexture, FTexture2DRHIRef& DestinationTexture) const
+void FGameplayMediaEncoder::CopyTexture(const FTextureRHIRef& SourceTexture, FTextureRHIRef& DestinationTexture) const
 {
 	FRHICommandListImmediate& RHICmdList = FRHICommandListExecutor::GetImmediateCommandList();
 
diff --git a/Engine/Source/Runtime/GameplayMediaEncoder/Public/GameplayMediaEncoder.h b/Engine/Source/Runtime/GameplayMediaEncoder/Public/GameplayMediaEncoder.h
--- a/Engine/Source/Runtime/GameplayMediaEncoder/Public/GameplayMediaEncoder.h
+++ b/Engine/Source/Runtime/GameplayMediaEncoder/Public/GameplayMediaEncoder.h
@@ -101,14 +101,14 @@ private:
 	GAMEPLAYMEDIAENCODER_API FTimespan GetMediaTimestamp() const;
 
 	// Back buffer capture
-	GAMEPLAYMEDIAENCODER_API void OnFrameBufferReady(SWindow& SlateWindow, const FTexture2DRHIRef& FrameBuffer);
+	GAMEPLAYMEDIAENCODER_API void OnFrameBufferReady(SWindow& SlateWindow, const FTextureRHIRef& FrameBuffer);
 
 	// ISubmixBufferListener interface
 	GAMEPLAYMEDIAENCODER_API const FString& GetListenerName() const override;
 	GAMEPLAYMEDIAENCODER_API void OnNewSubmixBuffer(const USoundSubmix* OwningSubmix, float* AudioData, int32 NumSamples, int32 NumChannels, const int32 SampleRate, double AudioClock) override;
 
 	GAMEPLAYMEDIAENCODER_API void ProcessAudioFrame(const float* AudioData, int32 NumSamples, int32 NumChannels, int32 SampleRate);
-	GAMEPLAYMEDIAENCODER_API void ProcessVideoFrame(const FTexture2DRHIRef& FrameBuffer);
+	GAMEPLAYMEDIAENCODER_API void ProcessVideoFrame(const FTextureRHIRef& FrameBuffer);
 
 	GAMEPLAYMEDIAENCODER_API void UpdateVideoConfig();
 
@@ -116,7 +116,7 @@ private:
 	GAMEPLAYMEDIAENCODER_API void OnEncodedVideoFrame(uint32 LayerIndex, const TSharedPtr<AVEncoder::FVideoEncoderInputFrame> Frame, const AVEncoder::FCodecPacket& Packet);
 
 	GAMEPLAYMEDIAENCODER_API TSharedPtr<AVEncoder::FVideoEncoderInputFrame> ObtainInputFrame();
-	GAMEPLAYMEDIAENCODER_API void CopyTexture(const FTexture2DRHIRef& SourceTexture, FTexture2DRHIRef& DestinationTexture) const;
+	GAMEPLAYMEDIAENCODER_API void CopyTexture(const FTextureRHIRef& SourceTexture, FTextureRHIRef& DestinationTexture) const;
 
 	GAMEPLAYMEDIAENCODER_API void FloatToPCM16(float const* floatSamples, int32 numSamples, TArray<int16>& out) const;
 
@@ -156,6 +156,6 @@ private:
 	FThreadSafeBool bChangeFramerate = false;
 
 	TArray<int16> PCM16;
-	TMap<TSharedPtr<AVEncoder::FVideoEncoderInputFrame>, FTexture2DRHIRef> BackBuffers;
+	TMap<TSharedPtr<AVEncoder::FVideoEncoderInputFrame>, FTextureRHIRef> BackBuffers;
 };
 
diff --git a/Engine/Source/Runtime/HeadMountedDisplay/Private/IXRLoadingScreen.cpp b/Engine/Source/Runtime/HeadMountedDisplay/Private/IXRLoadingScreen.cpp
--- a/Engine/Source/Runtime/HeadMountedDisplay/Private/IXRLoadingScreen.cpp
+++ b/Engine/Source/Runtime/HeadMountedDisplay/Private/IXRLoadingScreen.cpp
@@ -6,7 +6,7 @@
 #include "Engine/Engine.h"
 
 
-void IXRLoadingScreen::ShowLoadingScreen_Compat(bool bShow, FTexture2DRHIRef Texture, const FVector& Offset, const FVector2D& Scale)
+void IXRLoadingScreen::ShowLoadingScreen_Compat(bool bShow, FTextureRHIRef Texture, const FVector& Offset, const FVector2D& Scale)
 {
 	// Backwards compatibility with IStereoLayers::ShowSplashScreen
 	IXRLoadingScreen* LoadingScreen = GEngine && GEngine->XRSystem.IsValid()? GEngine->XRSystem->GetLoadingScreen() : nullptr;
diff --git a/Engine/Source/Runtime/HeadMountedDisplay/Public/IXRLoadingScreen.h b/Engine/Source/Runtime/HeadMountedDisplay/Public/IXRLoadingScreen.h
--- a/Engine/Source/Runtime/HeadMountedDisplay/Public/IXRLoadingScreen.h
+++ b/Engine/Source/Runtime/HeadMountedDisplay/Public/IXRLoadingScreen.h
@@ -91,5 +91,5 @@ public:
 	 * Internal utility method for implementing backwards compatibility with IStereoLayers::Show/HideSplashScreen.
 	 * Should be called from implementations overriding IStereoLayer::UpdateSplashScreen()
 	 */
-	static HEADMOUNTEDDISPLAY_API void ShowLoadingScreen_Compat(bool bShow, FTexture2DRHIRef Texture, const FVector& Offset, const FVector2D& Scale);
+	static HEADMOUNTEDDISPLAY_API void ShowLoadingScreen_Compat(bool bShow, FTextureRHIRef Texture, const FVector& Offset, const FVector2D& Scale);
 };
diff --git a/Engine/Source/Runtime/ImageWriteQueue/Private/ImageWriteBlueprintLibrary.cpp b/Engine/Source/Runtime/ImageWriteQueue/Private/ImageWriteBlueprintLibrary.cpp
--- a/Engine/Source/Runtime/ImageWriteQueue/Private/ImageWriteBlueprintLibrary.cpp
+++ b/Engine/Source/Runtime/ImageWriteQueue/Private/ImageWriteBlueprintLibrary.cpp
@@ -63,7 +63,7 @@ bool UImageWriteBlueprintLibrary::ResolvePixelData(UTexture* InTexture, const FO
 	ENQUEUE_RENDER_COMMAND(ResolvePixelData)(
 		[TextureResource, OnPixelsReady](FRHICommandListImmediate& RHICmdList)
 		{
-			FTexture2DRHIRef Texture2D = TextureResource->TextureRHI ? TextureResource->TextureRHI->GetTexture2D() : nullptr;
+			FTextureRHIRef Texture2D = TextureResource->TextureRHI ? TextureResource->TextureRHI->GetTexture2D() : nullptr;
 			if (!Texture2D)
 			{
 				OnPixelsReady(nullptr);
@@ -168,7 +168,7 @@ void UImageWriteBlueprintLibrary::ExportToDisk(UTexture* InTexture, const FStrin
 		return;
 	}
 
-	FTexture2DRHIRef Texture2D = InTexture->GetResource()->TextureRHI->GetTexture2D();
+	FTextureRHIRef Texture2D = InTexture->GetResource()->TextureRHI->GetTexture2D();
 	if (!Texture2D)
 	{
 		FFrame::KismetExecutionMessage(TEXT("Invalid texture supplied."), ELogVerbosity::Error);
diff --git a/Engine/Source/Runtime/Media/Public/IMediaTextureSampleConverter.h b/Engine/Source/Runtime/Media/Public/IMediaTextureSampleConverter.h
--- a/Engine/Source/Runtime/Media/Public/IMediaTextureSampleConverter.h
+++ b/Engine/Source/Runtime/Media/Public/IMediaTextureSampleConverter.h
@@ -33,7 +33,7 @@ public:
 		return ConverterInfoFlags_Default;
 	}
 
-	virtual bool Convert(FTexture2DRHIRef & InDstTexture, const FConversionHints & Hints) = 0;
+	virtual bool Convert(FTextureRHIRef & InDstTexture, const FConversionHints & Hints) = 0;
 };
 
 /**
@@ -48,6 +48,6 @@ public:
 	 * Apply a color conversion on the input and store the result in the destination texture.
 	 * @return true If the color conversion was successfully applied.
 	 */
-	virtual bool ApplyColorConversion(FTexture2DRHIRef& InSrcTexture, FTexture2DRHIRef& InDstTexture) = 0;
+	virtual bool ApplyColorConversion(FTextureRHIRef& InSrcTexture, FTextureRHIRef& InDstTexture) = 0;
 };
 
diff --git a/Engine/Source/Runtime/MediaAssets/Private/Misc/MediaTextureResource.cpp b/Engine/Source/Runtime/MediaAssets/Private/Misc/MediaTextureResource.cpp
--- a/Engine/Source/Runtime/MediaAssets/Private/Misc/MediaTextureResource.cpp
+++ b/Engine/Source/Runtime/MediaAssets/Private/Misc/MediaTextureResource.cpp
@@ -552,7 +552,7 @@ void FMediaTextureResource::Render(const FRenderParams& Params)
 
 				if (IMediaTextureSampleColorConverter* Converter = Sample->GetMediaTextureSampleColorConverter())
 				{
-					FTexture2DRHIRef TextureRef = IntermediateTarget ? IntermediateTarget : RenderTargetTextureRHI;
+					FTextureRHIRef TextureRef = IntermediateTarget ? IntermediateTarget : RenderTargetTextureRHI;
 
 					if (bRecreateOutputTarget && TextureRef)
 					{
diff --git a/Engine/Source/Runtime/MediaAssets/Private/Misc/MediaTextureResource.h b/Engine/Source/Runtime/MediaAssets/Private/Misc/MediaTextureResource.h
--- a/Engine/Source/Runtime/MediaAssets/Private/Misc/MediaTextureResource.h
+++ b/Engine/Source/Runtime/MediaAssets/Private/Misc/MediaTextureResource.h
@@ -150,7 +150,7 @@ protected:
 	 */
 	void ConvertSample(const TSharedPtr<IMediaTextureSample, ESPMode::ThreadSafe>& Sample, const FLinearColor& ClearColor, uint8 InNumMips);
 
-	void ConvertTextureToOutput(FRHITexture2D* InputTexture, const TSharedPtr<IMediaTextureSample, ESPMode::ThreadSafe>& Sample);
+	void ConvertTextureToOutput(FRHITexture* InputTexture, const TSharedPtr<IMediaTextureSample, ESPMode::ThreadSafe>& Sample);
 
 	/**
 	 * Render the given texture sample by using it as or copying it to the render target.
@@ -171,7 +171,7 @@ protected:
 	 *
 	 * @param NewTexture The texture to set.
 	 */
-	void UpdateTextureReference(FRHITexture2D* NewTexture);
+	void UpdateTextureReference(FRHITexture* NewTexture);
 
 	/**
 	 * Create/update intermediate render target as needed. If no color conversion is needed, the RT will be used as the output.
@@ -192,7 +192,7 @@ protected:
 	void CopyFromExternalTexture(const TSharedPtr <IMediaTextureSample, ESPMode::ThreadSafe>& Sample, const FGuid & TextureGUID);
 
 	bool RequiresConversion(const TSharedPtr<IMediaTextureSample, ESPMode::ThreadSafe>& Sample, uint8 numMips) const;
-	bool RequiresConversion(const FTexture2DRHIRef& SampleTexture, const FIntPoint & OutputDim, uint8 numMips) const;
+	bool RequiresConversion(const FTextureRHIRef& SampleTexture, const FIntPoint & OutputDim, uint8 numMips) const;
 
 	/** Compute CS conversion martix based on sample's data */
 	void GetColorSpaceConversionMatrixForSample(const TSharedPtr<IMediaTextureSample, ESPMode::ThreadSafe> Sample, FMatrix44f& ColorSpaceMtx);
diff --git a/Engine/Source/Runtime/MoviePlayer/Private/DefaultGameMoviePlayer.cpp b/Engine/Source/Runtime/MoviePlayer/Private/DefaultGameMoviePlayer.cpp
--- a/Engine/Source/Runtime/MoviePlayer/Private/DefaultGameMoviePlayer.cpp
+++ b/Engine/Source/Runtime/MoviePlayer/Private/DefaultGameMoviePlayer.cpp
@@ -735,7 +735,7 @@ void FDefaultGameMoviePlayer::TickStreamer(float DeltaTime)
 		/*IXRLoadingScreen* LoadingScreen;
 		if (GEngine && GEngine->XRSystem.IsValid() && (LoadingScreen = GEngine->XRSystem->GetLoadingScreen()) != nullptr)
 		{
-			FTexture2DRHIRef Movie2DTexture = ActiveMovieStreamer->GetTexture();
+			FTextureRHIRef Movie2DTexture = ActiveMovieStreamer->GetTexture();
 			LoadingScreen->ClearSplashes();
 			if (Movie2DTexture.IsValid() && !bMovieIsDone)
 			{
diff --git a/Engine/Source/Runtime/MoviePlayer/Public/MoviePlayer.h b/Engine/Source/Runtime/MoviePlayer/Public/MoviePlayer.h
--- a/Engine/Source/Runtime/MoviePlayer/Public/MoviePlayer.h
+++ b/Engine/Source/Runtime/MoviePlayer/Public/MoviePlayer.h
@@ -104,7 +104,7 @@ public:
 	/** Called to allow the movie streamer to cleanup any resources once there are no movies left to play. */
 	virtual void Cleanup() = 0;
 
-	virtual FTexture2DRHIRef GetTexture() { return nullptr; }
+	virtual FTextureRHIRef GetTexture() { return nullptr; }
 
 	virtual ~IMovieStreamer() {}
 
diff --git a/Engine/Source/Runtime/MovieSceneCapture/Private/FrameGrabber.cpp b/Engine/Source/Runtime/MovieSceneCapture/Private/FrameGrabber.cpp
--- a/Engine/Source/Runtime/MovieSceneCapture/Private/FrameGrabber.cpp
+++ b/Engine/Source/Runtime/MovieSceneCapture/Private/FrameGrabber.cpp
@@ -87,7 +87,7 @@ void FViewportSurfaceReader::Reset()
 	bQueuedForCapture = false;
 }
 
-void FViewportSurfaceReader::ResolveRenderTarget(FViewportSurfaceReader* RenderToReadback, const FTexture2DRHIRef& SourceBackBuffer, TFunction<void(FColor*, int32, int32)> Callback)
+void FViewportSurfaceReader::ResolveRenderTarget(FViewportSurfaceReader* RenderToReadback, const FTextureRHIRef& SourceBackBuffer, TFunction<void(FColor*, int32, int32)> Callback)
 {
 	static const FName RendererModuleName( "Renderer" );
 	// @todo: JIRA UE-41879 and UE-43829 - added defensive guards against memory trampling on this render command to try and ascertain why it occasionally crashes
@@ -358,7 +358,7 @@ TArray<FCapturedFrameData> FFrameGrabber::GetCapturedFrames()
 	return ReturnFrames;
 }
 
-void FFrameGrabber::OnBackBufferReadyToPresentCallback(SWindow& SlateWindow, const FTexture2DRHIRef& BackBuffer)
+void FFrameGrabber::OnBackBufferReadyToPresentCallback(SWindow& SlateWindow, const FTextureRHIRef& BackBuffer)
 {
 	// We only care about our own Slate window
 	if (&SlateWindow != TargetWindowPtr)
diff --git a/Engine/Source/Runtime/MovieSceneCapture/Public/FrameGrabber.h b/Engine/Source/Runtime/MovieSceneCapture/Public/FrameGrabber.h
--- a/Engine/Source/Runtime/MovieSceneCapture/Public/FrameGrabber.h
+++ b/Engine/Source/Runtime/MovieSceneCapture/Public/FrameGrabber.h
@@ -35,7 +35,7 @@ struct FViewportSurfaceReader
 	 * @param	BackBuffer		The backbuffer to resolve
 	 * @param	Callback 		Callback to call with the locked texture data. This will be called on an undefined thread.
 	 */
-	MOVIESCENECAPTURE_API void ResolveRenderTarget(FViewportSurfaceReader* RenderToReadback, const FTexture2DRHIRef& BackBuffer, TFunction<void(FColor*, int32, int32)> Callback);
+	MOVIESCENECAPTURE_API void ResolveRenderTarget(FViewportSurfaceReader* RenderToReadback, const FTextureRHIRef& BackBuffer, TFunction<void(FColor*, int32, int32)> Callback);
 
 	/** Get the current size of the texture */
 	MOVIESCENECAPTURE_API FIntPoint GetCurrentSize() const;
@@ -60,7 +60,7 @@ protected:
 	FEvent* AvailableEvent;
 
 	/** Texture used to store the resolved render target */
-	FTexture2DRHIRef ReadbackTexture;
+	FTextureRHIRef ReadbackTexture;
 
 	/** The rectangle to read from the surface */
 	FIntRect CaptureRect;
@@ -165,7 +165,7 @@ public:
 protected:
 	
 	/** Callback for when a backbuffer is ready for reading (called on render thread) */
-	MOVIESCENECAPTURE_API void OnBackBufferReadyToPresentCallback(SWindow& SlateWindow, const FTexture2DRHIRef& BackBuffer);
+	MOVIESCENECAPTURE_API void OnBackBufferReadyToPresentCallback(SWindow& SlateWindow, const FTextureRHIRef& BackBuffer);
 
 	/** Called when the specified surface index has been locked for reading with the render target data (called on render thread)  */
 	MOVIESCENECAPTURE_API void OnFrameReady(int32 SurfaceIndex, FColor* ColorBuffer, int32 Width, int32 Height);
diff --git a/Engine/Source/Runtime/NullDrv/Public/NullRHI.h b/Engine/Source/Runtime/NullDrv/Public/NullRHI.h
--- a/Engine/Source/Runtime/NullDrv/Public/NullRHI.h
+++ b/Engine/Source/Runtime/NullDrv/Public/NullRHI.h
@@ -213,7 +213,7 @@ public:
 	{ 
 		return 0; 
 	}
-	virtual FTexture2DRHIRef RHIAsyncReallocateTexture2D(FRHITexture2D* Texture2D, int32 NewMipCount, int32 NewSizeX, int32 NewSizeY, FThreadSafeCounter* RequestStatus) final override
+	virtual FTextureRHIRef RHIAsyncReallocateTexture2D(FRHITexture* Texture2D, int32 NewMipCount, int32 NewSizeX, int32 NewSizeY, FThreadSafeCounter* RequestStatus) final override
 	{
 		const FRHITextureCreateDesc Desc =
 			FRHITextureCreateDesc::Create2D(TEXT("FNullDynamicRHI::RHIAsyncReallocateTexture2D"), NewSizeX, NewSizeY, Texture2D->GetFormat())
@@ -224,47 +224,47 @@ public:
 
 		return new FNullTexture(Desc);
 	}
-	virtual ETextureReallocationStatus RHIFinalizeAsyncReallocateTexture2D(FRHITexture2D* Texture2D, bool bBlockUntilCompleted) final override
+	virtual ETextureReallocationStatus RHIFinalizeAsyncReallocateTexture2D(FRHITexture* Texture2D, bool bBlockUntilCompleted) final override
 	{ 
 		return TexRealloc_Succeeded; 
 	}
-	virtual ETextureReallocationStatus RHICancelAsyncReallocateTexture2D(FRHITexture2D* Texture2D, bool bBlockUntilCompleted) final override
+	virtual ETextureReallocationStatus RHICancelAsyncReallocateTexture2D(FRHITexture* Texture2D, bool bBlockUntilCompleted) final override
 	{ 
 		return TexRealloc_Succeeded; 
 	}
-	virtual void* RHILockTexture2D(FRHITexture2D* Texture, uint32 MipIndex, EResourceLockMode LockMode, uint32& DestStride, bool bLockWithinMiptail, uint64* OutLockedByteCount) final override
+	virtual void* RHILockTexture2D(FRHITexture* Texture, uint32 MipIndex, EResourceLockMode LockMode, uint32& DestStride, bool bLockWithinMiptail, uint64* OutLockedByteCount) final override
 	{ 
 		DestStride = 0; 
 		return GetStaticTextureBuffer(Texture->GetSizeX(), Texture->GetSizeY(), Texture->GetFormat(), DestStride, OutLockedByteCount);
 	}
-	virtual void RHIUnlockTexture2D(FRHITexture2D* Texture, uint32 MipIndex, bool bLockWithinMiptail) final override
+	virtual void RHIUnlockTexture2D(FRHITexture* Texture, uint32 MipIndex, bool bLockWithinMiptail) final override
 	{
 
 	}
-	virtual void* RHILockTexture2DArray(FRHITexture2DArray* Texture, uint32 TextureIndex, uint32 MipIndex, EResourceLockMode LockMode, uint32& DestStride, bool bLockWithinMiptail) final override
+	virtual void* RHILockTexture2DArray(FRHITexture* Texture, uint32 TextureIndex, uint32 MipIndex, EResourceLockMode LockMode, uint32& DestStride, bool bLockWithinMiptail) final override
 	{ 
 		DestStride = 0; 
 		return GetStaticTextureBuffer(Texture->GetSizeX(), Texture->GetSizeY(), Texture->GetFormat(), DestStride);
 	}
-	virtual void RHIUnlockTexture2DArray(FRHITexture2DArray* Texture, uint32 TextureIndex, uint32 MipIndex, bool bLockWithinMiptail) final override
+	virtual void RHIUnlockTexture2DArray(FRHITexture* Texture, uint32 TextureIndex, uint32 MipIndex, bool bLockWithinMiptail) final override
 	{
 
 	}
-	virtual void RHIUpdateTexture2D(FRHICommandListBase&, FRHITexture2D* Texture, uint32 MipIndex, const struct FUpdateTextureRegion2D& UpdateRegion, uint32 SourcePitch, const uint8* SourceData) final override
+	virtual void RHIUpdateTexture2D(FRHICommandListBase&, FRHITexture* Texture, uint32 MipIndex, const struct FUpdateTextureRegion2D& UpdateRegion, uint32 SourcePitch, const uint8* SourceData) final override
 	{
 
 	}
-	virtual void RHIUpdateTexture3D(FRHICommandListBase&, FRHITexture3D* Texture, uint32 MipIndex, const struct FUpdateTextureRegion3D& UpdateRegion, uint32 SourceRowPitch, uint32 SourceDepthPitch, const uint8* SourceData) final override
+	virtual void RHIUpdateTexture3D(FRHICommandListBase&, FRHITexture* Texture, uint32 MipIndex, const struct FUpdateTextureRegion3D& UpdateRegion, uint32 SourceRowPitch, uint32 SourceDepthPitch, const uint8* SourceData) final override
 	{
 
 	}
 
-	virtual void* RHILockTextureCubeFace(FRHITextureCube* Texture, uint32 FaceIndex, uint32 ArrayIndex, uint32 MipIndex, EResourceLockMode LockMode, uint32& DestStride, bool bLockWithinMiptail) final override
+	virtual void* RHILockTextureCubeFace(FRHITexture* Texture, uint32 FaceIndex, uint32 ArrayIndex, uint32 MipIndex, EResourceLockMode LockMode, uint32& DestStride, bool bLockWithinMiptail) final override
 	{ 
 		DestStride = 0; 
 		return GetStaticTextureBuffer(Texture->GetSize(), Texture->GetSize(), Texture->GetFormat(), DestStride);
 	}
-	virtual void RHIUnlockTextureCubeFace(FRHITextureCube* Texture, uint32 FaceIndex, uint32 ArrayIndex, uint32 MipIndex, bool bLockWithinMiptail) final override
+	virtual void RHIUnlockTextureCubeFace(FRHITexture* Texture, uint32 FaceIndex, uint32 ArrayIndex, uint32 MipIndex, bool bLockWithinMiptail) final override
 	{
 	}
 
@@ -336,7 +336,7 @@ public:
 	{
 	}
 
-	virtual FTexture2DRHIRef RHIGetViewportBackBuffer(FRHIViewport* Viewport) final override
+	virtual FTextureRHIRef RHIGetViewportBackBuffer(FRHIViewport* Viewport) final override
 	{
 		const FRHITextureCreateDesc Desc =
 			FRHITextureCreateDesc::Create2D(TEXT("FNullDynamicRHI::RHIGetViewportBackBuffer"), 1, 1, PF_B8G8R8A8)
diff --git a/Engine/Source/Runtime/OpenGLDrv/Private/OpenGLDrv.h b/Engine/Source/Runtime/OpenGLDrv/Private/OpenGLDrv.h
--- a/Engine/Source/Runtime/OpenGLDrv/Private/OpenGLDrv.h
+++ b/Engine/Source/Runtime/OpenGLDrv/Private/OpenGLDrv.h
@@ -288,9 +288,9 @@ public:
 	virtual int32 RHIGetGLMajorVersion() const final override;
 	virtual int32 RHIGetGLMinorVersion() const final override;
 	virtual bool RHISupportsFramebufferSRGBEnable() const final override;
-	virtual FTexture2DRHIRef RHICreateTexture2DFromResource(EPixelFormat Format, uint32 SizeX, uint32 SizeY, uint32 NumMips, uint32 NumSamples, uint32 NumSamplesTileMem, const FClearValueBinding& ClearValueBinding, GLuint Resource, ETextureCreateFlags Flags) final override;
-	virtual FTexture2DRHIRef RHICreateTexture2DArrayFromResource(EPixelFormat Format, uint32 SizeX, uint32 SizeY, uint32 ArraySize, uint32 NumMips, uint32 NumSamples, uint32 NumSamplesTileMem, const FClearValueBinding& ClearValueBinding, GLuint Resource, ETextureCreateFlags Flags) final override;
-	virtual FTextureCubeRHIRef RHICreateTextureCubeFromResource(EPixelFormat Format, uint32 Size, bool bArray, uint32 ArraySize, uint32 NumMips, uint32 NumSamples, uint32 NumSamplesTileMem, const FClearValueBinding& ClearValueBinding, GLuint Resource, ETextureCreateFlags Flags) final override;
+	virtual FTextureRHIRef RHICreateTexture2DFromResource(EPixelFormat Format, uint32 SizeX, uint32 SizeY, uint32 NumMips, uint32 NumSamples, uint32 NumSamplesTileMem, const FClearValueBinding& ClearValueBinding, GLuint Resource, ETextureCreateFlags Flags) final override;
+	virtual FTextureRHIRef RHICreateTexture2DArrayFromResource(EPixelFormat Format, uint32 SizeX, uint32 SizeY, uint32 ArraySize, uint32 NumMips, uint32 NumSamples, uint32 NumSamplesTileMem, const FClearValueBinding& ClearValueBinding, GLuint Resource, ETextureCreateFlags Flags) final override;
+	virtual FTextureRHIRef RHICreateTextureCubeFromResource(EPixelFormat Format, uint32 Size, bool bArray, uint32 ArraySize, uint32 NumMips, uint32 NumSamples, uint32 NumSamplesTileMem, const FClearValueBinding& ClearValueBinding, GLuint Resource, ETextureCreateFlags Flags) final override;
 	virtual GLuint RHIGetResource(FRHITexture* InTexture) const final override;
 	virtual bool RHIIsValidTexture(GLuint InTexture) const final override;
 	virtual void RHISetExternalGPUTime(uint32 InExternalGPUTime) final override;
@@ -365,17 +365,17 @@ public:
 	virtual FTextureRHIRef RHIAsyncCreateTexture2D(uint32 SizeX, uint32 SizeY, uint8 Format, uint32 NumMips, ETextureCreateFlags Flags, ERHIAccess InResourceState, void** InitialMipData, uint32 NumInitialMips, const TCHAR* DebugName, FGraphEventRef& OutCompletionEvent) final override;
 	virtual void RHIGenerateMips(FRHITexture* Texture) final override;
 	virtual uint32 RHIComputeMemorySize(FRHITexture* TextureRHI) final override;
-	virtual FTexture2DRHIRef RHIAsyncReallocateTexture2D(FRHITexture2D* Texture2D, int32 NewMipCount, int32 NewSizeX, int32 NewSizeY, FThreadSafeCounter* RequestStatus) final override;
-	virtual ETextureReallocationStatus RHIFinalizeAsyncReallocateTexture2D(FRHITexture2D* Texture2D, bool bBlockUntilCompleted) final override;
-	virtual ETextureReallocationStatus RHICancelAsyncReallocateTexture2D(FRHITexture2D* Texture2D, bool bBlockUntilCompleted) final override;
-	virtual void* RHILockTexture2D(FRHITexture2D* Texture, uint32 MipIndex, EResourceLockMode LockMode, uint32& DestStride, bool bLockWithinMiptail, uint64* OutLockedByteCount = nullptr) final override;
-	virtual void RHIUnlockTexture2D(FRHITexture2D* Texture, uint32 MipIndex, bool bLockWithinMiptail) final override;
-	virtual void* RHILockTexture2DArray(FRHITexture2DArray* Texture, uint32 TextureIndex, uint32 MipIndex, EResourceLockMode LockMode, uint32& DestStride, bool bLockWithinMiptail) final override;
-	virtual void RHIUnlockTexture2DArray(FRHITexture2DArray* Texture, uint32 TextureIndex, uint32 MipIndex, bool bLockWithinMiptail) final override;
-	virtual void RHIUpdateTexture2D(FRHICommandListBase& RHICmdList, FRHITexture2D* Texture, uint32 MipIndex, const struct FUpdateTextureRegion2D& UpdateRegion, uint32 SourcePitch, const uint8* SourceData) final override;
-	virtual void RHIUpdateTexture3D(FRHICommandListBase& RHICmdList, FRHITexture3D* Texture, uint32 MipIndex, const struct FUpdateTextureRegion3D& UpdateRegion, uint32 SourceRowPitch, uint32 SourceDepthPitch, const uint8* SourceData) final override;
-	virtual void* RHILockTextureCubeFace(FRHITextureCube* Texture, uint32 FaceIndex, uint32 ArrayIndex, uint32 MipIndex, EResourceLockMode LockMode, uint32& DestStride, bool bLockWithinMiptail) final override;
-	virtual void RHIUnlockTextureCubeFace(FRHITextureCube* Texture, uint32 FaceIndex, uint32 ArrayIndex, uint32 MipIndex, bool bLockWithinMiptail) final override;
+	virtual FTextureRHIRef RHIAsyncReallocateTexture2D(FRHITexture* Texture2D, int32 NewMipCount, int32 NewSizeX, int32 NewSizeY, FThreadSafeCounter* RequestStatus) final override;
+	virtual ETextureReallocationStatus RHIFinalizeAsyncReallocateTexture2D(FRHITexture* Texture2D, bool bBlockUntilCompleted) final override;
+	virtual ETextureReallocationStatus RHICancelAsyncReallocateTexture2D(FRHITexture* Texture2D, bool bBlockUntilCompleted) final override;
+	virtual void* RHILockTexture2D(FRHITexture* Texture, uint32 MipIndex, EResourceLockMode LockMode, uint32& DestStride, bool bLockWithinMiptail, uint64* OutLockedByteCount = nullptr) final override;
+	virtual void RHIUnlockTexture2D(FRHITexture* Texture, uint32 MipIndex, bool bLockWithinMiptail) final override;
+	virtual void* RHILockTexture2DArray(FRHITexture* Texture, uint32 TextureIndex, uint32 MipIndex, EResourceLockMode LockMode, uint32& DestStride, bool bLockWithinMiptail) final override;
+	virtual void RHIUnlockTexture2DArray(FRHITexture* Texture, uint32 TextureIndex, uint32 MipIndex, bool bLockWithinMiptail) final override;
+	virtual void RHIUpdateTexture2D(FRHICommandListBase& RHICmdList, FRHITexture* Texture, uint32 MipIndex, const struct FUpdateTextureRegion2D& UpdateRegion, uint32 SourcePitch, const uint8* SourceData) final override;
+	virtual void RHIUpdateTexture3D(FRHICommandListBase& RHICmdList, FRHITexture* Texture, uint32 MipIndex, const struct FUpdateTextureRegion3D& UpdateRegion, uint32 SourceRowPitch, uint32 SourceDepthPitch, const uint8* SourceData) final override;
+	virtual void* RHILockTextureCubeFace(FRHITexture* Texture, uint32 FaceIndex, uint32 ArrayIndex, uint32 MipIndex, EResourceLockMode LockMode, uint32& DestStride, bool bLockWithinMiptail) final override;
+	virtual void RHIUnlockTextureCubeFace(FRHITexture* Texture, uint32 FaceIndex, uint32 ArrayIndex, uint32 MipIndex, bool bLockWithinMiptail) final override;
 	virtual void RHIBindDebugLabelName(FRHICommandListBase& RHICmdList, FRHITexture* Texture, const TCHAR* Name) final override;
 	virtual void RHIReadSurfaceData(FRHITexture* Texture,FIntRect Rect,TArray<FColor>& OutData,FReadSurfaceDataFlags InFlags) final override;
 	virtual void RHIReadSurfaceData(FRHITexture* Texture, FIntRect Rect, TArray<FLinearColor>& OutData, FReadSurfaceDataFlags InFlags) final override;
@@ -385,7 +385,7 @@ public:
 	virtual void RHIRead3DSurfaceFloatData(FRHITexture* Texture,FIntRect Rect,FIntPoint ZMinMax,TArray<FFloat16Color>& OutData) final override;
 	virtual FRenderQueryRHIRef RHICreateRenderQuery(ERenderQueryType QueryType) final override;
 	virtual bool RHIGetRenderQueryResult(FRHIRenderQuery* RenderQuery, uint64& OutResult, bool bWait, uint32 GPUIndex = INDEX_NONE) final override;
-	virtual FTexture2DRHIRef RHIGetViewportBackBuffer(FRHIViewport* Viewport) final override;
+	virtual FTextureRHIRef RHIGetViewportBackBuffer(FRHIViewport* Viewport) final override;
 	virtual void RHIAliasTextureResources(FTextureRHIRef& DestTextureRHI, FTextureRHIRef& SrcTextureRHI) final override;
 	virtual FTextureRHIRef RHICreateAliasedTexture(FTextureRHIRef& SourceTexture) final override;
 	virtual void RHIAdvanceFrameForGetViewportBackBuffer(FRHIViewport* Viewport) final override;
@@ -730,21 +730,21 @@ public:
 		}
 	};
 
-	virtual void* LockTexture2D_RenderThread(class FRHICommandListImmediate& RHICmdList, FRHITexture2D* Texture, uint32 MipIndex, EResourceLockMode LockMode, uint32& DestStride, bool bLockWithinMiptail, bool bNeedsDefaultRHIFlush, uint64* OutLockedByteCount = nullptr) final override;
-	virtual void UnlockTexture2D_RenderThread(class FRHICommandListImmediate& RHICmdList, FRHITexture2D* Texture, uint32 MipIndex, bool bLockWithinMiptail, bool bNeedsDefaultRHIFlush) final override;
-	virtual void* RHILockTextureCubeFace_RenderThread(class FRHICommandListImmediate& RHICmdList, FRHITextureCube* Texture, uint32 FaceIndex, uint32 ArrayIndex, uint32 MipIndex, EResourceLockMode LockMode, uint32& DestStride, bool bLockWithinMiptail) final override;
-	virtual void RHIUnlockTextureCubeFace_RenderThread(class FRHICommandListImmediate& RHICmdList, FRHITextureCube* Texture, uint32 FaceIndex, uint32 ArrayIndex, uint32 MipIndex, bool bLockWithinMiptail) final override;
-	virtual void* LockTexture2DArray_RenderThread(class FRHICommandListImmediate& RHICmdList, FRHITexture2DArray* Texture, uint32 ArrayIndex, uint32 MipIndex, EResourceLockMode LockMode, uint32& DestStride, bool bLockWithinMiptail) final override;
-	virtual void UnlockTexture2DArray_RenderThread(class FRHICommandListImmediate& RHICmdList, FRHITexture2DArray* Texture, uint32 ArrayIndex, uint32 MipIndex, bool bLockWithinMiptail) final override;
+	virtual void* LockTexture2D_RenderThread(class FRHICommandListImmediate& RHICmdList, FRHITexture* Texture, uint32 MipIndex, EResourceLockMode LockMode, uint32& DestStride, bool bLockWithinMiptail, bool bNeedsDefaultRHIFlush, uint64* OutLockedByteCount = nullptr) final override;
+	virtual void UnlockTexture2D_RenderThread(class FRHICommandListImmediate& RHICmdList, FRHITexture* Texture, uint32 MipIndex, bool bLockWithinMiptail, bool bNeedsDefaultRHIFlush) final override;
+	virtual void* RHILockTextureCubeFace_RenderThread(class FRHICommandListImmediate& RHICmdList, FRHITexture* Texture, uint32 FaceIndex, uint32 ArrayIndex, uint32 MipIndex, EResourceLockMode LockMode, uint32& DestStride, bool bLockWithinMiptail) final override;
+	virtual void RHIUnlockTextureCubeFace_RenderThread(class FRHICommandListImmediate& RHICmdList, FRHITexture* Texture, uint32 FaceIndex, uint32 ArrayIndex, uint32 MipIndex, bool bLockWithinMiptail) final override;
+	virtual void* LockTexture2DArray_RenderThread(class FRHICommandListImmediate& RHICmdList, FRHITexture* Texture, uint32 ArrayIndex, uint32 MipIndex, EResourceLockMode LockMode, uint32& DestStride, bool bLockWithinMiptail) final override;
+	virtual void UnlockTexture2DArray_RenderThread(class FRHICommandListImmediate& RHICmdList, FRHITexture* Texture, uint32 ArrayIndex, uint32 MipIndex, bool bLockWithinMiptail) final override;
 
-	virtual FTexture2DRHIRef AsyncReallocateTexture2D_RenderThread(class FRHICommandListImmediate& RHICmdList, FRHITexture2D* Texture2D, int32 NewMipCount, int32 NewSizeX, int32 NewSizeY, FThreadSafeCounter* RequestStatus) final override;
+	virtual FTextureRHIRef AsyncReallocateTexture2D_RenderThread(class FRHICommandListImmediate& RHICmdList, FRHITexture* Texture2D, int32 NewMipCount, int32 NewSizeX, int32 NewSizeY, FThreadSafeCounter* RequestStatus) final override;
 
-	virtual ETextureReallocationStatus FinalizeAsyncReallocateTexture2D_RenderThread(class FRHICommandListImmediate& RHICmdList, FRHITexture2D* Texture2D, bool bBlockUntilCompleted) final override
+	virtual ETextureReallocationStatus FinalizeAsyncReallocateTexture2D_RenderThread(class FRHICommandListImmediate& RHICmdList, FRHITexture* Texture2D, bool bBlockUntilCompleted) final override
 	{
 		return this->RHIFinalizeAsyncReallocateTexture2D(Texture2D, bBlockUntilCompleted);
 	}
 
-	virtual ETextureReallocationStatus CancelAsyncReallocateTexture2D_RenderThread(class FRHICommandListImmediate& RHICmdList, FRHITexture2D* Texture2D, bool bBlockUntilCompleted) final override
+	virtual ETextureReallocationStatus CancelAsyncReallocateTexture2D_RenderThread(class FRHICommandListImmediate& RHICmdList, FRHITexture* Texture2D, bool bBlockUntilCompleted) final override
 	{
 		return this->RHICancelAsyncReallocateTexture2D(Texture2D, bBlockUntilCompleted);
 	}
diff --git a/Engine/Source/Runtime/OpenGLDrv/Private/OpenGLTexture.cpp b/Engine/Source/Runtime/OpenGLDrv/Private/OpenGLTexture.cpp
--- a/Engine/Source/Runtime/OpenGLDrv/Private/OpenGLTexture.cpp
+++ b/Engine/Source/Runtime/OpenGLDrv/Private/OpenGLTexture.cpp
@@ -1439,7 +1439,7 @@ FTextureRHIRef FOpenGLDynamicRHI::RHICreateTexture(FRHICommandListBase& RHICmdLi
 FTextureRHIRef FOpenGLDynamicRHI::RHIAsyncCreateTexture2D(uint32 SizeX, uint32 SizeY, uint8 Format, uint32 NumMips, ETextureCreateFlags Flags, ERHIAccess InResourceState, void** InitialMipData, uint32 NumInitialMips, const TCHAR* DebugName, FGraphEventRef& OutCompletionEvent)
 {
 	check(0);
-	return FTexture2DRHIRef();
+	return FTextureRHIRef();
 }
 
 /** Generates mip maps for the surface. */
@@ -1485,12 +1485,12 @@ uint32 FOpenGLDynamicRHI::RHIComputeMemorySize(FRHITexture* TextureRHI)
 	return Texture->MemorySize;
 }
 
-FTexture2DRHIRef FOpenGLDynamicRHI::AsyncReallocateTexture2D_RenderThread(class FRHICommandListImmediate& RHICmdList, FRHITexture2D* Texture2DRHI, int32 NewMipCount, int32 NewSizeX, int32 NewSizeY, FThreadSafeCounter* RequestStatus)
+FTextureRHIRef FOpenGLDynamicRHI::AsyncReallocateTexture2D_RenderThread(class FRHICommandListImmediate& RHICmdList, FRHITexture* Texture2DRHI, int32 NewMipCount, int32 NewSizeX, int32 NewSizeY, FThreadSafeCounter* RequestStatus)
 {
 	return this->RHIAsyncReallocateTexture2D(Texture2DRHI, NewMipCount, NewSizeX, NewSizeY, RequestStatus);
 }
 
-FTexture2DRHIRef FOpenGLDynamicRHI::RHIAsyncReallocateTexture2D(FRHITexture2D* Texture2DRHI, int32 NewMipCount, int32 NewSizeX, int32 NewSizeY, FThreadSafeCounter* RequestStatus)
+FTextureRHIRef FOpenGLDynamicRHI::RHIAsyncReallocateTexture2D(FRHITexture* Texture2DRHI, int32 NewMipCount, int32 NewSizeX, int32 NewSizeY, FThreadSafeCounter* RequestStatus)
 {
 	FRHICommandListImmediate& RHICmdList = FRHICommandListImmediate::Get();
 
@@ -1536,7 +1536,7 @@ FTexture2DRHIRef FOpenGLDynamicRHI::RHIAsyncReallocateTexture2D(FRHITexture2D* T
  * @param Texture2D		- Texture to check the reallocation status for
  * @return				- Current reallocation status
  */
-ETextureReallocationStatus FOpenGLDynamicRHI::RHIFinalizeAsyncReallocateTexture2D(FRHITexture2D* Texture2D, bool bBlockUntilCompleted )
+ETextureReallocationStatus FOpenGLDynamicRHI::RHIFinalizeAsyncReallocateTexture2D(FRHITexture* Texture2D, bool bBlockUntilCompleted )
 {
 	return TexRealloc_Succeeded;
 }
@@ -1549,12 +1549,12 @@ ETextureReallocationStatus FOpenGLDynamicRHI::RHIFinalizeAsyncReallocateTexture2
  * @param bBlockUntilCompleted	If true, blocks until the cancellation is fully completed
  * @return						Reallocation status
  */
-ETextureReallocationStatus FOpenGLDynamicRHI::RHICancelAsyncReallocateTexture2D(FRHITexture2D* Texture2D, bool bBlockUntilCompleted )
+ETextureReallocationStatus FOpenGLDynamicRHI::RHICancelAsyncReallocateTexture2D(FRHITexture* Texture2D, bool bBlockUntilCompleted )
 {
 	return TexRealloc_Succeeded;
 }
 
-void* FOpenGLDynamicRHI::RHILockTexture2D(FRHITexture2D* TextureRHI,uint32 MipIndex,EResourceLockMode LockMode,uint32& DestStride,bool bLockWithinMiptail, uint64* OutLockedByteCount)
+void* FOpenGLDynamicRHI::RHILockTexture2D(FRHITexture* TextureRHI,uint32 MipIndex,EResourceLockMode LockMode,uint32& DestStride,bool bLockWithinMiptail, uint64* OutLockedByteCount)
 {
 	FOpenGLTexture* Texture = ResourceCast(TextureRHI);
 	if (OutLockedByteCount)
@@ -1565,22 +1565,22 @@ void* FOpenGLDynamicRHI::RHILockTexture2D(FRHITexture2D* TextureRHI,uint32 MipIn
 	return Texture->Lock(MipIndex, 0, LockMode, DestStride);
 }
 
-void FOpenGLDynamicRHI::RHIUnlockTexture2D(FRHITexture2D* TextureRHI, uint32 MipIndex, bool bLockWithinMiptail)
+void FOpenGLDynamicRHI::RHIUnlockTexture2D(FRHITexture* TextureRHI, uint32 MipIndex, bool bLockWithinMiptail)
 {
 	ResourceCast(TextureRHI)->Unlock(MipIndex, 0);
 }
 
-void* FOpenGLDynamicRHI::RHILockTexture2DArray(FRHITexture2DArray* TextureRHI, uint32 TextureIndex, uint32 MipIndex, EResourceLockMode LockMode, uint32& DestStride, bool bLockWithinMiptail)
+void* FOpenGLDynamicRHI::RHILockTexture2DArray(FRHITexture* TextureRHI, uint32 TextureIndex, uint32 MipIndex, EResourceLockMode LockMode, uint32& DestStride, bool bLockWithinMiptail)
 {
 	return ResourceCast(TextureRHI)->Lock(MipIndex, TextureIndex, LockMode, DestStride);
 }
 
-void FOpenGLDynamicRHI::RHIUnlockTexture2DArray(FRHITexture2DArray* TextureRHI, uint32 TextureIndex, uint32 MipIndex, bool bLockWithinMiptail)
+void FOpenGLDynamicRHI::RHIUnlockTexture2DArray(FRHITexture* TextureRHI, uint32 TextureIndex, uint32 MipIndex, bool bLockWithinMiptail)
 {
 	ResourceCast(TextureRHI)->Unlock(MipIndex, TextureIndex);
 }
 
-void FOpenGLDynamicRHI::RHIUpdateTexture2D(FRHICommandListBase& RHICmdList, FRHITexture2D* TextureRHI, uint32 MipIndex, const FUpdateTextureRegion2D& UpdateRegion, uint32 SourcePitch, const uint8* SourceData)
+void FOpenGLDynamicRHI::RHIUpdateTexture2D(FRHICommandListBase& RHICmdList, FRHITexture* TextureRHI, uint32 MipIndex, const FUpdateTextureRegion2D& UpdateRegion, uint32 SourcePitch, const uint8* SourceData)
 {
 	const FPixelFormatInfo& FormatInfo = GPixelFormats[TextureRHI->GetFormat()];
 
@@ -1677,7 +1677,7 @@ void FOpenGLDynamicRHI::RHIUpdateTexture2D(FRHICommandListBase& RHICmdList, FRHI
 	});
 }
 
-void FOpenGLDynamicRHI::RHIUpdateTexture3D(FRHICommandListBase& RHICmdList, FRHITexture3D* TextureRHI, uint32 MipIndex, const FUpdateTextureRegion3D& UpdateRegion, uint32 SourceRowPitch, uint32 SourceDepthPitch, const uint8* SourceData)
+void FOpenGLDynamicRHI::RHIUpdateTexture3D(FRHICommandListBase& RHICmdList, FRHITexture* TextureRHI, uint32 MipIndex, const FUpdateTextureRegion3D& UpdateRegion, uint32 SourceRowPitch, uint32 SourceDepthPitch, const uint8* SourceData)
 {
 	uint8* RHITSourceData = nullptr;
 	if (!ShouldRunGLRenderContextOpOnThisThread(RHICmdList))
@@ -1805,12 +1805,12 @@ void FOpenGLDynamicRHI::InvalidateUAVResourceInCache(GLuint Resource)
 /*-----------------------------------------------------------------------------
 	Cubemap texture support.
 -----------------------------------------------------------------------------*/
-void* FOpenGLDynamicRHI::RHILockTextureCubeFace(FRHITextureCube* TextureCubeRHI, uint32 FaceIndex, uint32 ArrayIndex, uint32 MipIndex, EResourceLockMode LockMode, uint32& DestStride, bool bLockWithinMiptail)
+void* FOpenGLDynamicRHI::RHILockTextureCubeFace(FRHITexture* TextureCubeRHI, uint32 FaceIndex, uint32 ArrayIndex, uint32 MipIndex, EResourceLockMode LockMode, uint32& DestStride, bool bLockWithinMiptail)
 {
 	return ResourceCast(TextureCubeRHI)->Lock(MipIndex, FaceIndex + 6 * ArrayIndex, LockMode, DestStride);
 }
 
-void FOpenGLDynamicRHI::RHIUnlockTextureCubeFace(FRHITextureCube* TextureCubeRHI, uint32 FaceIndex, uint32 ArrayIndex, uint32 MipIndex, bool bLockWithinMiptail)
+void FOpenGLDynamicRHI::RHIUnlockTextureCubeFace(FRHITexture* TextureCubeRHI, uint32 FaceIndex, uint32 ArrayIndex, uint32 MipIndex, bool bLockWithinMiptail)
 {
 	ResourceCast(TextureCubeRHI)->Unlock(MipIndex, FaceIndex + ArrayIndex * 6);
 }
@@ -1925,7 +1925,7 @@ void FOpenGLDynamicRHI::RHICopyTexture(FRHITexture* SourceTextureRHI, FRHITextur
 }
 
 
-FTexture2DRHIRef FOpenGLDynamicRHI::RHICreateTexture2DFromResource(EPixelFormat Format, uint32 SizeX, uint32 SizeY, uint32 NumMips, uint32 NumSamples, uint32 NumSamplesTileMem, const FClearValueBinding& ClearValueBinding, GLuint Resource, ETextureCreateFlags TexCreateFlags)
+FTextureRHIRef FOpenGLDynamicRHI::RHICreateTexture2DFromResource(EPixelFormat Format, uint32 SizeX, uint32 SizeY, uint32 NumMips, uint32 NumSamples, uint32 NumSamplesTileMem, const FClearValueBinding& ClearValueBinding, GLuint Resource, ETextureCreateFlags TexCreateFlags)
 {
 	const FRHITextureCreateDesc Desc =
 		FRHITextureCreateDesc::Create2D(TEXT("RHICreateTexture2DFromResource"), SizeX, SizeY, Format)
@@ -1938,7 +1938,7 @@ FTexture2DRHIRef FOpenGLDynamicRHI::RHICreateTexture2DFromResource(EPixelFormat
 	return new FOpenGLTexture(Desc, Resource);
 }
 
-FTexture2DRHIRef FOpenGLDynamicRHI::RHICreateTexture2DArrayFromResource(EPixelFormat Format, uint32 SizeX, uint32 SizeY, uint32 ArraySize, uint32 NumMips, uint32 NumSamples, uint32 NumSamplesTileMem, const FClearValueBinding& ClearValueBinding, GLuint Resource, ETextureCreateFlags TexCreateFlags)
+FTextureRHIRef FOpenGLDynamicRHI::RHICreateTexture2DArrayFromResource(EPixelFormat Format, uint32 SizeX, uint32 SizeY, uint32 ArraySize, uint32 NumMips, uint32 NumSamples, uint32 NumSamplesTileMem, const FClearValueBinding& ClearValueBinding, GLuint Resource, ETextureCreateFlags TexCreateFlags)
 {
 	const FRHITextureCreateDesc Desc =
 		FRHITextureCreateDesc::Create2DArray(TEXT("RHICreateTexture2DArrayFromResource"), SizeX, SizeY, ArraySize, Format)
@@ -1951,7 +1951,7 @@ FTexture2DRHIRef FOpenGLDynamicRHI::RHICreateTexture2DArrayFromResource(EPixelFo
 	return new FOpenGLTexture(Desc, Resource);
 }
 
-FTextureCubeRHIRef FOpenGLDynamicRHI::RHICreateTextureCubeFromResource(EPixelFormat Format, uint32 Size, bool bArray, uint32 ArraySize, uint32 NumMips, uint32 NumSamples, uint32 NumSamplesTileMem, const FClearValueBinding& ClearValueBinding, GLuint Resource, ETextureCreateFlags TexCreateFlags)
+FTextureRHIRef FOpenGLDynamicRHI::RHICreateTextureCubeFromResource(EPixelFormat Format, uint32 Size, bool bArray, uint32 ArraySize, uint32 NumMips, uint32 NumSamples, uint32 NumSamplesTileMem, const FClearValueBinding& ClearValueBinding, GLuint Resource, ETextureCreateFlags TexCreateFlags)
 {
 	const FRHITextureCreateDesc Desc =
 		FRHITextureCreateDesc::Create(TEXT("RHICreateTextureCubeFromResource"), bArray ? ETextureDimension::TextureCube : ETextureDimension::TextureCubeArray)
@@ -1985,7 +1985,7 @@ FTextureRHIRef FOpenGLDynamicRHI::RHICreateAliasedTexture(FTextureRHIRef& Source
 	return new FOpenGLTexture(*ResourceCast(SourceTexture), *Name, FOpenGLTexture::AliasResource);
 }
 
-void* FOpenGLDynamicRHI::LockTexture2D_RenderThread(class FRHICommandListImmediate& RHICmdList, FRHITexture2D* Texture, uint32 MipIndex, EResourceLockMode LockMode, uint32& DestStride, bool bLockWithinMiptail, bool bNeedsDefaultRHIFlush, uint64* OutLockedByteCount)
+void* FOpenGLDynamicRHI::LockTexture2D_RenderThread(class FRHICommandListImmediate& RHICmdList, FRHITexture* Texture, uint32 MipIndex, EResourceLockMode LockMode, uint32& DestStride, bool bLockWithinMiptail, bool bNeedsDefaultRHIFlush, uint64* OutLockedByteCount)
 {
 	check(IsInRenderingThread());
 	static auto* CVarRHICmdBufferWriteLocks = IConsoleManager::Get().FindTConsoleVariableDataInt(TEXT("r.RHICmdBufferWriteLocks"));
@@ -2016,7 +2016,7 @@ void* FOpenGLDynamicRHI::LockTexture2D_RenderThread(class FRHICommandListImmedia
 	return Result;
 }
 
-void FOpenGLDynamicRHI::UnlockTexture2D_RenderThread(class FRHICommandListImmediate& RHICmdList, FRHITexture2D* Texture, uint32 MipIndex, bool bLockWithinMiptail, bool bNeedsDefaultRHIFlush)
+void FOpenGLDynamicRHI::UnlockTexture2D_RenderThread(class FRHICommandListImmediate& RHICmdList, FRHITexture* Texture, uint32 MipIndex, bool bLockWithinMiptail, bool bNeedsDefaultRHIFlush)
 {
 	check(IsInRenderingThread());
 	static auto* CVarRHICmdBufferWriteLocks = IConsoleManager::Get().FindTConsoleVariableDataInt(TEXT("r.RHICmdBufferWriteLocks"));
@@ -2048,7 +2048,7 @@ void FOpenGLDynamicRHI::UnlockTexture2D_RenderThread(class FRHICommandListImmedi
 	}
 }
 
-void* FOpenGLDynamicRHI::RHILockTextureCubeFace_RenderThread(class FRHICommandListImmediate& RHICmdList, FRHITextureCube* Texture, uint32 FaceIndex, uint32 ArrayIndex, uint32 MipIndex, EResourceLockMode LockMode, uint32& DestStride, bool bLockWithinMiptail)
+void* FOpenGLDynamicRHI::RHILockTextureCubeFace_RenderThread(class FRHICommandListImmediate& RHICmdList, FRHITexture* Texture, uint32 FaceIndex, uint32 ArrayIndex, uint32 MipIndex, EResourceLockMode LockMode, uint32& DestStride, bool bLockWithinMiptail)
 {
 	check(IsInRenderingThread());
 	static auto* CVarRHICmdBufferWriteLocks = IConsoleManager::Get().FindTConsoleVariableDataInt(TEXT("r.RHICmdBufferWriteLocks"));
@@ -2073,7 +2073,7 @@ void* FOpenGLDynamicRHI::RHILockTextureCubeFace_RenderThread(class FRHICommandLi
 	return Result;
 }
 
-void FOpenGLDynamicRHI::RHIUnlockTextureCubeFace_RenderThread(class FRHICommandListImmediate& RHICmdList, FRHITextureCube* Texture, uint32 FaceIndex, uint32 ArrayIndex, uint32 MipIndex, bool bLockWithinMiptail)
+void FOpenGLDynamicRHI::RHIUnlockTextureCubeFace_RenderThread(class FRHICommandListImmediate& RHICmdList, FRHITexture* Texture, uint32 FaceIndex, uint32 ArrayIndex, uint32 MipIndex, bool bLockWithinMiptail)
 {
 	check(IsInRenderingThread());
 	static auto* CVarRHICmdBufferWriteLocks = IConsoleManager::Get().FindTConsoleVariableDataInt(TEXT("r.RHICmdBufferWriteLocks"));
@@ -2102,7 +2102,7 @@ void FOpenGLDynamicRHI::RHIUnlockTextureCubeFace_RenderThread(class FRHICommandL
 	}
 }
 
-void* FOpenGLDynamicRHI::LockTexture2DArray_RenderThread(class FRHICommandListImmediate& RHICmdList, FRHITexture2DArray* Texture, uint32 ArrayIndex, uint32 MipIndex, EResourceLockMode LockMode, uint32& DestStride, bool bLockWithinMiptail)
+void* FOpenGLDynamicRHI::LockTexture2DArray_RenderThread(class FRHICommandListImmediate& RHICmdList, FRHITexture* Texture, uint32 ArrayIndex, uint32 MipIndex, EResourceLockMode LockMode, uint32& DestStride, bool bLockWithinMiptail)
 {
 	check(IsInRenderingThread());
 	static auto* CVarRHICmdBufferWriteLocks = IConsoleManager::Get().FindTConsoleVariableDataInt(TEXT("r.RHICmdBufferWriteLocks"));
@@ -2128,7 +2128,7 @@ void* FOpenGLDynamicRHI::LockTexture2DArray_RenderThread(class FRHICommandListIm
 	return Result;
 }
 
-void FOpenGLDynamicRHI::UnlockTexture2DArray_RenderThread(class FRHICommandListImmediate& RHICmdList, FRHITexture2DArray* Texture, uint32 ArrayIndex, uint32 MipIndex, bool bLockWithinMiptail)
+void FOpenGLDynamicRHI::UnlockTexture2DArray_RenderThread(class FRHICommandListImmediate& RHICmdList, FRHITexture* Texture, uint32 ArrayIndex, uint32 MipIndex, bool bLockWithinMiptail)
 {
 	check(IsInRenderingThread());
 	static auto* CVarRHICmdBufferWriteLocks = IConsoleManager::Get().FindTConsoleVariableDataInt(TEXT("r.RHICmdBufferWriteLocks"));
diff --git a/Engine/Source/Runtime/OpenGLDrv/Private/OpenGLViewport.cpp b/Engine/Source/Runtime/OpenGLDrv/Private/OpenGLViewport.cpp
--- a/Engine/Source/Runtime/OpenGLDrv/Private/OpenGLViewport.cpp
+++ b/Engine/Source/Runtime/OpenGLDrv/Private/OpenGLViewport.cpp
@@ -224,7 +224,7 @@ void FOpenGLDynamicRHI::RHIEndDrawingViewport(FRHIViewport* ViewportRHI,bool bPr
 }
 
 
-FTexture2DRHIRef FOpenGLDynamicRHI::RHIGetViewportBackBuffer(FRHIViewport* ViewportRHI)
+FTextureRHIRef FOpenGLDynamicRHI::RHIGetViewportBackBuffer(FRHIViewport* ViewportRHI)
 {
 	FOpenGLViewport* Viewport = ResourceCast(ViewportRHI);
 	return Viewport->GetBackBuffer();
diff --git a/Engine/Source/Runtime/OpenGLDrv/Public/IOpenGLDynamicRHI.h b/Engine/Source/Runtime/OpenGLDrv/Public/IOpenGLDynamicRHI.h
--- a/Engine/Source/Runtime/OpenGLDrv/Public/IOpenGLDynamicRHI.h
+++ b/Engine/Source/Runtime/OpenGLDrv/Public/IOpenGLDynamicRHI.h
@@ -14,9 +14,9 @@ struct IOpenGLDynamicRHI : public FDynamicRHIPSOFallback
 
 	virtual bool RHISupportsFramebufferSRGBEnable() const = 0;
 
-	virtual FTexture2DRHIRef RHICreateTexture2DFromResource(EPixelFormat Format, uint32 SizeX, uint32 SizeY, uint32 NumMips, uint32 NumSamples, uint32 NumSamplesTileMem, const FClearValueBinding& ClearValueBinding, GLuint Resource, ETextureCreateFlags Flags) = 0;
-	virtual FTexture2DRHIRef RHICreateTexture2DArrayFromResource(EPixelFormat Format, uint32 SizeX, uint32 SizeY, uint32 ArraySize, uint32 NumMips, uint32 NumSamples, uint32 NumSamplesTileMem, const FClearValueBinding& ClearValueBinding, GLuint Resource, ETextureCreateFlags Flags) = 0;
-	virtual FTextureCubeRHIRef RHICreateTextureCubeFromResource(EPixelFormat Format, uint32 Size, bool bArray, uint32 ArraySize, uint32 NumMips, uint32 NumSamples, uint32 NumSamplesTileMem, const FClearValueBinding& ClearValueBinding, GLuint Resource, ETextureCreateFlags Flags) = 0;
+	virtual FTextureRHIRef RHICreateTexture2DFromResource(EPixelFormat Format, uint32 SizeX, uint32 SizeY, uint32 NumMips, uint32 NumSamples, uint32 NumSamplesTileMem, const FClearValueBinding& ClearValueBinding, GLuint Resource, ETextureCreateFlags Flags) = 0;
+	virtual FTextureRHIRef RHICreateTexture2DArrayFromResource(EPixelFormat Format, uint32 SizeX, uint32 SizeY, uint32 ArraySize, uint32 NumMips, uint32 NumSamples, uint32 NumSamplesTileMem, const FClearValueBinding& ClearValueBinding, GLuint Resource, ETextureCreateFlags Flags) = 0;
+	virtual FTextureRHIRef RHICreateTextureCubeFromResource(EPixelFormat Format, uint32 Size, bool bArray, uint32 ArraySize, uint32 NumMips, uint32 NumSamples, uint32 NumSamplesTileMem, const FClearValueBinding& ClearValueBinding, GLuint Resource, ETextureCreateFlags Flags) = 0;
 
 	virtual GLuint RHIGetResource(FRHITexture* InTexture) const = 0;
 	virtual bool RHIIsValidTexture(GLuint InTexture) const = 0;
diff --git a/Engine/Source/Runtime/RHI/Private/DynamicRHI.cpp b/Engine/Source/Runtime/RHI/Private/DynamicRHI.cpp
--- a/Engine/Source/Runtime/RHI/Private/DynamicRHI.cpp
+++ b/Engine/Source/Runtime/RHI/Private/DynamicRHI.cpp
@@ -540,12 +540,12 @@ void FDynamicRHI::RHIUpdateTextureReference(FRHICommandListBase& RHICmdList, FRH
 	TextureRef->SetReferencedTexture(ReferencedTexture);
 }
 
-void FDynamicRHI::RHIVirtualTextureSetFirstMipInMemory(FRHICommandListImmediate& RHICmdList, FRHITexture2D* TextureRHI, uint32 FirstMip)
+void FDynamicRHI::RHIVirtualTextureSetFirstMipInMemory(FRHICommandListImmediate& RHICmdList, FRHITexture* TextureRHI, uint32 FirstMip)
 {
 	UE_LOG(LogRHI, Fatal, TEXT("The current RHI does not implement support for virtually allocated textures."));
 }
 
-void FDynamicRHI::RHIVirtualTextureSetFirstMipVisible(FRHICommandListImmediate& RHICmdList, FRHITexture2D* TextureRHI, uint32 FirstMip)
+void FDynamicRHI::RHIVirtualTextureSetFirstMipVisible(FRHICommandListImmediate& RHICmdList, FRHITexture* TextureRHI, uint32 FirstMip)
 {
 	UE_LOG(LogRHI, Fatal, TEXT("The current RHI does not implement support for virtually allocated textures."));
 }
