Fix Android link error

diff --git a/Engine/Source/Runtime/OpenGLDrv/Private/Android/AndroidOpenGLPrivate.cpp b/Engine/Source/Runtime/OpenGLDrv/Private/Android/AndroidOpenGLPrivate.cpp
new file mode 100644
--- /dev/null
+++ b/Engine/Source/Runtime/OpenGLDrv/Private/Android/AndroidOpenGLPrivate.cpp
@@ -0,0 +1,93 @@
+// Copyright Epic Games, Inc. All Rights Reserved.
+
+#include "AndroidOpenGLPrivate.h"
+
+bool GAndroidGPUInfoReady = false;
+
+FAndroidGPUInfo::FAndroidGPUInfo()
+{
+	// this is only valid in the game thread, make sure we are initialized there before being called on other threads!
+	check(IsInGameThread())
+
+	// make sure GL is started so we can get the supported formats
+	AndroidEGL* EGL = AndroidEGL::GetInstance();
+
+	if (!EGL->IsInitialized())
+	{
+		FAndroidAppEntry::PlatformInit();
+	}
+
+	// Do not create a window surface if the app is for Oculus Mobile (use small buffer)
+	bool bCreateSurface = !AndroidThunkCpp_IsOculusMobileApplication();
+	FPlatformMisc::LowLevelOutputDebugString(TEXT("FAndroidGPUInfo"));
+	EGL->InitSurface(false, bCreateSurface);
+	EGL->SetCurrentSharedContext();
+
+	// get extensions
+	// Process the extension caps directly here, as FOpenGL might not yet be setup
+	// Do not process extensions here, because extension pointers may not be setup
+	const ANSICHAR* GlGetStringOutput = (const ANSICHAR*)glGetString(GL_EXTENSIONS);
+	FString ExtensionsString = GlGetStringOutput;
+
+	GLVersion = (const ANSICHAR*)glGetString(GL_VERSION);
+
+	// highest priority is the per-texture version
+	if (ExtensionsString.Contains(TEXT("GL_KHR_texture_compression_astc_ldr")))
+	{
+		TargetPlatformNames.Add(TEXT("Android_ASTC"));
+	}
+	if (ExtensionsString.Contains(TEXT("GL_NV_texture_compression_s3tc")) || ExtensionsString.Contains(TEXT("GL_EXT_texture_compression_s3tc")))
+	{
+		TargetPlatformNames.Add(TEXT("Android_DXT"));
+	}
+
+	TargetPlatformNames.Add(TEXT("Android_ETC2"));
+
+	// finally, generic Android
+	TargetPlatformNames.Add(TEXT("Android"));
+
+	bSupportsFloatingPointRenderTargets =
+		ExtensionsString.Contains(TEXT("GL_EXT_color_buffer_half_float"))
+		// According to https://www.khronos.org/registry/gles/extensions/EXT/EXT_color_buffer_float.txt
+		|| (ExtensionsString.Contains(TEXT("GL_EXT_color_buffer_float")));
+
+	bSupportsFrameBufferFetch = ExtensionsString.Contains(TEXT("GL_EXT_shader_framebuffer_fetch")) || ExtensionsString.Contains(TEXT("GL_NV_shader_framebuffer_fetch"))
+		|| ExtensionsString.Contains(TEXT("GL_ARM_shader_framebuffer_fetch ")); // has space at the end to exclude GL_ARM_shader_framebuffer_fetch_depth_stencil match
+
+	GAndroidGPUInfoReady = true;
+
+	VendorName = FString(ANSI_TO_TCHAR((const ANSICHAR*)glGetString(GL_VENDOR)));
+}
+
+void FAndroidGPUInfo::ReadGPUFamily()
+{
+	GPUFamily = (const ANSICHAR*)glGetString(GL_RENDERER);
+	check(!GPUFamily.IsEmpty());
+
+	// thirdparty api requires std::unique_ptr
+	std::unique_ptr<libgpuinfo::instance> ArmGPUInfoInstance = libgpuinfo::instance::create();
+
+	if (ArmGPUInfoInstance)
+	{
+		const libgpuinfo::gpuinfo& ArmGPUInfo = ArmGPUInfoInstance->get_info();
+		// Note:
+		// if libgpuinfo is not upto date then the gpu may not appear in gpuinfo's internal list,
+		// To avoid this we ignore the name and use the lib to extract only the core count. (which does not use the list)
+		const FRegexPattern RegexPattern(TEXT("^Mali(?:.+[MC|MP]([0-9]+))?")); // find anything that starts with Mali and capture the number after the last M[CP]
+		FRegexMatcher RegexMatcher(RegexPattern, *GPUFamily);
+		if (RegexMatcher.FindNext() && ArmGPUInfo.num_shader_cores > 0)
+		{
+			FString Capture = RegexMatcher.GetCaptureGroup(1);
+			if (Capture.IsEmpty())
+			{
+				FString ARMLibName = FString::Format(TEXT("{0} MP{1}"), { *GPUFamily, ArmGPUInfo.num_shader_cores });
+				UE_LOG(LogAndroid, Log, TEXT("FAndroidGPUInfo renaming GPUFamily: %s -> %s"), *GPUFamily, *ARMLibName);
+				GPUFamily = ARMLibName;
+			}
+			else if ((uint32)FCString::Atoi64(*Capture) != ArmGPUInfo.num_shader_cores)
+			{
+				UE_LOG(LogAndroid, Warning, TEXT("FAndroidGPUInfo GPUFamily core count mismatch: %s, expected MP%d"), *GPUFamily, ArmGPUInfo.num_shader_cores);
+			}
+		}
+	}
+}
diff --git a/Engine/Source/Runtime/OpenGLDrv/Private/Android/AndroidOpenGLPrivate.h b/Engine/Source/Runtime/OpenGLDrv/Private/Android/AndroidOpenGLPrivate.h
--- a/Engine/Source/Runtime/OpenGLDrv/Private/Android/AndroidOpenGLPrivate.h
+++ b/Engine/Source/Runtime/OpenGLDrv/Private/Android/AndroidOpenGLPrivate.h
@@ -11,13 +11,12 @@
 #include "Internationalization/Regex.h"
 #include "Misc/CString.h"
 
-bool GAndroidGPUInfoReady = false;
+extern bool GAndroidGPUInfoReady;
 
 // call out to JNI to see if the application was packaged for Oculus Mobile
 extern bool AndroidThunkCpp_IsOculusMobileApplication();
 extern bool ShouldUseGPUFencesToLimitLatency();
 
-
 class FAndroidGPUInfo
 {
 public:
@@ -49,91 +48,6 @@ public:
 private:
 	FString GPUFamily;
 
-	FAndroidGPUInfo()
-	{
-		// this is only valid in the game thread, make sure we are initialized there before being called on other threads!
-		check(IsInGameThread())
-
-		// make sure GL is started so we can get the supported formats
-		AndroidEGL* EGL = AndroidEGL::GetInstance();
-
-		if (!EGL->IsInitialized())
-		{
-			FAndroidAppEntry::PlatformInit();
-		}
-
-		// Do not create a window surface if the app is for Oculus Mobile (use small buffer)
-		bool bCreateSurface = !AndroidThunkCpp_IsOculusMobileApplication();
-		FPlatformMisc::LowLevelOutputDebugString(TEXT("FAndroidGPUInfo"));
-		EGL->InitSurface(false, bCreateSurface);
-		EGL->SetCurrentSharedContext();
-
-		// get extensions
-		// Process the extension caps directly here, as FOpenGL might not yet be setup
-		// Do not process extensions here, because extension pointers may not be setup
-		const ANSICHAR* GlGetStringOutput = (const ANSICHAR*) glGetString(GL_EXTENSIONS);
-		FString ExtensionsString = GlGetStringOutput;
-
-		GLVersion = (const ANSICHAR*)glGetString(GL_VERSION);
-
-		// highest priority is the per-texture version
-		if (ExtensionsString.Contains(TEXT("GL_KHR_texture_compression_astc_ldr")))
-		{
-			TargetPlatformNames.Add(TEXT("Android_ASTC"));
-		}
-		if (ExtensionsString.Contains(TEXT("GL_NV_texture_compression_s3tc")) || ExtensionsString.Contains(TEXT("GL_EXT_texture_compression_s3tc")))
-		{
-			TargetPlatformNames.Add(TEXT("Android_DXT"));
-		}
-		
-		TargetPlatformNames.Add(TEXT("Android_ETC2"));
-
-		// finally, generic Android
-		TargetPlatformNames.Add(TEXT("Android"));
-
-		bSupportsFloatingPointRenderTargets = 
-			ExtensionsString.Contains(TEXT("GL_EXT_color_buffer_half_float")) 
-			// According to https://www.khronos.org/registry/gles/extensions/EXT/EXT_color_buffer_float.txt
-			|| (ExtensionsString.Contains(TEXT("GL_EXT_color_buffer_float")));
-
-		bSupportsFrameBufferFetch = ExtensionsString.Contains(TEXT("GL_EXT_shader_framebuffer_fetch")) || ExtensionsString.Contains(TEXT("GL_NV_shader_framebuffer_fetch")) 
-			|| ExtensionsString.Contains(TEXT("GL_ARM_shader_framebuffer_fetch ")); // has space at the end to exclude GL_ARM_shader_framebuffer_fetch_depth_stencil match
-
-		GAndroidGPUInfoReady = true;
-
-		VendorName = FString(ANSI_TO_TCHAR((const ANSICHAR*)glGetString(GL_VENDOR)));
-	}
-
-	void ReadGPUFamily()
-	{
-		GPUFamily = (const ANSICHAR*)glGetString(GL_RENDERER);
-		check(!GPUFamily.IsEmpty());
-
-		// thirdparty api requires std::unique_ptr
-		std::unique_ptr<libgpuinfo::instance> ArmGPUInfoInstance = libgpuinfo::instance::create();
-
-		if (ArmGPUInfoInstance)
-		{
-			const libgpuinfo::gpuinfo& ArmGPUInfo = ArmGPUInfoInstance->get_info();
-			// Note:
-			// if libgpuinfo is not upto date then the gpu may not appear in gpuinfo's internal list,
-			// To avoid this we ignore the name and use the lib to extract only the core count. (which does not use the list)
-			const FRegexPattern RegexPattern(TEXT("^Mali(?:.+[MC|MP]([0-9]+))?")); // find anything that starts with Mali and capture the number after the last M[CP]
-			FRegexMatcher RegexMatcher(RegexPattern, *GPUFamily);
-			if (RegexMatcher.FindNext() && ArmGPUInfo.num_shader_cores > 0)
-			{
-				FString Capture = RegexMatcher.GetCaptureGroup(1);
-				if (Capture.IsEmpty())
-				{
-					FString ARMLibName = FString::Format(TEXT("{0} MP{1}"), { *GPUFamily, ArmGPUInfo.num_shader_cores });
-					UE_LOG(LogAndroid, Log, TEXT("FAndroidGPUInfo renaming GPUFamily: %s -> %s"), *GPUFamily, *ARMLibName);
-					GPUFamily = ARMLibName;
-				}
-				else if((uint32)FCString::Atoi64(*Capture) != ArmGPUInfo.num_shader_cores)
-				{
-					UE_LOG(LogAndroid, Warning, TEXT("FAndroidGPUInfo GPUFamily core count mismatch: %s, expected MP%d"), *GPUFamily, ArmGPUInfo.num_shader_cores);
-				}
-			}
-		}
-	}
+	FAndroidGPUInfo();
+	void ReadGPUFamily();
 };
