Disable noncoherent fbfetch while programmable blending is being used

This fixes artifacts on Adreno when using thin translucent materials

diff --git a/Engine/Source/Runtime/OpenGLDrv/Private/Android/AndroidOpenGL.cpp b/Engine/Source/Runtime/OpenGLDrv/Private/Android/AndroidOpenGL.cpp
--- a/Engine/Source/Runtime/OpenGLDrv/Private/Android/AndroidOpenGL.cpp
+++ b/Engine/Source/Runtime/OpenGLDrv/Private/Android/AndroidOpenGL.cpp
@@ -836,6 +836,7 @@ void FAndroidOpenGL::SetupDefaultGLContextState(const FString& ExtensionsString)
 		ExtensionsString.Contains(TEXT("GL_QCOM_shader_framebuffer_fetch_noncoherent")) && 
 		ExtensionsString.Contains(TEXT("GL_EXT_shader_framebuffer_fetch")))
 	{
+		bDefaultStateNonCoherentFramebufferFetchEnabled = true;
 		glEnable(GL_FRAMEBUFFER_FETCH_NONCOHERENT_QCOM);
 	}
 }
@@ -858,6 +859,26 @@ void FAndroidOpenGL::EnableAdrenoTilingModeHint(bool bEnable)
 	}
 }
 
+bool FAndroidOpenGL::ResetNonCoherentFramebufferFetch()
+{
+	if (bDefaultStateNonCoherentFramebufferFetchEnabled)
+	{
+		glEnable(GL_FRAMEBUFFER_FETCH_NONCOHERENT_QCOM);
+		return true;
+	}
+	return false;
+}
+
+void FAndroidOpenGL::DisableNonCoherentFramebufferFetch()
+{
+	if (bDefaultStateNonCoherentFramebufferFetchEnabled)
+	{
+		glDisable(GL_FRAMEBUFFER_FETCH_NONCOHERENT_QCOM);
+	}
+}
+
+bool FAndroidOpenGL::bDefaultStateNonCoherentFramebufferFetchEnabled = false;
+
 void FAndroidOpenGL::ProcessExtensions(const FString& ExtensionsString)
 {
 	FString VersionString = FString(ANSI_TO_TCHAR((const ANSICHAR*)glGetString(GL_VERSION)));
diff --git a/Engine/Source/Runtime/OpenGLDrv/Private/Android/AndroidOpenGL.h b/Engine/Source/Runtime/OpenGLDrv/Private/Android/AndroidOpenGL.h
--- a/Engine/Source/Runtime/OpenGLDrv/Private/Android/AndroidOpenGL.h
+++ b/Engine/Source/Runtime/OpenGLDrv/Private/Android/AndroidOpenGL.h
@@ -115,6 +115,10 @@ struct FAndroidOpenGL : public FOpenGLES
 	static void EnableAdrenoTilingModeHint(bool bEnable);
 	static bool bRequiresAdrenoTilingHint;
 
+	static bool ResetNonCoherentFramebufferFetch();
+	static void DisableNonCoherentFramebufferFetch();
+	static bool bDefaultStateNonCoherentFramebufferFetchEnabled;
+
 	/** supported OpenGL ES version queried from the system */
 	static int32 GLMajorVerion;
 	static int32 GLMinorVersion;
diff --git a/Engine/Source/Runtime/OpenGLDrv/Private/OpenGLRenderTarget.cpp b/Engine/Source/Runtime/OpenGLDrv/Private/OpenGLRenderTarget.cpp
--- a/Engine/Source/Runtime/OpenGLDrv/Private/OpenGLRenderTarget.cpp
+++ b/Engine/Source/Runtime/OpenGLDrv/Private/OpenGLRenderTarget.cpp
@@ -1070,13 +1070,18 @@ void FOpenGLDynamicRHI::RHIBeginRenderPass(const FRHIRenderPassInfo& InInfo, con
 	{
 		glEnable(GL_SHADER_PIXEL_LOCAL_STORAGE_EXT);
 	}
-#endif
 
-#if PLATFORM_ANDROID
 	if (FAndroidOpenGL::RequiresAdrenoTilingModeHint())
 	{
 		FAndroidOpenGL::EnableAdrenoTilingModeHint(FCString::Strcmp(InName, TEXT("SceneColorRendering")) == 0);
 	}
+
+	// Reenable non-coherent framebuffer fetch if needed
+	FOpenGLContextState& ContextState = GetContextStateForCurrentContext();
+	if (!ContextState.bNonCoherentFramebufferFetchEnabled)
+	{
+		ContextState.bNonCoherentFramebufferFetchEnabled = FAndroidOpenGL::ResetNonCoherentFramebufferFetch();
+	}
 #endif
 }
 
@@ -1140,7 +1145,10 @@ void FOpenGLDynamicRHI::RHINextSubpass()
 	if (RenderPassInfo.SubpassHint == ESubpassHint::DepthReadSubpass ||
 		RenderPassInfo.SubpassHint == ESubpassHint::DeferredShadingSubpass)
 	{
-		FOpenGL::FrameBufferFetchBarrier();
+		if (GetContextStateForCurrentContext().bNonCoherentFramebufferFetchEnabled)
+		{
+			FOpenGL::FrameBufferFetchBarrier();
+		}
 	}
 }
 
diff --git a/Engine/Source/Runtime/OpenGLDrv/Private/OpenGLShaders.cpp b/Engine/Source/Runtime/OpenGLDrv/Private/OpenGLShaders.cpp
--- a/Engine/Source/Runtime/OpenGLDrv/Private/OpenGLShaders.cpp
+++ b/Engine/Source/Runtime/OpenGLDrv/Private/OpenGLShaders.cpp
@@ -1002,6 +1002,15 @@ FOpenGLShader::FOpenGLShader(TArrayView<const uint8> Code, const FSHAHash& Hash,
 	uint32 CodeCRC = FCrc::MemCrc32(GlslCodeOriginal.GetData(), GlslCodeOriginal.Num());
 	ShaderCodeKey = FOpenGLCompiledShaderKey(TypeEnum, GlslCodeOriginal.Num(), CodeCRC);
 
+	if (TypeEnum == GL_FRAGMENT_SHADER && FOpenGL::SupportsShaderFramebufferFetch())
+	{
+		// _Globals_gl_LastFragColor should only exist when 'FramebufferFetchGLES2()' is being used, not for MRT/deferred
+		if (FCStringAnsi::Strstr(GlslCodeOriginal.GetData(), "_Globals_gl_LastFragColor") != nullptr)
+		{
+			bUsesProgrammableBlending = true;
+		}
+	}
+
 	FAnsiCharArray GlslCodeFinal;
 	{
 		FScopeLock Lock(&GCompiledShaderCacheCS);
@@ -1269,10 +1278,13 @@ public:
 		FShaderStage Vertex;
 		FShaderStage Pixel;
 		TOptional<FShaderStage> Geometry;
+		bool bUsesProgrammableBlending;
 
 		FGraphicsProgram(FOpenGLLinkedProgramBase& ProgramBase, FOpenGLVertexShader* VertexShader, FOpenGLPixelShader* PixelShader, FOpenGLGeometryShader* GeometryShader)
 			: Vertex(ProgramBase, *VertexShader, CrossCompiler::SHADER_STAGE_VERTEX, 0)
 			, Pixel(ProgramBase, *PixelShader, CrossCompiler::SHADER_STAGE_PIXEL, Vertex.Bindings.NumUniformBuffers)
+			, Geometry()
+			, bUsesProgrammableBlending(PixelShader->bUsesProgrammableBlending)
 		{
 			if (GeometryShader)
 			{
@@ -2848,13 +2860,25 @@ void FOpenGLDynamicRHI::BindPendingShaderState( FOpenGLContextState& ContextStat
 
 	bool ForceUniformBindingUpdate = false;
 
-	GLuint PendingProgram = PendingState.BoundShaderState->LinkedProgram->Program;
+	FOpenGLLinkedProgram* const PendingLinkedProgram = PendingState.BoundShaderState->LinkedProgram;
+	const GLuint PendingProgram = PendingLinkedProgram->Program;
 	if (ContextState.Program != PendingProgram)
 	{
 		FOpenGL::BindProgramPipeline(PendingProgram);
 		ContextState.Program = PendingProgram;
 		MarkShaderParameterCachesDirty(PendingState.ShaderParameters, false);
 		PendingState.LinkedProgramAndDirtyFlag = nullptr;
+
+#if PLATFORM_ANDROID
+		// Disable non-coherent framebuffer fetch if it's being used for programmable blending to make sure that we actually fetch the last pixel value in draw order
+		if (ContextState.bNonCoherentFramebufferFetchEnabled)
+		{
+			if (PendingLinkedProgram->GetGraphicsProgram().bUsesProgrammableBlending)
+			{
+				FAndroidOpenGL::DisableNonCoherentFramebufferFetch();
+			}
+		}
+#endif
 	}
 
 	if (PendingState.bAnyDirtyRealUniformBuffers[SF_Vertex] || 
diff --git a/Engine/Source/Runtime/OpenGLDrv/Private/OpenGLState.h b/Engine/Source/Runtime/OpenGLDrv/Private/OpenGLState.h
--- a/Engine/Source/Runtime/OpenGLDrv/Private/OpenGLState.h
+++ b/Engine/Source/Runtime/OpenGLDrv/Private/OpenGLState.h
@@ -340,6 +340,7 @@ struct FOpenGLContextState final : public FOpenGLCommonState
 	float							ClearDepth;
 	int32							FirstNonzeroRenderTarget;
 	bool							bAlphaToCoverageEnabled;
+	bool							bNonCoherentFramebufferFetchEnabled;
 	
 	FOpenGLVertexDeclaration*		VertexDecl;
 	FOpenGLCachedAttr				VertexAttrs[NUM_OPENGL_VERTEX_STREAMS];
@@ -384,6 +385,7 @@ struct FOpenGLContextState final : public FOpenGLCommonState
 	,	ClearDepth(-1.0f)
 	,	FirstNonzeroRenderTarget(0)
 	,	bAlphaToCoverageEnabled(false)
+	,	bNonCoherentFramebufferFetchEnabled(false)
 	,	VertexDecl(0)
 	,   VertexAttrs()
 	,	VertexStreams()
diff --git a/Engine/Source/Runtime/OpenGLDrv/Public/OpenGLShaderResources.h b/Engine/Source/Runtime/OpenGLDrv/Public/OpenGLShaderResources.h
--- a/Engine/Source/Runtime/OpenGLDrv/Public/OpenGLShaderResources.h
+++ b/Engine/Source/Runtime/OpenGLDrv/Public/OpenGLShaderResources.h
@@ -207,6 +207,8 @@ public:
 
 	FOpenGLCompiledShaderKey ShaderCodeKey;
 
+	bool bUsesProgrammableBlending = false;
+
 #if DEBUG_GL_SHADERS
 	TArray<ANSICHAR> GlslCode;
 	const ANSICHAR*  GlslCodeString; // make it easier in VS to see shader code in debug mode; points to begin of GlslCode
