Leave only immutable texture path for GL(only use glTexStorage2D), don't set GL_TEXTURE_MAX_LEVEL as it isn't needed anymore

diff --git a/Engine/Source/Runtime/OpenGLDrv/Private/OpenGLCommands.cpp b/Engine/Source/Runtime/OpenGLDrv/Private/OpenGLCommands.cpp
--- a/Engine/Source/Runtime/OpenGLDrv/Private/OpenGLCommands.cpp
+++ b/Engine/Source/Runtime/OpenGLDrv/Private/OpenGLCommands.cpp
@@ -457,10 +457,6 @@ void FOpenGLDynamicRHI::CachedSetupTextureStageInner(FOpenGLContextState& Contex
 				{
 					FOpenGL::TexParameter(Target, GL_TEXTURE_BASE_LEVEL, BaseMip);
 				}
-				if (!bSameNumMips)
-				{
-					FOpenGL::TexParameter(Target, GL_TEXTURE_MAX_LEVEL, MaxMip);
-				}
 				if (MipLimits)
 				{
 					MipLimits->Key = BaseMip;
diff --git a/Engine/Source/Runtime/OpenGLDrv/Private/OpenGLTexture.cpp b/Engine/Source/Runtime/OpenGLDrv/Private/OpenGLTexture.cpp
--- a/Engine/Source/Runtime/OpenGLDrv/Private/OpenGLTexture.cpp
+++ b/Engine/Source/Runtime/OpenGLDrv/Private/OpenGLTexture.cpp
@@ -595,12 +595,6 @@ void FOpenGLDynamicRHI::InitializeGLTextureInternal(FOpenGLTexture* Texture, voi
 
 		glTexParameteri(Target, GL_TEXTURE_BASE_LEVEL, 0);
 
-		// Do not use GL_TEXTURE_MAX_LEVEL if external texture
-		if (Target != GL_TEXTURE_EXTERNAL_OES)
-		{
-			glTexParameteri(Target, GL_TEXTURE_MAX_LEVEL, Desc.NumMips - 1);
-		}
-
 		TextureMipLimits.Add(TextureID, TPair<GLenum, GLenum>(0, Desc.NumMips - 1));
 
 		if (FOpenGL::SupportsASTCDecodeMode())
@@ -673,9 +667,10 @@ void FOpenGLDynamicRHI::InitializeGLTextureInternal(FOpenGLTexture* Texture, voi
 			case GL_RENDERBUFFER:
 			case GL_TEXTURE_2D:
 			case GL_TEXTURE_CUBE_MAP:
-				// Try to create the texture using immutable storage
-				if (FOpenGL::TexStorage2D(Target, Desc.NumMips, GLFormat.InternalFormat[bSRGB], Desc.Extent.X, Desc.Extent.Y, GLFormat.Format, GLFormat.Type, Desc.Flags))
 				{
+					// Try to create the texture using immutable storage
+					FOpenGL::TexStorage2D(Target, Desc.NumMips, GLFormat.InternalFormat[bSRGB], Desc.Extent.X, Desc.Extent.Y, GLFormat.Format, GLFormat.Type, Desc.Flags);
+
 					// Texture created with immutable storage. Now fill in the bulk data.
 					bAllocatedStorage = true;
 
@@ -714,45 +709,6 @@ void FOpenGLDynamicRHI::InitializeGLTextureInternal(FOpenGLTexture* Texture, voi
 						});
 					}
 				}
-				else if (GLFormat.bCompressed && !BulkDataPtr)
-				{
-					// Compressed textures created without using the TexStorage functions cannot be allocated via TexImage without bulk data.
-					// Do nothing here. Texture memory will be allocated when the renderer locks/unlocks the mips for writing.
-				}
-				else
-				{
-					// Failed to create immutable storage. Fall back to the standard TexImage functions.
-					// This both allocates the memory and fills in the bulk data simultaneously.
-					EnumerateSubresources(BulkDataPtr, [&](GLenum CurrentTarget, uint32 MipSizeX, uint32 MipSizeY, uint32 MipSizeZ, uint32 MipIndex, uint32 ArraySlice, void const* MipSliceData, uint32 MipSliceSize)
-					{
-						if (GLFormat.bCompressed)
-						{
-							glCompressedTexImage2D(
-								CurrentTarget,
-								MipIndex,
-								GLFormat.InternalFormat[bSRGB],
-								MipSizeX, MipSizeY, 
-								0,
-								MipSliceSize,
-								MipSliceData);
-						}
-						else
-						{
-							glTexImage2D(
-								CurrentTarget,
-								MipIndex,
-								GLFormat.InternalFormat[bSRGB],
-								MipSizeX, MipSizeY,
-								0,
-								GLFormat.Format,
-								GLFormat.Type,
-								MipSliceData);
-						}
-
-						// Always continue iterating to allocate all mips/slices.
-						return true;
-					});
-				}
 				break;
 
 			case GL_TEXTURE_2D_ARRAY:
@@ -813,14 +769,8 @@ void FOpenGLDynamicRHI::InitializeGLTextureInternal(FOpenGLTexture* Texture, voi
 					// Try to create an immutable storage texture and fallback if it fails
 					const int32 NumSamples = Texture->GetDesc().NumSamples;
 					const bool FixedSampleLocations = true;
-					if (FOpenGL::TexStorage2DMultisample(Target, NumSamples, GLFormat.InternalFormat[bSRGB], Desc.Extent.X, Desc.Extent.Y, FixedSampleLocations))
-					{
-						bAllocatedStorage = true;
-					}
-					else
-					{
-						FOpenGL::TexImage2DMultisample(Target, NumSamples, GLFormat.InternalFormat[bSRGB], Desc.Extent.X, Desc.Extent.Y, FixedSampleLocations);
-					}
+					FOpenGL::TexStorage2DMultisample(Target, NumSamples, GLFormat.InternalFormat[bSRGB], Desc.Extent.X, Desc.Extent.Y, FixedSampleLocations);
+					bAllocatedStorage = true;
 				}
 				break;
 			}
