Remove RHIInvalidateCachedState
 - Only implemented in OpenGL RHI, and never called.

diff --git a/Engine/Source/Runtime/OpenGLDrv/Private/OpenGLCommands.cpp b/Engine/Source/Runtime/OpenGLDrv/Private/OpenGLCommands.cpp
--- a/Engine/Source/Runtime/OpenGLDrv/Private/OpenGLCommands.cpp
+++ b/Engine/Source/Runtime/OpenGLDrv/Private/OpenGLCommands.cpp
@@ -2921,17 +2921,6 @@ void FOpenGLDynamicRHI::RHISubmitCommandLists(FRHISubmitCommandListsArgs&& Args)
 	// Nothing to do
 }
 
-void FOpenGLDynamicRHI::RHIInvalidateCachedState()
-{
-	RenderingContextState = FOpenGLContextState();
-	SharedContextState = FOpenGLContextState();
-
-	GLint NumUAVUnits = FMath::Max(FOpenGL::GetMaxCombinedUAVUnits(), FOpenGL::GetMaxComputeUAVUnits());
-
-	RenderingContextState.InitializeResources(FOpenGL::GetMaxCombinedTextureImageUnits(), NumUAVUnits);
-	SharedContextState.InitializeResources(FOpenGL::GetMaxCombinedTextureImageUnits(), NumUAVUnits);
-}
-
 void FOpenGLDynamicRHI::RHICopyToStagingBuffer(FRHIBuffer* SourceBufferRHI, FRHIStagingBuffer* DestinationStagingBufferRHI, uint32 InOffset, uint32 InNumBytes)
 {
 	VERIFY_GL_SCOPE();
diff --git a/Engine/Source/Runtime/OpenGLDrv/Private/OpenGLDrv.h b/Engine/Source/Runtime/OpenGLDrv/Private/OpenGLDrv.h
--- a/Engine/Source/Runtime/OpenGLDrv/Private/OpenGLDrv.h
+++ b/Engine/Source/Runtime/OpenGLDrv/Private/OpenGLDrv.h
@@ -464,7 +464,6 @@ public:
 	virtual void RHIBeginBreadcrumbGPU(FRHIBreadcrumbNode* Breadcrumb) final override;
 	virtual void RHIEndBreadcrumbGPU  (FRHIBreadcrumbNode* Breadcrumb) final override;
 #endif
-	virtual void RHIInvalidateCachedState() final override;
 	virtual void RHIDiscardRenderTargets(bool Depth,bool Stencil,uint32 ColorBitMask) final override;
 
 	// FIXME: Broken on Android for cubemaps
diff --git a/Engine/Source/Runtime/RHI/Public/RHICommandList.h b/Engine/Source/Runtime/RHI/Public/RHICommandList.h
--- a/Engine/Source/Runtime/RHI/Public/RHICommandList.h
+++ b/Engine/Source/Runtime/RHI/Public/RHICommandList.h
@@ -2184,11 +2184,6 @@ FRHICOMMAND_MACRO(FRHICommandEndDrawingViewport)
 	RHI_API void Execute(FRHICommandListBase& CmdList);
 };
 
-FRHICOMMAND_MACRO(FRHICommandInvalidateCachedState)
-{
-	RHI_API void Execute(FRHICommandListBase& CmdList);
-};
-
 FRHICOMMAND_MACRO(FRHICommandDiscardRenderTargets)
 {
 	uint32 ColorBitMask;
@@ -3603,15 +3598,9 @@ public:
 		IncrementSubpass();
 	}
 
+	UE_DEPRECATED(5.5, "RHIInvalidateCachedState is not implemented by any platform RHI.")
 	FORCEINLINE_DEBUGGABLE void RHIInvalidateCachedState()
-	{
-		if (Bypass())
-		{
-			GetContext().RHIInvalidateCachedState();
-			return;
-		}
-		ALLOC_COMMAND(FRHICommandInvalidateCachedState)();
-	}
+	{}
 
 	FORCEINLINE void DiscardRenderTargets(bool Depth, bool Stencil, uint32 ColorBitMask)
 	{
diff --git a/Engine/Source/Runtime/RHI/Public/RHICommandListCommandExecutes.inl b/Engine/Source/Runtime/RHI/Public/RHICommandListCommandExecutes.inl
--- a/Engine/Source/Runtime/RHI/Public/RHICommandListCommandExecutes.inl
+++ b/Engine/Source/Runtime/RHI/Public/RHICommandListCommandExecutes.inl
@@ -558,12 +558,6 @@ void FRHICommandEndDrawingViewport::Execute(FRHICommandListBase& CmdList)
 	INTERNAL_DECORATOR(RHIEndDrawingViewport)(Viewport, bPresent, bLockToVsync);
 }
 
-void FRHICommandInvalidateCachedState::Execute(FRHICommandListBase& CmdList)
-{
-	RHISTAT(RHIInvalidateCachedState);
-	INTERNAL_DECORATOR(RHIInvalidateCachedState)();
-}
-
 void FRHICommandDiscardRenderTargets::Execute(FRHICommandListBase& CmdList)
 {
 	RHISTAT(RHIDiscardRenderTargets);
diff --git a/Engine/Source/Runtime/RHI/Public/RHIContext.h b/Engine/Source/Runtime/RHI/Public/RHIContext.h
--- a/Engine/Source/Runtime/RHI/Public/RHIContext.h
+++ b/Engine/Source/Runtime/RHI/Public/RHIContext.h
@@ -347,12 +347,6 @@ public:
 	virtual void RHIEndBreadcrumbGPU  (FRHIBreadcrumbNode* Breadcrumb) = 0;
 #endif
 
-	/**
-	 * Some RHI implementations (OpenGL) cache render state internally
-	 * Signal to RHI that cached state is no longer valid
-	 */
-	virtual void RHIInvalidateCachedState() {}
-
 	/**
 	 * Performs a copy of the data in 'SourceBuffer' to 'DestinationStagingBuffer.' This will occur inline on the GPU timeline. This is a mechanism to perform nonblocking readback of a buffer at a point in time.
 	 * @param SourceBuffer The source vertex buffer that will be inlined copied.
