Remove "rhi.ResourceTableCaching" and associated counters in RHIBeginScene / RHIEndScene implementations
 - The counters were unused.

diff --git a/Engine/Source/Runtime/Apple/MetalRHI/Private/MetalContext.cpp b/Engine/Source/Runtime/Apple/MetalRHI/Private/MetalContext.cpp
--- a/Engine/Source/Runtime/Apple/MetalRHI/Private/MetalContext.cpp
+++ b/Engine/Source/Runtime/Apple/MetalRHI/Private/MetalContext.cpp
@@ -352,7 +352,6 @@ FMetalDeviceContext::FMetalDeviceContext(MTL::Device* MetalDevice, uint32 InDevi
 	: FMetalContext(MetalDevice, *Queue)
 	, DeviceIndex(InDeviceIndex)
 	, CaptureManager(MetalDevice, *Queue)
-	, SceneFrameCounter(0)
 	, FrameCounter(0)
 	, PSOManager(0)
 	, FrameNumberRHIThread(0)
@@ -601,14 +600,6 @@ void FMetalDeviceContext::BeginScene()
 #if ENABLE_METAL_GPUPROFILE
 	FPlatformTLS::SetTlsValue(CurrentContextTLSSlot, this);
 #endif
-	
-	// Increment the frame counter. INDEX_NONE is a special value meaning "uninitialized", so if
-	// we hit it just wrap around to zero.
-	SceneFrameCounter++;
-	if (SceneFrameCounter == INDEX_NONE)
-	{
-		SceneFrameCounter++;
-	}
 }
 
 void FMetalDeviceContext::EndScene()
diff --git a/Engine/Source/Runtime/Apple/MetalRHI/Private/MetalContext.h b/Engine/Source/Runtime/Apple/MetalRHI/Private/MetalContext.h
--- a/Engine/Source/Runtime/Apple/MetalRHI/Private/MetalContext.h
+++ b/Engine/Source/Runtime/Apple/MetalRHI/Private/MetalContext.h
@@ -295,9 +295,6 @@ private:
 	/** Event for coordinating pausing of render thread to keep inline with the ios display link. */
 	FEvent* FrameReadyEvent;
 	
-	/** Internal frame counter, incremented on each call to RHIBeginScene. */
-	uint32 SceneFrameCounter;
-	
 	/** Internal frame counter, used to ensure that we only drain the buffer pool one after each frame within RHIEndFrame. */
 	uint32 FrameCounter;
 	
diff --git a/Engine/Source/Runtime/OpenGLDrv/Private/OpenGLDevice.cpp b/Engine/Source/Runtime/OpenGLDrv/Private/OpenGLDevice.cpp
--- a/Engine/Source/Runtime/OpenGLDrv/Private/OpenGLDevice.cpp
+++ b/Engine/Source/Runtime/OpenGLDrv/Private/OpenGLDevice.cpp
@@ -248,26 +248,11 @@ void FOpenGLDynamicRHI::RHIAdvanceFrameFence()
 
 void FOpenGLDynamicRHI::RHIBeginScene()
 {
-	// Increment the frame counter. INDEX_NONE is a special value meaning "uninitialized", so if
-	// we hit it just wrap around to zero.
-	SceneFrameCounter++;
-	if (SceneFrameCounter == INDEX_NONE)
-	{
-		SceneFrameCounter++;
-	}
-
-	static auto* ResourceTableCachingCvar = IConsoleManager::Get().FindTConsoleVariableDataInt(TEXT("rhi.ResourceTableCaching"));
-	if (ResourceTableCachingCvar == NULL || ResourceTableCachingCvar->GetValueOnAnyThread() == 1)
-	{
-		ResourceTableFrameCounter = SceneFrameCounter;
-	}
-
 	BeginSceneContextType = (int32)PlatformOpenGLCurrentContext(PlatformDevice);
 }
 
 void FOpenGLDynamicRHI::RHIEndScene()
 {
-	ResourceTableFrameCounter = INDEX_NONE;
 	BeginSceneContextType = CONTEXT_Other;
 }
 
@@ -1307,9 +1292,7 @@ FDynamicRHI* FOpenGLDynamicRHIModule::CreateRHI(ERHIFeatureLevel::Type InRequest
 
 
 FOpenGLDynamicRHI::FOpenGLDynamicRHI()
-:	SceneFrameCounter(0)
-,	ResourceTableFrameCounter(INDEX_NONE)
-,	bRevertToSharedContextAfterDrawingViewport(false)
+:	bRevertToSharedContextAfterDrawingViewport(false)
 ,	bIsRenderingContextAcquired(false)
 ,   BeginSceneContextType(CONTEXT_Other)
 ,	PlatformDevice(NULL)
diff --git a/Engine/Source/Runtime/OpenGLDrv/Private/OpenGLDrv.h b/Engine/Source/Runtime/OpenGLDrv/Private/OpenGLDrv.h
--- a/Engine/Source/Runtime/OpenGLDrv/Private/OpenGLDrv.h
+++ b/Engine/Source/Runtime/OpenGLDrv/Private/OpenGLDrv.h
@@ -814,12 +814,6 @@ private:
 	/** called once per frame, used for resource processing */
 	void EndFrameTick();
 
-	/** Counter incremented each time RHIBeginScene is called. */
-	uint32 SceneFrameCounter;
-
-	/** Value used to detect when resource tables need to be recached. INDEX_NONE means always recache. */
-	uint32 ResourceTableFrameCounter;
-
 	/** RHI device state, independent of underlying OpenGL context used */
 	FOpenGLRHIState						PendingState;
 	FSamplerStateRHIRef					PointSamplerState;
diff --git a/Engine/Source/Runtime/RHI/Private/RHI.cpp b/Engine/Source/Runtime/RHI/Private/RHI.cpp
--- a/Engine/Source/Runtime/RHI/Private/RHI.cpp
+++ b/Engine/Source/Runtime/RHI/Private/RHI.cpp
@@ -1104,12 +1104,6 @@ static_assert((int32)ERHIZBuffer::FarPlane == 0 || (int32)ERHIZBuffer::FarPlane
 /**
  * RHI configuration settings.
  */
-
-static TAutoConsoleVariable<int32> ResourceTableCachingCvar(
-	TEXT("rhi.ResourceTableCaching"),
-	1,
-	TEXT("If 1, the RHI will cache resource table contents within a frame. Otherwise resource tables are rebuilt for every draw call.")
-	);
 static TAutoConsoleVariable<int32> GSaveScreenshotAfterProfilingGPUCVar(
 	TEXT("r.ProfileGPU.Screenshot"),
 	1,
diff --git a/Engine/Source/Runtime/Windows/D3D11RHI/Private/D3D11Device.cpp b/Engine/Source/Runtime/Windows/D3D11RHI/Private/D3D11Device.cpp
--- a/Engine/Source/Runtime/Windows/D3D11RHI/Private/D3D11Device.cpp
+++ b/Engine/Source/Runtime/Windows/D3D11RHI/Private/D3D11Device.cpp
@@ -53,9 +53,7 @@ FD3D11DynamicRHI::FD3D11DynamicRHI(IDXGIFactory1* InDXGIFactory1, D3D_FEATURE_LE
 	, CurrentDepthTexture(NULL)
 	, NumSimultaneousRenderTargets(0)
 	, NumUAVs(0)
-	, SceneFrameCounter(0)
 	, PresentCounter(0)
-	, ResourceTableFrameCounter(INDEX_NONE)
 	, CurrentDSVAccessType(FExclusiveDepthStencil::DepthWrite_StencilWrite)
 	, bDiscardSharedConstants(false)
 	, GPUProfilingData(this)
diff --git a/Engine/Source/Runtime/Windows/D3D11RHI/Private/D3D11RHI.cpp b/Engine/Source/Runtime/Windows/D3D11RHI/Private/D3D11RHI.cpp
--- a/Engine/Source/Runtime/Windows/D3D11RHI/Private/D3D11RHI.cpp
+++ b/Engine/Source/Runtime/Windows/D3D11RHI/Private/D3D11RHI.cpp
@@ -461,24 +461,12 @@ float FD3D11EventNode::GetTiming()
 
 void FD3D11DynamicRHI::RHIBeginScene()
 {
-	// Increment the frame counter. INDEX_NONE is a special value meaning "uninitialized", so if
-	// we hit it just wrap around to zero.
-	SceneFrameCounter++;
-	if (SceneFrameCounter == INDEX_NONE)
-	{
-		SceneFrameCounter++;
-	}
-
-	static auto* ResourceTableCachingCvar = IConsoleManager::Get().FindTConsoleVariableDataInt(TEXT("rhi.ResourceTableCaching"));
-	if (ResourceTableCachingCvar == NULL || ResourceTableCachingCvar->GetValueOnAnyThread() == 1)
-	{
-		ResourceTableFrameCounter = SceneFrameCounter;
-	}
+	// Nothing to do
 }
 
 void FD3D11DynamicRHI::RHIEndScene()
 {
-	ResourceTableFrameCounter = INDEX_NONE;
+	// Nothing to do
 }
 
 FD3DGPUProfiler::FD3DGPUProfiler(class FD3D11DynamicRHI* InD3DRHI)
diff --git a/Engine/Source/Runtime/Windows/D3D11RHI/Private/D3D11RHIPrivate.h b/Engine/Source/Runtime/Windows/D3D11RHI/Private/D3D11RHIPrivate.h
--- a/Engine/Source/Runtime/Windows/D3D11RHI/Private/D3D11RHIPrivate.h
+++ b/Engine/Source/Runtime/Windows/D3D11RHI/Private/D3D11RHIPrivate.h
@@ -786,9 +786,6 @@ protected:
 	uint32 NumSimultaneousRenderTargets;
 	uint32 NumUAVs;
 
-	/** Internal frame counter, incremented on each call to RHIBeginScene. */
-	uint32 SceneFrameCounter;
-
 	/** Internal frame counter that just counts calls to Present */
 	uint32 PresentCounter;
 
@@ -800,12 +797,6 @@ protected:
 	} ActiveQueries;
 	friend class FD3D11RenderQuery;
 
-	/**
-	 * Internal counter used for resource table caching.
-	 * INDEX_NONE means caching is not allowed.
-	 */
-	uint32 ResourceTableFrameCounter;
-
 	/** D3D11 defines a maximum of 14 constant buffers per shader stage. */
 	enum { MAX_UNIFORM_BUFFERS_PER_SHADER_STAGE = 14 };
 
