Fix warning V595: The 'VertexShader' pointer was utilized before it was verified against nullptr.
Fix warning V595: The 'PixelShader' pointer was utilized before it was verified against nullptr.

#rnx

diff --git a/Engine/Source/Runtime/OpenGLDrv/Private/OpenGLShaders.cpp b/Engine/Source/Runtime/OpenGLDrv/Private/OpenGLShaders.cpp
--- a/Engine/Source/Runtime/OpenGLDrv/Private/OpenGLShaders.cpp
+++ b/Engine/Source/Runtime/OpenGLDrv/Private/OpenGLShaders.cpp
@@ -2462,11 +2462,11 @@ FOpenGLLinkedProgram::FOpenGLLinkedProgram(FOpenGLVertexShader* VertexShader, FO
 	if (!VerifyLinkedProgram(Program))
 	{
 #if DEBUG_GL_SHADERS
-		if (VertexShader)
+		//if (VertexShader)
 		{
 			UE_LOG(LogRHI, Error, TEXT("Vertex Shader:\n%s"), ANSI_TO_TCHAR(VertexShader->GlslCode.GetData()));
 		}
-		if (PixelShader)
+		//if (PixelShader)
 		{
 			UE_LOG(LogRHI, Error, TEXT("Pixel Shader:\n%s"), ANSI_TO_TCHAR(PixelShader->GlslCode.GetData()));
 		}
