Fix OpenGL startup assert
 - The start up of the rendering and RHI threads can lead to RHIFinalizeContext being called from the render thread even when the RHI thread is created, due to Bypass not having been disabled yet. This was causing the VERIFY_GL_SCOPE in that function to assert.
 - Replaced the VERIFY_GL_SCOPE with an if that simply skips the call to glFlush() if the calling thread has no active GL context. This is simpler than finding a way to avoid the call to RHIFinalizeContext.

Also fix a race between the RHI thread and render thread on startup
 - The GIsRunningRHIInDedicatedThread_InternalUseOnly etc booleans were set after the RHI thread was already created.
 - Moved the RHI thread creation to after those bools are set, so the bools don't change while the RHI thread starts.

Removed PlatformRebindResources, and OnAcquireThreadOwnership / OnReleaseThreadOwnership as these functions were unused.

Replaced the explicit thread checks in the OpenGL CHECK_EXPECTED_GL_THREAD macro with a call to PlatformOpenGLContextValid(), and enabled this in non test/shipping builds.

diff --git a/Engine/Source/Runtime/OpenGLDrv/Private/Android/AndroidOpenGL.cpp b/Engine/Source/Runtime/OpenGLDrv/Private/Android/AndroidOpenGL.cpp
--- a/Engine/Source/Runtime/OpenGLDrv/Private/Android/AndroidOpenGL.cpp
+++ b/Engine/Source/Runtime/OpenGLDrv/Private/Android/AndroidOpenGL.cpp
@@ -276,10 +276,6 @@ void PlatformFlushIfNeeded()
 {
 }
 
-void PlatformRebindResources(FPlatformOpenGLDevice* Device)
-{
-}
-
 void PlatformSharedContextSetup(FPlatformOpenGLDevice* Device)
 {
 	Device->SetCurrentSharedContext();
diff --git a/Engine/Source/Runtime/OpenGLDrv/Private/Linux/LinuxOpenGLPlatform.cpp b/Engine/Source/Runtime/OpenGLDrv/Private/Linux/LinuxOpenGLPlatform.cpp
--- a/Engine/Source/Runtime/OpenGLDrv/Private/Linux/LinuxOpenGLPlatform.cpp
+++ b/Engine/Source/Runtime/OpenGLDrv/Private/Linux/LinuxOpenGLPlatform.cpp
@@ -531,11 +531,6 @@ void PlatformFlushIfNeeded()
 	glFinish();
 }
 
-void PlatformRebindResources(FPlatformOpenGLDevice* Device)
-{
-	// @todo: Figure out if we need to rebind frame & renderbuffers after switching contexts
-}
-
 void PlatformRenderingContextSetup(FPlatformOpenGLDevice* Device)
 {
 	check(Device && Device->RenderingContext.hWnd && Device->RenderingContext.hGLContext);
diff --git a/Engine/Source/Runtime/OpenGLDrv/Private/OpenGLCommands.cpp b/Engine/Source/Runtime/OpenGLDrv/Private/OpenGLCommands.cpp
--- a/Engine/Source/Runtime/OpenGLDrv/Private/OpenGLCommands.cpp
+++ b/Engine/Source/Runtime/OpenGLDrv/Private/OpenGLCommands.cpp
@@ -2913,8 +2913,10 @@ IRHIPlatformCommandList* FOpenGLDynamicRHI::RHIFinalizeContext(FRHIFinalizeConte
 	check(Args.Context == this);
 
 	// Flush the context to ensure recorded commands will reach the driver.
-	VERIFY_GL_SCOPE();
-	FOpenGL::Flush();
+	if (PlatformOpenGLContextValid())
+	{
+		FOpenGL::Flush();
+	}
 
 	// Clear some context state
 	FMemory::Memset(PendingState.BoundUniformBuffers, 0, sizeof(PendingState.BoundUniformBuffers));
diff --git a/Engine/Source/Runtime/OpenGLDrv/Private/OpenGLDevice.cpp b/Engine/Source/Runtime/OpenGLDrv/Private/OpenGLDevice.cpp
--- a/Engine/Source/Runtime/OpenGLDrv/Private/OpenGLDevice.cpp
+++ b/Engine/Source/Runtime/OpenGLDrv/Private/OpenGLDevice.cpp
@@ -1550,29 +1550,16 @@ void FOpenGLDynamicRHI::RHIAcquireThreadOwnership()
 {
 	check(!bRevertToSharedContextAfterDrawingViewport);	// if this is true, then main thread is rendering using our context right now.
 	PlatformRenderingContextSetup(PlatformDevice);
-	PlatformRebindResources(PlatformDevice);
+
 	bIsRenderingContextAcquired = true;
 	VERIFY_GL(RHIAcquireThreadOwnership);
-	{
-		FScopeLock lock(&CustomPresentSection);
-		if (CustomPresent)
-		{
-			CustomPresent->OnAcquireThreadOwnership();
-		}
-	}
 }
 
 void FOpenGLDynamicRHI::RHIReleaseThreadOwnership()
 {
-	{
-		FScopeLock lock(&CustomPresentSection);
-		if (CustomPresent)
-		{
-			CustomPresent->OnReleaseThreadOwnership();
-		}
-	}
 	VERIFY_GL(RHIReleaseThreadOwnership);
 	bIsRenderingContextAcquired = false;
+
 	PlatformNULLContextSetup();
 }
 
diff --git a/Engine/Source/Runtime/OpenGLDrv/Private/OpenGLDrvPrivate.h b/Engine/Source/Runtime/OpenGLDrv/Private/OpenGLDrvPrivate.h
--- a/Engine/Source/Runtime/OpenGLDrv/Private/OpenGLDrvPrivate.h
+++ b/Engine/Source/Runtime/OpenGLDrv/Private/OpenGLDrvPrivate.h
@@ -237,7 +237,6 @@ void PlatformRenderingContextSetup(FPlatformOpenGLDevice* Device);
 
 // Some platforms require flushing and rebinding renderbuffers after context/thread changes
 void PlatformFlushIfNeeded();
-void PlatformRebindResources(FPlatformOpenGLDevice* Device);
 
 /**
  * Set up shared context on current thread. This thread should be always set on main game thread (expect at times
diff --git a/Engine/Source/Runtime/OpenGLDrv/Private/Windows/WindowsOpenGLPlatform.cpp b/Engine/Source/Runtime/OpenGLDrv/Private/Windows/WindowsOpenGLPlatform.cpp
--- a/Engine/Source/Runtime/OpenGLDrv/Private/Windows/WindowsOpenGLPlatform.cpp
+++ b/Engine/Source/Runtime/OpenGLDrv/Private/Windows/WindowsOpenGLPlatform.cpp
@@ -535,11 +535,6 @@ void PlatformFlushIfNeeded()
 	glFinish();
 }
 
-void PlatformRebindResources(FPlatformOpenGLDevice* Device)
-{
-	// @todo: Figure out if we need to rebind frame & renderbuffers after switching contexts
-}
-
 FPlatformOpenGLContext* PlatformGetOpenGLRenderingContext(FPlatformOpenGLDevice* Device)
 {
 	return &Device->RenderingContext;
@@ -553,6 +548,7 @@ void PlatformRenderingContextSetup(FPlatformOpenGLDevice* Device)
 	{
 		glFlush();
 	}
+
 	if (Device->ViewportContexts.Num() == 1)
 	{
 		// use the HDC of the window, to reduce context swap overhead
diff --git a/Engine/Source/Runtime/OpenGLDrv/Public/OpenGLUtil.h b/Engine/Source/Runtime/OpenGLDrv/Public/OpenGLUtil.h
--- a/Engine/Source/Runtime/OpenGLDrv/Public/OpenGLUtil.h
+++ b/Engine/Source/Runtime/OpenGLDrv/Public/OpenGLUtil.h
@@ -26,7 +26,7 @@ int32 GetOGLDebugOutputLevel();
 #endif
 
 // Additional check that our GL calls are occurring on the expected thread
-#define ENABLE_VERIFY_GL_THREAD (UE_BUILD_DEBUG)
+#define ENABLE_VERIFY_GL_THREAD (!(UE_BUILD_TEST || UE_BUILD_SHIPPING))
 
 /** Set to 1 to verify that the the engine side uniform buffer layout matches the driver side of the GLSL shader*/
 #define ENABLE_UNIFORM_BUFFER_LAYOUT_VERIFICATION ( 0 & UE_BUILD_DEBUG & (!PLATFORM_ANDROID))
@@ -52,20 +52,10 @@ GLenum GetOpenGLCubeFace(ECubeFace Face);
 extern bool PlatformOpenGLContextValid();
 
 #if ENABLE_VERIFY_GL_THREAD
-	#if UE_BUILD_TEST
-		#define GLCONTEXT_CLAUSE 
-	#else
-		#define GLCONTEXT_CLAUSE PlatformOpenGLContextValid() &&
-	#endif
-	// check that the current thread has a valid context and matches our RT / RHIT expectations.
-	// Note that the game thread can access the shared context.
-	// use commandline switch -norhithread if this causes issues.
 	#define CHECK_EXPECTED_GL_THREAD() \
-		if(!( \
-			 GLCONTEXT_CLAUSE\
-			(IsInGameThread() || ( (IsInRenderingThread() && !IsRunningRHIInSeparateThread()) || (IsInRHIThread() && IsRunningRHIInSeparateThread()) ))))\
+		if (!PlatformOpenGLContextValid()) \
 		{ \
-			UE_LOG(LogRHI, Fatal, TEXT("Potential use of GL context from incorrect thread. [ValidContext = %d] && [IsInGameThread() = %d, IsInRenderingThread() && !IsRunningRHIInSeparateThread() = %d, IsInRHIThread() && IsRunningRHIInSeparateThread() = %d]"), PlatformOpenGLContextValid(), IsInGameThread(), IsInRenderingThread() && !IsRunningRHIInSeparateThread(), IsInRHIThread() && IsRunningRHIInSeparateThread());\
+			UE_LOG(LogRHI, Fatal, TEXT("Potential use of GL context from incorrect thread. [ValidContext = %d] && [IsInGameThread() = %d, IsInRenderingThread() = %d, IsInRHIThread() = %d, IsRunningRHIInSeparateThread() = %d]"), PlatformOpenGLContextValid(), IsInGameThread(), IsInRenderingThread(), IsInRHIThread(), IsRunningRHIInSeparateThread()); \
 		}
 #else
 	#define CHECK_EXPECTED_GL_THREAD() 
