Fix crash caused by race condition in OpenGL RHI's RHIAdvanceFrameFence

diff --git a/Engine/Source/Runtime/OpenGLDrv/Private/OpenGLDevice.cpp b/Engine/Source/Runtime/OpenGLDrv/Private/OpenGLDevice.cpp
--- a/Engine/Source/Runtime/OpenGLDrv/Private/OpenGLDevice.cpp
+++ b/Engine/Source/Runtime/OpenGLDrv/Private/OpenGLDevice.cpp
@@ -233,14 +233,12 @@ void FOpenGLDynamicRHI::RHIEndFrame()
 
 void FOpenGLDynamicRHI::RHIAdvanceFrameFence()
 {
-	//
-	// This function was previously RHIPerFrameRHIFlushComplete. Changed to RHIAdvanceFrameFence to be called by the render thread on RHICmdList.EndFrame().
-	// @todo dev-pr clean up threading in OpenGL so buffer / query pool management can happen on the RHI thread.
-	//
-
-	BeginFrame_UniformBufferPoolCleanup();
-	BeginFrame_VertexBufferCleanup();
-	BeginFrame_QueryBatchCleanup();
+	RunOnGLRenderContextThread([]()
+	{
+		BeginFrame_UniformBufferPoolCleanup();
+		BeginFrame_VertexBufferCleanup();
+		BeginFrame_QueryBatchCleanup();
+	});
 
 	OpenGL_PollAllFences();
 }
