Fix missing header file.

diff --git a/Engine/Source/Runtime/OpenGLDrv/Private/Android/AndroidEGL.cpp b/Engine/Source/Runtime/OpenGLDrv/Private/Android/AndroidEGL.cpp
--- a/Engine/Source/Runtime/OpenGLDrv/Private/Android/AndroidEGL.cpp
+++ b/Engine/Source/Runtime/OpenGLDrv/Private/Android/AndroidEGL.cpp
@@ -16,6 +16,7 @@
 #include "UObject/GarbageCollection.h"
 #include "Android/AndroidPlatformFramePacer.h"
 #include <dlfcn.h>
+#include "UnrealEngine.h"
 
 
 AndroidEGL* AndroidEGL::Singleton = NULL;
