Remove the DXT format support for Android OpenGL since the compressed RVT could be created with the DXT format but the compression method is ETC.
ETC is aways supported on GLES 3.1.

diff --git a/Engine/Source/Runtime/OpenGLDrv/Private/OpenGLES.cpp b/Engine/Source/Runtime/OpenGLDrv/Private/OpenGLES.cpp
--- a/Engine/Source/Runtime/OpenGLDrv/Private/OpenGLES.cpp
+++ b/Engine/Source/Runtime/OpenGLDrv/Private/OpenGLES.cpp
@@ -44,9 +44,6 @@ bool FOpenGLES::bSupportsShaderDepthStencilFetch = false;
 /** GL_EXT_multisampled_render_to_texture */
 bool FOpenGLES::bSupportsMultisampledRenderToTexture = false;
 
-/** GL_NV_texture_compression_s3tc, GL_EXT_texture_compression_s3tc */
-bool FOpenGLES::bSupportsDXT = false;
-
 /** OpenGL ES 3.0 profile */
 bool FOpenGLES::bSupportsETC2 = false;
 
@@ -179,7 +176,6 @@ void FOpenGLES::ProcessExtensions(const FString& ExtensionsString)
 	bSupportsPixelLocalStorage = ExtensionsString.Contains(TEXT("GL_EXT_shader_pixel_local_storage"));
 	bSupportsShaderDepthStencilFetch = ExtensionsString.Contains(TEXT("GL_ARM_shader_framebuffer_fetch_depth_stencil"));
 	bSupportsMultisampledRenderToTexture = ExtensionsString.Contains(TEXT("GL_EXT_multisampled_render_to_texture"));
-	bSupportsDXT = ExtensionsString.Contains(TEXT("GL_NV_texture_compression_s3tc")) || ExtensionsString.Contains(TEXT("GL_EXT_texture_compression_s3tc"));
 	bSupportsNVFrameBufferBlit = ExtensionsString.Contains(TEXT("GL_NV_framebuffer_blit"));
 	bSupportsBufferStorage = ExtensionsString.Contains(TEXT("GL_EXT_buffer_storage"));
 	bSupportsDepthClamp = ExtensionsString.Contains(TEXT("GL_EXT_depth_clamp"));
diff --git a/Engine/Source/Runtime/OpenGLDrv/Private/OpenGLES.h b/Engine/Source/Runtime/OpenGLDrv/Private/OpenGLES.h
--- a/Engine/Source/Runtime/OpenGLDrv/Private/OpenGLES.h
+++ b/Engine/Source/Runtime/OpenGLDrv/Private/OpenGLES.h
@@ -59,7 +59,7 @@ struct FOpenGLES : public FOpenGLBase
 	static FORCEINLINE bool SupportsMultisampledRenderToTexture() { return bSupportsMultisampledRenderToTexture; }
 	static FORCEINLINE bool SupportsVertexArrayBGRA() { return false; }
 	static FORCEINLINE bool SupportsBGRA8888() { return bSupportsBGRA8888; }
-	static FORCEINLINE bool SupportsDXT() { return bSupportsDXT; }
+	static FORCEINLINE bool SupportsDXT() { return false; }
 	static FORCEINLINE bool SupportsETC2() { return bSupportsETC2; }
 	static FORCEINLINE GLenum GetDepthFormat() { return GL_DEPTH_COMPONENT24; }
 	static FORCEINLINE GLenum GetShadowDepthFormat() { return GL_DEPTH_COMPONENT16; }
@@ -642,9 +642,6 @@ protected:
 	/** GL_APPLE_texture_format_BGRA8888 */
 	static bool bSupportsBGRA8888;
 
-	/** GL_NV_texture_compression_s3tc, GL_EXT_texture_compression_s3tc */
-	static bool bSupportsDXT;
-
 	/** OpenGL ES 3.0 profile */
 	static bool bSupportsETC2;
 
