Android GLES, Add cvar to disable use of GL_MAP_UNSYNCHRONIZED_BIT when locking buffers.
 - for future perf tests.

diff --git a/Engine/Source/Runtime/OpenGLDrv/Private/OpenGLCommands.cpp b/Engine/Source/Runtime/OpenGLDrv/Private/OpenGLCommands.cpp
--- a/Engine/Source/Runtime/OpenGLDrv/Private/OpenGLCommands.cpp
+++ b/Engine/Source/Runtime/OpenGLDrv/Private/OpenGLCommands.cpp
@@ -83,7 +83,12 @@ namespace OpenGLConsoleVariables
 		bUseBufferDiscard,
 		TEXT("If true, use dynamic buffer orphaning hint.")
 		);
-	
+	int32 bUseUnsynchronizedBufferMapping =1;
+	static FAutoConsoleVariableRef CVarUseUnsynchronizedBufferMapping(
+		TEXT("OpenGL.UseUnsynchronizedBufferMapping"),
+		bUseUnsynchronizedBufferMapping,
+		TEXT("If true, use GL_MAP_UNSYNCHRONIZED_BIT when mapping write only buffers.")
+	);
 	int32 bUsePersistentMappingStagingBuffer= 1;
 	static FAutoConsoleVariableRef CVarUsePersistentMappingStagingBuffer(
 		TEXT("OpenGL.UsePersistentMappingStagingBuffer"),
diff --git a/Engine/Source/Runtime/OpenGLDrv/Private/OpenGLES.h b/Engine/Source/Runtime/OpenGLDrv/Private/OpenGLES.h
--- a/Engine/Source/Runtime/OpenGLDrv/Private/OpenGLES.h
+++ b/Engine/Source/Runtime/OpenGLDrv/Private/OpenGLES.h
@@ -13,6 +13,11 @@
 #include "OpenGLBase.h"
 #include "OpenGLUtil.h"		// for VERIFY_GL
 
+namespace OpenGLConsoleVariables
+{
+	extern int32 bUseUnsynchronizedBufferMapping;
+}
+
 struct FOpenGLES : public FOpenGLBase
 {
 	static FORCEINLINE bool IsES31Usable()
@@ -160,7 +165,7 @@ struct FOpenGLES : public FOpenGLBase
 			Access = (GL_MAP_READ_BIT | GL_MAP_PERSISTENT_BIT | GL_MAP_COHERENT_BIT);
 			break;
 		case EResourceLockMode::RLM_WriteOnly:
-			Access = (GL_MAP_INVALIDATE_BUFFER_BIT | GL_MAP_WRITE_BIT | GL_MAP_UNSYNCHRONIZED_BIT);
+			Access = OpenGLConsoleVariables::bUseUnsynchronizedBufferMapping ? (GL_MAP_INVALIDATE_BUFFER_BIT | GL_MAP_WRITE_BIT | GL_MAP_UNSYNCHRONIZED_BIT) : (GL_MAP_INVALIDATE_BUFFER_BIT | GL_MAP_WRITE_BIT);
 			break;
 		case EResourceLockMode::RLM_WriteOnlyUnsynchronized:
 			Access = (GL_MAP_WRITE_BIT | GL_MAP_UNSYNCHRONIZED_BIT);
