Remove implementation of RHIBeginScene and RHIEndScene in OpenGL RHI
 - Previously these were used to set which cached context state is active.
 - This is now achieved inside GetContextStateForCurrentContext by setting a pointer which is cleared whenever the active context is changed.

Remove InvalidContextState
 - This is never used after removing the "bool bAssertIfInvalid = true" optional arg of GetContextStateForCurrentContext and making it always assert. Nowhere in the RHI called that function passing false.

Remove unused functions in OpenGL RHI
 - PlatformOpenGLCurrentContextHandle
 - PlatformLabelObjects
 - PlatformGetOpenGLRenderingContext
 - PlatformGetBackbufferDimensions

diff --git a/Engine/Source/Runtime/OpenGLDrv/Private/Android/AndroidOpenGL.cpp b/Engine/Source/Runtime/OpenGLDrv/Private/Android/AndroidOpenGL.cpp
--- a/Engine/Source/Runtime/OpenGLDrv/Private/Android/AndroidOpenGL.cpp
+++ b/Engine/Source/Runtime/OpenGLDrv/Private/Android/AndroidOpenGL.cpp
@@ -298,11 +298,6 @@ EOpenGLCurrentContext PlatformOpenGLCurrentContext(FPlatformOpenGLDevice* Device
 	return (EOpenGLCurrentContext)AndroidEGL::GetInstance()->GetCurrentContextType();
 }
 
-void* PlatformOpenGLCurrentContextHandle(FPlatformOpenGLDevice* Device)
-{
-	return AndroidEGL::GetInstance()->GetCurrentContext();
-}
-
 void PlatformRestoreDesktopDisplayMode()
 {
 }
@@ -378,11 +373,6 @@ bool PlatformOpenGLContextValid()
 	return AndroidEGL::GetInstance()->IsCurrentContextValid();
 }
 
-void PlatformGetBackbufferDimensions( uint32& OutWidth, uint32& OutHeight )
-{
-	AndroidEGL::GetInstance()->GetDimensions(OutWidth, OutHeight);
-}
-
 // =============================================================
 
 void PlatformGetNewOcclusionQuery( GLuint* OutQuery, uint64* OutQueryContext )
@@ -441,11 +431,6 @@ void FPlatformOpenGLDevice::LoadEXT()
 	glGetObjectPtrLabelKHR = (PFNGLGETOBJECTPTRLABELKHRPROC)((void*)eglGetProcAddress("glGetObjectPtrLabelKHR"));
 }
 
-FPlatformOpenGLContext* PlatformGetOpenGLRenderingContext(FPlatformOpenGLDevice* Device)
-{
-	return AndroidEGL::GetInstance()->GetRenderingContext();
-}
-
 FPlatformOpenGLContext* PlatformCreateOpenGLContext(FPlatformOpenGLDevice* Device, void* InWindowHandle)
 {
 	//Assumes Device is already initialized and context already created.
@@ -556,15 +541,6 @@ void FPlatformOpenGLDevice::SetupCurrentContext()
 	}
 }
 
-void PlatformLabelObjects()
-{
-	GLuint FrameBuffer = AndroidEGL::GetInstance()->GetResolveFrameBuffer();
-	if (FrameBuffer != 0)
-	{
-		FOpenGL::LabelObject(GL_FRAMEBUFFER, FrameBuffer, "ResolveFB");
-	}
-}
-
 //--------------------------------
 #define VIRTUALIZE_QUERIES (1)
 
diff --git a/Engine/Source/Runtime/OpenGLDrv/Private/Linux/LinuxOpenGLPlatform.cpp b/Engine/Source/Runtime/OpenGLDrv/Private/Linux/LinuxOpenGLPlatform.cpp
--- a/Engine/Source/Runtime/OpenGLDrv/Private/Linux/LinuxOpenGLPlatform.cpp
+++ b/Engine/Source/Runtime/OpenGLDrv/Private/Linux/LinuxOpenGLPlatform.cpp
@@ -262,11 +262,6 @@ void PlatformDestroyOpenGLDevice(FPlatformOpenGLDevice* Device)
 	delete Device;
 }
 
-FPlatformOpenGLContext* PlatformGetOpenGLRenderingContext(FPlatformOpenGLDevice* Device)
-{
-	return &Device->RenderingContext;
-}
-
 /**
  * Create an OpenGL context.
  */
@@ -862,27 +857,6 @@ EOpenGLCurrentContext PlatformOpenGLCurrentContext(FPlatformOpenGLDevice* Device
 	return CONTEXT_Invalid;
 }
 
-void* PlatformOpenGLCurrentContextHandle(FPlatformOpenGLDevice* Device)
-{
-	return Linux_GetCurrentContext();
-}
-
-void PlatformGetBackbufferDimensions(uint32& OutWidth, uint32& OutHeight)
-{
-	SDL_HWindow CurrentWindow = SDL_GL_GetCurrentWindow();
-
-	int Width = 0;
-	int Height = 0;
-
-	if (CurrentWindow)
-	{
-		SDL_GL_GetDrawableSize(CurrentWindow, &Width, &Height);
-	}
-
-	OutWidth = Width;
-	OutHeight = Height;
-}
-
 // =============================================================
 
 struct FOpenGLReleasedQuery
diff --git a/Engine/Source/Runtime/OpenGLDrv/Private/OpenGLCommands.cpp b/Engine/Source/Runtime/OpenGLDrv/Private/OpenGLCommands.cpp
--- a/Engine/Source/Runtime/OpenGLDrv/Private/OpenGLCommands.cpp
+++ b/Engine/Source/Runtime/OpenGLDrv/Private/OpenGLCommands.cpp
@@ -2993,12 +2993,12 @@ void FOpenGLDynamicRHI::RHIPostExternalCommandsReset()
 	if((int)RCS.ElementArrayBufferBound >= 0) glBindBuffer(GL_ELEMENT_ARRAY_BUFFER, RCS.ElementArrayBufferBound);
 	glBindBuffer(GL_PIXEL_UNPACK_BUFFER, RCS.PixelUnpackBufferBound);
 	glBindBuffer(GL_UNIFORM_BUFFER, RCS.UniformBufferBound);
-	SCS.BlendState = RCS.BlendState = InvalidContextState.BlendState;
+	SCS.BlendState = RCS.BlendState = {};
 	glActiveTexture(GL_TEXTURE0);
 	SCS.ActiveTexture = RCS.ActiveTexture = 0;
-	SCS.RasterizerState = RCS.RasterizerState = InvalidContextState.RasterizerState;
+	SCS.RasterizerState = RCS.RasterizerState = {};
 	UpdateRasterizerStateInOpenGLContext(RCS);
-	SCS.ActiveStreamMask = RCS.ActiveStreamMask = InvalidContextState.ActiveStreamMask;
+	SCS.ActiveStreamMask = RCS.ActiveStreamMask = 0;
 	for (GLuint UniformBufferIndex = 0; UniformBufferIndex < CrossCompiler::NUM_SHADER_STAGES * OGL_MAX_UNIFORM_BUFFER_BINDINGS; UniformBufferIndex++)
 	{
 		SCS.UniformBuffers[UniformBufferIndex] = FOpenGLCachedUniformBuffer_Invalid;	// that'll enforce state update on next cache test
diff --git a/Engine/Source/Runtime/OpenGLDrv/Private/OpenGLDevice.cpp b/Engine/Source/Runtime/OpenGLDrv/Private/OpenGLDevice.cpp
--- a/Engine/Source/Runtime/OpenGLDrv/Private/OpenGLDevice.cpp
+++ b/Engine/Source/Runtime/OpenGLDrv/Private/OpenGLDevice.cpp
@@ -181,32 +181,40 @@ extern void BeginFrame_QueryBatchCleanup();
 extern void OpenGL_PollAllFences();
 
 
-FOpenGLContextState& FOpenGLDynamicRHI::GetContextStateForCurrentContext(bool bAssertIfInvalid)
+FOpenGLContextState& FOpenGLDynamicRHI::GetContextStateForCurrentContext()
 {
-	// most common case
-	if (BeginSceneContextType == CONTEXT_Rendering)
+	if (!CachedContextState)
 	{
-		return RenderingContextState;
-	}
-	
-	int32 ContextType = (int32)PlatformOpenGLCurrentContext(PlatformDevice);
-	if (bAssertIfInvalid)
-	{
-			check(ContextType >= 0);
-	}
-	else if (ContextType < 0)
-	{
-		return InvalidContextState;
-	}
+		int32 ContextType = (int32)PlatformOpenGLCurrentContext(PlatformDevice);
+		checkf(ContextType >= 0, TEXT("Invalid GL context on current thread."));
 
-	if (ContextType == CONTEXT_Rendering)
-	{
-		return RenderingContextState;
+		if (ContextType == CONTEXT_Rendering)
+		{
+			CachedContextState = &RenderingContextState;
+		}
+		else
+		{
+			CachedContextState = &SharedContextState;
+		}
 	}
 	else
 	{
-		return SharedContextState;
+	#if DO_CHECK
+		switch ((int32)PlatformOpenGLCurrentContext(PlatformDevice))
+		{
+		default:
+		case CONTEXT_Other:
+		case CONTEXT_Invalid:
+			checkf(false, TEXT("Invalid GL context on current thread."));
+			break;
+
+		case CONTEXT_Rendering: check(CachedContextState == &RenderingContextState); break;
+		case CONTEXT_Shared   : check(CachedContextState == &SharedContextState   ); break;
+		}
+	#endif // DO_CHECK
 	}
+
+	return *CachedContextState;
 }
 
 void FOpenGLDynamicRHI::RHIBeginFrame()
@@ -250,12 +258,10 @@ void FOpenGLDynamicRHI::RHIAdvanceFrameFence()
 
 void FOpenGLDynamicRHI::RHIBeginScene()
 {
-	BeginSceneContextType = (int32)PlatformOpenGLCurrentContext(PlatformDevice);
 }
 
 void FOpenGLDynamicRHI::RHIEndScene()
 {
-	BeginSceneContextType = CONTEXT_Other;
 }
 
 #if PLATFORM_ANDROID
@@ -1296,7 +1302,6 @@ FDynamicRHI* FOpenGLDynamicRHIModule::CreateRHI(ERHIFeatureLevel::Type InRequest
 FOpenGLDynamicRHI::FOpenGLDynamicRHI()
 :	bRevertToSharedContextAfterDrawingViewport(false)
 ,	bIsRenderingContextAcquired(false)
-,   BeginSceneContextType(CONTEXT_Other)
 ,	PlatformDevice(NULL)
 #if (RHI_NEW_GPU_PROFILER == 0)
 ,	GPUProfilingData(this)
@@ -1556,6 +1561,7 @@ void FOpenGLDynamicRHI::RHIAcquireThreadOwnership()
 {
 	check(!bRevertToSharedContextAfterDrawingViewport);	// if this is true, then main thread is rendering using our context right now.
 	PlatformRenderingContextSetup(PlatformDevice);
+	CachedContextState = nullptr;
 
 	bIsRenderingContextAcquired = true;
 	VERIFY_GL(RHIAcquireThreadOwnership);
@@ -1567,6 +1573,7 @@ void FOpenGLDynamicRHI::RHIReleaseThreadOwnership()
 	bIsRenderingContextAcquired = false;
 
 	PlatformNULLContextSetup();
+	CachedContextState = nullptr;
 }
 
 void FOpenGLDynamicRHI::RegisterQuery( FOpenGLRenderQuery* Query )
@@ -1603,11 +1610,6 @@ void FOpenGLDynamicRHI::InvalidateQueries( void )
 	}
 }
 
-void* FOpenGLDynamicRHI::GetOpenGLCurrentContextHandle()
-{
-	return PlatformOpenGLCurrentContextHandle(PlatformDevice);
-}
-
 void FOpenGLDynamicRHI::SetCustomPresent(FRHICustomPresent* InCustomPresent)
 {
 	FScopeLock lock(&CustomPresentSection);
diff --git a/Engine/Source/Runtime/OpenGLDrv/Private/OpenGLDrv.h b/Engine/Source/Runtime/OpenGLDrv/Private/OpenGLDrv.h
--- a/Engine/Source/Runtime/OpenGLDrv/Private/OpenGLDrv.h
+++ b/Engine/Source/Runtime/OpenGLDrv/Private/OpenGLDrv.h
@@ -525,7 +525,7 @@ public:
 	void CachedSetupTextureStageInner(FOpenGLContextState& ContextState, GLint TextureIndex, GLenum Target, GLuint Resource, GLint BaseMip, GLint NumMips);
 	void CachedSetupUAVStage(FOpenGLContextState& ContextState, GLint UAVIndex, GLenum Format, GLuint Resource, bool bLayered, GLint Layer, GLenum Access);
 	void UpdateSRV(FOpenGLShaderResourceView* SRV);
-	FOpenGLContextState& GetContextStateForCurrentContext(bool bAssertIfInvalid = true);
+	FOpenGLContextState& GetContextStateForCurrentContext();
 
 	FORCEINLINE void CachedBindArrayBuffer( FOpenGLContextState& ContextState, GLuint Buffer )
 	{
@@ -601,8 +601,6 @@ public:
 
 	void InitializeGLTexture(FOpenGLTexture* Texture, const void* BulkDataPtr, uint64 BulkDataSize);
 
-	void* GetOpenGLCurrentContextHandle();
-
 	void SetCustomPresent(class FRHICustomPresent* InCustomPresent);
 
 #define RHITHREAD_GLTRACE 1
@@ -826,11 +824,10 @@ private:
 	TGlobalResource< TBoundShaderStateHistory<10000> > BoundShaderStateHistory;
 
 	/** Per-context state caching */
-	FOpenGLContextState InvalidContextState;
 	FOpenGLContextState	SharedContextState;
 	FOpenGLContextState	RenderingContextState;
 	// Cached context type on BeginScene
-	int32 BeginSceneContextType;
+	FOpenGLContextState* CachedContextState = nullptr;
 
 	template <typename TRHIShader>
 	void ApplyStaticUniformBuffers(TRHIShader* Shader);
diff --git a/Engine/Source/Runtime/OpenGLDrv/Private/OpenGLDrvPrivate.h b/Engine/Source/Runtime/OpenGLDrv/Private/OpenGLDrvPrivate.h
--- a/Engine/Source/Runtime/OpenGLDrv/Private/OpenGLDrvPrivate.h
+++ b/Engine/Source/Runtime/OpenGLDrv/Private/OpenGLDrvPrivate.h
@@ -162,11 +162,6 @@ int32 PlatformGlGetError();
  */
 EOpenGLCurrentContext PlatformOpenGLCurrentContext(FPlatformOpenGLDevice* Device);
 
-/**
- * Retrieve the address of the OpenGL context currently in use.
- */
-void* PlatformOpenGLCurrentContextHandle(FPlatformOpenGLDevice* Device);
-
 /**
  * Get new occlusion query from current context. This is provided from a cache inside the context
  * if some entries are in there, and created otherwise. All other released queries present in the
@@ -203,11 +198,6 @@ FPlatformOpenGLDevice* PlatformCreateOpenGLDevice();
 */
 bool PlatformCanEnableGPUCapture();
 
-/**
- * Label Objects. Needs a separate function because label GLSL api procedure would be loaded later down the line, and we need to label objects after that.
- */
-void PlatformLabelObjects();
-
 /**
  * Destroy the OpenGL device for single OpenGL RHI. This should happen when destroying the RHI.
  */
@@ -219,11 +209,6 @@ void PlatformDestroyOpenGLDevice(FPlatformOpenGLDevice* Device);
  */
 FPlatformOpenGLContext* PlatformCreateOpenGLContext(FPlatformOpenGLDevice* Device, void* InWindowHandle);
 
-/**
- * Get rendering OpenGL context on current thread.
- */
-FPlatformOpenGLContext* PlatformGetOpenGLRenderingContext(FPlatformOpenGLDevice* Device);
-
 /**
  * Destroy a viewport OpenGL context.
  */
@@ -280,11 +265,6 @@ bool PlatformGetAvailableResolutions(FScreenResolutionArray& Resolutions, bool b
  */
 void PlatformRestoreDesktopDisplayMode();
 
-/**
- * Get current backbuffer dimensions.
- */
-void PlatformGetBackbufferDimensions( uint32& OutWidth, uint32& OutHeight );
-
 // Returns native window handle.
 void* PlatformGetWindow(FPlatformOpenGLContext* Context, void** AddParam);
 
diff --git a/Engine/Source/Runtime/OpenGLDrv/Private/OpenGLRenderTarget.cpp b/Engine/Source/Runtime/OpenGLDrv/Private/OpenGLRenderTarget.cpp
--- a/Engine/Source/Runtime/OpenGLDrv/Private/OpenGLRenderTarget.cpp
+++ b/Engine/Source/Runtime/OpenGLDrv/Private/OpenGLRenderTarget.cpp
@@ -756,12 +756,10 @@ void FOpenGLDynamicRHI::RHIReadSurfaceData(FRHITexture* TextureRHI,FIntRect Rect
 	TArray<uint8> Temp;
 
 	FOpenGLContextState& ContextState = GetContextStateForCurrentContext();
-	if (&ContextState != &InvalidContextState)
-	{
-		ReadSurfaceDataRaw(ContextState, TextureRHI, Rect, Temp, InFlags);
 
-		FMemory::Memcpy(OutData.GetData(), Temp.GetData(), Size * sizeof(FColor));
-	}
+	ReadSurfaceDataRaw(ContextState, TextureRHI, Rect, Temp, InFlags);
+	FMemory::Memcpy(OutData.GetData(), Temp.GetData(), Size * sizeof(FColor));
+
 	RHITHREAD_GLCOMMAND_EPILOGUE();
 }
 
diff --git a/Engine/Source/Runtime/OpenGLDrv/Private/OpenGLViewport.cpp b/Engine/Source/Runtime/OpenGLDrv/Private/OpenGLViewport.cpp
--- a/Engine/Source/Runtime/OpenGLDrv/Private/OpenGLViewport.cpp
+++ b/Engine/Source/Runtime/OpenGLDrv/Private/OpenGLViewport.cpp
@@ -129,6 +129,7 @@ void FOpenGLDynamicRHI::RHIBeginDrawingViewport(FRHIViewport* ViewportRHI, FRHIT
 
 		bRevertToSharedContextAfterDrawingViewport = true;
 		PlatformRenderingContextSetup(PlatformDevice);
+		CachedContextState = nullptr;
 	}
 
 	// Set the render target and viewport.
@@ -214,6 +215,8 @@ void FOpenGLDynamicRHI::RHIEndDrawingViewport(FRHIViewport* ViewportRHI,bool bPr
 		if (bRevertToSharedContextAfterDrawingViewport)
 		{
 			PlatformSharedContextSetup(PlatformDevice);
+			CachedContextState = nullptr;
+
 			bRevertToSharedContextAfterDrawingViewport = false;
 		}
 	}
diff --git a/Engine/Source/Runtime/OpenGLDrv/Private/Windows/WindowsOpenGLPlatform.cpp b/Engine/Source/Runtime/OpenGLDrv/Private/Windows/WindowsOpenGLPlatform.cpp
--- a/Engine/Source/Runtime/OpenGLDrv/Private/Windows/WindowsOpenGLPlatform.cpp
+++ b/Engine/Source/Runtime/OpenGLDrv/Private/Windows/WindowsOpenGLPlatform.cpp
@@ -535,11 +535,6 @@ void PlatformFlushIfNeeded()
 	glFinish();
 }
 
-FPlatformOpenGLContext* PlatformGetOpenGLRenderingContext(FPlatformOpenGLDevice* Device)
-{
-	return &Device->RenderingContext;
-}
-
 void PlatformRenderingContextSetup(FPlatformOpenGLDevice* Device)
 {
 	check(Device && Device->RenderingContext.DeviceContext && Device->RenderingContext.OpenGLContext);
@@ -867,22 +862,6 @@ EOpenGLCurrentContext PlatformOpenGLCurrentContext(FPlatformOpenGLDevice* Device
 	}
 }
 
-void* PlatformOpenGLCurrentContextHandle(FPlatformOpenGLDevice* Device)
-{
-	return GetCurrentContext();
-}
-
-void PlatformGetBackbufferDimensions(uint32& OutWidth, uint32& OutHeight)
-{
-	OutWidth = OutHeight = 0;
-	HDC DeviceContext = wglGetCurrentDC();
-	if (DeviceContext)
-	{
-		OutWidth = GetDeviceCaps(DeviceContext, HORZRES);
-		OutHeight = GetDeviceCaps(DeviceContext, VERTRES);
-	}
-}
-
 // =============================================================
 
 struct FOpenGLReleasedQuery
