Fix Android NonPCH compilation due to missing includes

diff --git a/Engine/Source/Runtime/OpenGLDrv/Private/Android/AndroidOpenGLFramePacer.cpp b/Engine/Source/Runtime/OpenGLDrv/Private/Android/AndroidOpenGLFramePacer.cpp
--- a/Engine/Source/Runtime/OpenGLDrv/Private/Android/AndroidOpenGLFramePacer.cpp
+++ b/Engine/Source/Runtime/OpenGLDrv/Private/Android/AndroidOpenGLFramePacer.cpp
@@ -18,6 +18,7 @@
 #include "swappy/swappy_common.h"
 #include "HAL/Thread.h"
 #include "Misc/ScopeRWLock.h"
+#include "EngineGlobals.h"
 
 struct FSwappyThreadManager : public SwappyThreadFunctions
 {
