OpenGL: Do not glFlush on every RHICmdList translation as it happens way too many times in a frame and is unnecessary

diff --git a/Engine/Source/Runtime/OpenGLDrv/Private/OpenGLCommands.cpp b/Engine/Source/Runtime/OpenGLDrv/Private/OpenGLCommands.cpp
--- a/Engine/Source/Runtime/OpenGLDrv/Private/OpenGLCommands.cpp
+++ b/Engine/Source/Runtime/OpenGLDrv/Private/OpenGLCommands.cpp
@@ -2918,10 +2918,10 @@ IRHIPlatformCommandList* FOpenGLDynamicRHI::RHIFinalizeContext(FRHIFinalizeConte
 	check(Args.Context == this);
 
 	// Flush the context to ensure recorded commands will reach the driver.
-	if (PlatformOpenGLContextValid())
-	{
-		FOpenGL::Flush();
-	}
+	//if (PlatformOpenGLContextValid())
+	//{
+	//	FOpenGL::Flush(); // This gets called way too many times in a frame, which is unnecessary 
+	//}
 
 	// Clear some context state
 	FMemory::Memset(PendingState.BoundUniformBuffers, 0, sizeof(PendingState.BoundUniformBuffers));
