OpenGL: Do not wrap compute dispatches into memory barriers
GLES barriers seem to do nothing

diff --git a/Engine/Source/Runtime/OpenGLDrv/Private/OpenGLCommands.cpp b/Engine/Source/Runtime/OpenGLDrv/Private/OpenGLCommands.cpp
--- a/Engine/Source/Runtime/OpenGLDrv/Private/OpenGLCommands.cpp
+++ b/Engine/Source/Runtime/OpenGLDrv/Private/OpenGLCommands.cpp
@@ -2833,9 +2833,7 @@ void FOpenGLDynamicRHI::RHIDispatchComputeShader(uint32 ThreadGroupCountX, uint3
 	SetupUAVsForCompute(ContextState, ComputeShader);
 	CommitComputeShaderConstants(ComputeShader);
 	
-	FOpenGL::MemoryBarrier(GL_ALL_BARRIER_BITS);
 	FOpenGL::DispatchCompute(ThreadGroupCountX, ThreadGroupCountY, ThreadGroupCountZ);
-	FOpenGL::MemoryBarrier(GL_ALL_BARRIER_BITS);
 }
 
 void FOpenGLDynamicRHI::RHIDispatchIndirectComputeShader(FRHIBuffer* ArgumentBufferRHI, uint32 ArgumentOffset)
@@ -2862,15 +2860,11 @@ void FOpenGLDynamicRHI::RHIDispatchIndirectComputeShader(FRHIBuffer* ArgumentBuf
 	
 	CommitComputeShaderConstants(ComputeShader);
 	
-	FOpenGL::MemoryBarrier(GL_ALL_BARRIER_BITS);
-
 	glBindBuffer( GL_DISPATCH_INDIRECT_BUFFER, ArgumentBuffer->Resource);
 	
 	FOpenGL::DispatchComputeIndirect(ArgumentOffset);
 
 	glBindBuffer( GL_DISPATCH_INDIRECT_BUFFER, 0);
-	
-	FOpenGL::MemoryBarrier(GL_ALL_BARRIER_BITS);
 }
 
 void FOpenGLDynamicRHI::RHISetMultipleViewports(uint32 Count, const FViewportBounds* Data)
