OpenGL: Restore memory barriers around compute dispatches, as there are render flickering issues on Windows without them
Disable them only for OpenGL ES
#rnx

diff --git a/Engine/Source/Runtime/OpenGLDrv/Private/OpenGLCommands.cpp b/Engine/Source/Runtime/OpenGLDrv/Private/OpenGLCommands.cpp
--- a/Engine/Source/Runtime/OpenGLDrv/Private/OpenGLCommands.cpp
+++ b/Engine/Source/Runtime/OpenGLDrv/Private/OpenGLCommands.cpp
@@ -2833,7 +2833,9 @@ void FOpenGLDynamicRHI::RHIDispatchComputeShader(uint32 ThreadGroupCountX, uint3
 	SetupUAVsForCompute(ContextState, ComputeShader);
 	CommitComputeShaderConstants(ComputeShader);
 	
+	FOpenGL::MemoryBarrier(GL_ALL_BARRIER_BITS);
 	FOpenGL::DispatchCompute(ThreadGroupCountX, ThreadGroupCountY, ThreadGroupCountZ);
+	FOpenGL::MemoryBarrier(GL_ALL_BARRIER_BITS);
 }
 
 void FOpenGLDynamicRHI::RHIDispatchIndirectComputeShader(FRHIBuffer* ArgumentBufferRHI, uint32 ArgumentOffset)
@@ -2862,7 +2864,9 @@ void FOpenGLDynamicRHI::RHIDispatchIndirectComputeShader(FRHIBuffer* ArgumentBuf
 	
 	glBindBuffer( GL_DISPATCH_INDIRECT_BUFFER, ArgumentBuffer->Resource);
 	
+	FOpenGL::MemoryBarrier(GL_ALL_BARRIER_BITS);
 	FOpenGL::DispatchComputeIndirect(ArgumentOffset);
+	FOpenGL::MemoryBarrier(GL_ALL_BARRIER_BITS);
 
 	glBindBuffer( GL_DISPATCH_INDIRECT_BUFFER, 0);
 }
diff --git a/Engine/Source/Runtime/OpenGLDrv/Private/OpenGLES.h b/Engine/Source/Runtime/OpenGLDrv/Private/OpenGLES.h
--- a/Engine/Source/Runtime/OpenGLDrv/Private/OpenGLES.h
+++ b/Engine/Source/Runtime/OpenGLDrv/Private/OpenGLES.h
@@ -417,7 +417,8 @@ struct FOpenGLES : public FOpenGLBase
 
 	static FORCEINLINE void MemoryBarrier(GLbitfield Barriers)
 	{
-		glMemoryBarrier(Barriers);
+		// It seems like barriers do nothing on a GLES drivers
+		//glMemoryBarrier(Barriers);
 	}
 
 	static FORCEINLINE void DispatchCompute(GLuint NumGroupsX, GLuint NumGroupsY, GLuint NumGroupsZ)
