Remove RHIBeginFrame
 - We now only have RHIEndFrame to mark the boundary between engine frames for the purpose of stat gathering and RHI cleanup tasks. The next frame begins immediately after the EndFrame command.
 - RHIEndFrame is called on the RHI thread after all prior command list submissions. Platform RHIs that need to enqueue RHICmdList work can override RHIEndFrame_RenderThread to add this work either side of the actual EndFrame command.
 - The work that platform RHIs performed in RHIBeginFrame has been merged into their implementations of RHIEndFrame.

diff --git a/Engine/Source/Developer/SourceControl/Private/ScopedSourceControlProgress.cpp b/Engine/Source/Developer/SourceControl/Private/ScopedSourceControlProgress.cpp
--- a/Engine/Source/Developer/SourceControl/Private/ScopedSourceControlProgress.cpp
+++ b/Engine/Source/Developer/SourceControl/Private/ScopedSourceControlProgress.cpp
@@ -236,15 +236,9 @@ void FScopedSourceControlProgress::Tick()
 #if SOURCE_CONTROL_WITH_SLATE
 	if (!(FApp::IsUnattended() || IsRunningCommandlet()) && WindowPtr.IsValid() && FSlateApplication::Get().CanDisplayWindows())
 	{
-		// Testing if we are already ticking the rendering. That is to prevent a double "BeginFrame" in case the user wrongly uses the FSlateApplication::OnPreTick to start a slow task.
+		// Testing if we are already ticking the rendering. That is to prevent a double "EndFrame" in case the user wrongly uses the FSlateApplication::OnPreTick to start a slow task.
 		bool bIsTicking = FSlateApplication::Get().IsTicking();
 
-		// Mark begin frame.
-		if (!bIsTicking && GIsRHIInitialized)
-		{
-			ENQUEUE_RENDER_COMMAND(BeginFrameCmd)([](FRHICommandListImmediate& RHICmdList) { RHICmdList.BeginFrame(); });
-		}
-
 		// Tick Slate application
 		FSlateApplication::Get().Tick();
 
diff --git a/Engine/Source/Editor/UnrealEd/Private/FeedbackContextEditor.cpp b/Engine/Source/Editor/UnrealEd/Private/FeedbackContextEditor.cpp
--- a/Engine/Source/Editor/UnrealEd/Private/FeedbackContextEditor.cpp
+++ b/Engine/Source/Editor/UnrealEd/Private/FeedbackContextEditor.cpp
@@ -424,12 +424,6 @@ static void TickSlate(TSharedPtr<SWindow> SlowTaskWindow)
 		// Testing if we are already ticking the rendering. That is to prevent a double "BeginFrame" in case the user wrongly uses the FSlateApplication::OnPreTick to start a slow task.
 		bool bIsTicking = FSlateApplication::Get().IsTicking();
 
-		// Mark begin frame
-		if (!bIsTicking && GIsRHIInitialized)
-		{
-			ENQUEUE_RENDER_COMMAND(BeginFrameCmd)([](FRHICommandListImmediate& RHICmdList) { RHICmdList.BeginFrame(); });
-		}
-
 		// Tick Slate application
 		FSlateApplication::Get().Tick();
 
diff --git a/Engine/Source/Runtime/Apple/MetalRHI/Private/MetalDynamicRHI.h b/Engine/Source/Runtime/Apple/MetalRHI/Private/MetalDynamicRHI.h
--- a/Engine/Source/Runtime/Apple/MetalRHI/Private/MetalDynamicRHI.h
+++ b/Engine/Source/Runtime/Apple/MetalRHI/Private/MetalDynamicRHI.h
@@ -41,6 +41,9 @@ public:
 	virtual const TCHAR* GetName() override { return TEXT("Metal"); }
 	virtual ERHIInterfaceType GetInterfaceType() const override { return ERHIInterfaceType::Metal; }
 
+	virtual void RHIEndFrame_RenderThread(FRHICommandListImmediate& RHICmdList) final override;
+	virtual void RHIEndFrame() final override;
+
 	virtual FRHIShaderLibraryRef RHICreateShaderLibrary(EShaderPlatform Platform, FString const& FilePath, FString const& Name) final override;
 	virtual FSamplerStateRHIRef RHICreateSamplerState(const FSamplerStateInitializerRHI& Initializer) final override;
 	virtual FRasterizerStateRHIRef RHICreateRasterizerState(const FRasterizerStateInitializerRHI& Initializer) final override;
diff --git a/Engine/Source/Runtime/Apple/MetalRHI/Private/MetalProfiler.cpp b/Engine/Source/Runtime/Apple/MetalRHI/Private/MetalProfiler.cpp
--- a/Engine/Source/Runtime/Apple/MetalRHI/Private/MetalProfiler.cpp
+++ b/Engine/Source/Runtime/Apple/MetalRHI/Private/MetalProfiler.cpp
@@ -594,20 +594,10 @@ void FMetalProfiler::BeginFrame()
 	}
 	
 	FMetalGPUProfiler::BeginFrame();
-	
-	if (MetalGPUProfilerIsInSafeThread() && GetEmitDrawEvents())
-	{
-		PushEvent(TEXT("FRAME"), FColor(0, 255, 0, 255));
-	}
 }
 
 void FMetalProfiler::EndFrame()
-{
-	if (MetalGPUProfilerIsInSafeThread() && GetEmitDrawEvents())
-	{
-		PopEvent();
-	}
-	
+{	
 	FMetalGPUProfiler::EndFrame();
 	
 	if (MetalGPUProfilerIsInSafeThread() && bEnabled)
diff --git a/Engine/Source/Runtime/Apple/MetalRHI/Private/MetalRHI.cpp b/Engine/Source/Runtime/Apple/MetalRHI/Private/MetalRHI.cpp
--- a/Engine/Source/Runtime/Apple/MetalRHI/Private/MetalRHI.cpp
+++ b/Engine/Source/Runtime/Apple/MetalRHI/Private/MetalRHI.cpp
@@ -1226,22 +1226,36 @@ void FMetalDynamicRHI::Init()
 	GIsRHIInitialized = true;
 }
 
-void FMetalRHICommandContext::RHIBeginFrame()
+void FMetalDynamicRHI::RHIEndFrame_RenderThread(FRHICommandListImmediate& RHICmdList)
 {
-    MTL_SCOPED_AUTORELEASE_POOL;
+	RHICmdList.EnqueueLambdaMultiPipe(ERHIPipeline::Graphics, FRHICommandListBase::EThreadFence::Enabled, TEXT("Metal EndFrame"),
+		[this](FMetalContextArray const& Contexts)
+	{
+		MTL_SCOPED_AUTORELEASE_POOL;
+
 #if ENABLE_METAL_GPUPROFILE
-    Profiler->BeginFrame();
+		Contexts[ERHIPipeline::Graphics]->GetProfiler()->EndFrame();
 #endif
-    ((FMetalDeviceContext*)Context)->BeginFrame();
-}
+		Contexts[ERHIPipeline::Graphics]->GetInternalContext().EndFrame();
+	});
+
+	FDynamicRHI::RHIEndFrame_RenderThread(RHICmdList);
+
+	RHICmdList.EnqueueLambdaMultiPipe(ERHIPipeline::Graphics, FRHICommandListBase::EThreadFence::Enabled, TEXT("Metal BeginFrame"),
+		[this](FMetalContextArray const& Contexts)
+	{
+		MTL_SCOPED_AUTORELEASE_POOL;
 
-void FMetalRHICommandContext::RHIEndFrame()
-{
-    MTL_SCOPED_AUTORELEASE_POOL;
 #if ENABLE_METAL_GPUPROFILE
-    Profiler->EndFrame();
+		Contexts[ERHIPipeline::Graphics]->GetProfiler()->BeginFrame();
 #endif
-    ((FMetalDeviceContext*)Context)->EndFrame();
+		Contexts[ERHIPipeline::Graphics]->GetInternalContext().BeginFrame();
+	});
+}
+
+void FMetalDynamicRHI::RHIEndFrame()
+{
+	// Currently we have nothing to do
 }
 
 #if WITH_RHI_BREADCRUMBS
diff --git a/Engine/Source/Runtime/Apple/MetalRHI/Private/MetalRHIContext.cpp b/Engine/Source/Runtime/Apple/MetalRHI/Private/MetalRHIContext.cpp
--- a/Engine/Source/Runtime/Apple/MetalRHI/Private/MetalRHIContext.cpp
+++ b/Engine/Source/Runtime/Apple/MetalRHI/Private/MetalRHIContext.cpp
@@ -242,3 +242,27 @@ void FMetalRHICommandContext::RHIEndOcclusionQueryBatch()
 	check(CommandBufferFence.IsValid());
 	CommandBufferFence.Reset();
 }
+
+FMetalContextArray::FMetalContextArray(FRHIContextArray const& Contexts)
+	: TRHIPipelineArray(InPlace, nullptr)
+{
+	for (ERHIPipeline Pipeline : MakeFlagsRange(ERHIPipeline::All))
+	{
+		IRHIComputeContext* Context = Contexts[Pipeline];
+
+		switch (Pipeline)
+		{
+		default:
+			checkNoEntry();
+			break;
+
+		case ERHIPipeline::Graphics:
+			(*this)[Pipeline] = Context ? static_cast<FMetalRHICommandContext*>(&Context->GetLowestLevelContext()) : nullptr;
+			break;
+
+		case ERHIPipeline::AsyncCompute:
+			(*this)[Pipeline] = Context ? static_cast<FMetalRHICommandContext*>(&Context->GetLowestLevelContext()) : nullptr;
+			break;
+		}
+	}
+}
diff --git a/Engine/Source/Runtime/Apple/MetalRHI/Public/MetalRHIContext.h b/Engine/Source/Runtime/Apple/MetalRHI/Public/MetalRHIContext.h
--- a/Engine/Source/Runtime/Apple/MetalRHI/Public/MetalRHIContext.h
+++ b/Engine/Source/Runtime/Apple/MetalRHI/Public/MetalRHIContext.h
@@ -68,12 +68,6 @@ public:
 	// This method is queued with an RHIThread, otherwise it will flush after it is queued; without an RHI thread there is no benefit to queuing this frame advance commands
 	virtual void RHIEndDrawingViewport(FRHIViewport* Viewport, bool bPresent, bool bLockToVsync) override;
 	
-	// This method is queued with an RHIThread, otherwise it will flush after it is queued; without an RHI thread there is no benefit to queuing this frame advance commands
-	virtual void RHIBeginFrame() override;
-	
-	// This method is queued with an RHIThread, otherwise it will flush after it is queued; without an RHI thread there is no benefit to queuing this frame advance commands
-	virtual void RHIEndFrame() override;
-	
 	virtual void RHISetStreamSource(uint32 StreamIndex, FRHIBuffer* VertexBuffer, uint32 Offset) final override;
 
 	// @param MinX including like Win32 RECT
@@ -229,3 +223,8 @@ public:
 protected:
 	friend class FMetalDynamicRHI;
 };
+
+struct FMetalContextArray : public TRHIPipelineArray<FMetalRHICommandContext*>
+{
+	FMetalContextArray(FRHIContextArray const& Contexts);
+};
diff --git a/Engine/Source/Runtime/D3D12RHI/Private/D3D12Allocation.cpp b/Engine/Source/Runtime/D3D12RHI/Private/D3D12Allocation.cpp
--- a/Engine/Source/Runtime/D3D12RHI/Private/D3D12Allocation.cpp
+++ b/Engine/Source/Runtime/D3D12RHI/Private/D3D12Allocation.cpp
@@ -1614,7 +1614,8 @@ void FD3D12DefaultBufferAllocator::BeginFrame(FD3D12ContextArray const& Contexts
 	}
 
 	{
-		FD3D12CommandContext& CommandContext = GetParentDevice()->GetDefaultCommandContext();
+		// @todo dev-pr - explicit use of graphics context - nothing is synchronizing async compute - needs refactor
+		FD3D12CommandContext& CommandContext = *Contexts[ERHIPipeline::Graphics]->GetSingleDeviceContext(GetParentDevice()->GetGPUIndex());
 
 		TRACE_CPUPROFILER_EVENT_SCOPE(FlushPendingBufferCopyOps);
 		RHI_BREADCRUMB_EVENT(CommandContext, BufferPoolCopyOps);
@@ -1813,7 +1814,8 @@ void FD3D12TextureAllocatorPool::BeginFrame(FD3D12ContextArray const& Contexts)
 	}
 
 	{
-		FD3D12CommandContext& CommandContext = GetParentDevice()->GetDefaultCommandContext();
+		// @todo dev-pr - explicit use of graphics context - nothing is synchronizing async compute - needs refactor
+		FD3D12CommandContext& CommandContext = *Contexts[ERHIPipeline::Graphics]->GetSingleDeviceContext(GetParentDevice()->GetGPUIndex());
 
 		TRACE_CPUPROFILER_EVENT_SCOPE(FlushPendingTextureCopyOps);
 		RHI_BREADCRUMB_EVENT(CommandContext, TexturePoolCopyOps);
diff --git a/Engine/Source/Runtime/D3D12RHI/Private/D3D12CommandContext.cpp b/Engine/Source/Runtime/D3D12RHI/Private/D3D12CommandContext.cpp
--- a/Engine/Source/Runtime/D3D12RHI/Private/D3D12CommandContext.cpp
+++ b/Engine/Source/Runtime/D3D12RHI/Private/D3D12CommandContext.cpp
@@ -598,6 +598,11 @@ FD3D12SyncPoint* FD3D12CopyScope::GetSyncPoint() const
 	return SyncPoint;
 }
 
+void FD3D12ContextCommon::NewPayload()
+{
+	Payloads.Add(new FD3D12Payload(Device->GetQueue(QueueType)));
+}
+
 void FD3D12ContextCommon::FlushCommands(ED3D12FlushFlags FlushFlags)
 {
 	// We should only be flushing the default context
@@ -651,34 +656,6 @@ void FD3D12ContextCommon::ConditionalSplitCommandList()
 	}
 }
 
-void FD3D12DynamicRHI::RHIBeginFrame(FRHICommandListImmediate& RHICmdList)
-{
-	RHICmdList.EnqueueLambdaMultiPipe(GetEnabledRHIPipelines(), FRHICommandListBase::EThreadFence::Enabled, TEXT("FD3D12DynamicRHI::RHIBeginFrame"),
-		[this](FD3D12ContextArray const& Contexts)
-		{
-			for (auto& Adapter : ChosenAdapters)
-			{
-				for (auto& Device : Adapter->GetDevices())
-				{
-#if (RHI_NEW_GPU_PROFILER == 0)
-					Device->GetGPUProfiler().BeginFrame();
-#endif
-
-					Device->GetDefaultBufferAllocator().BeginFrame(Contexts);
-					Device->GetTextureAllocator().BeginFrame(Contexts);
-				}
-			}
-		}
-	);
-}
-
-void FD3D12CommandContext::RHIBeginFrame()
-{
-#if D3D12_RHI_RAYTRACING
-	Device->GetRayTracingCompactionRequestHandler()->Update(*this);
-#endif // D3D12_RHI_RAYTRACING
-}
-
 void FD3D12CommandContext::ClearState(EClearStateMode Mode)
 {
 	StateCache.ClearState();
@@ -726,80 +703,35 @@ void FD3D12CommandContext::ClearAllShaderResources()
 	StateCache.ClearSRVs();
 }
 
-void FD3D12CommandContextBase::RHIEndFrame()
-{
-	FD3D12Device* Device = ParentAdapter->GetDevice(0);
-
-	ParentAdapter->EndFrame();
-
-	for (uint32 GPUIndex : GPUMask)
-	{
-		Device = ParentAdapter->GetDevice(GPUIndex);
-
-		FD3D12CommandContext& DefaultContext = Device->GetDefaultCommandContext();
-		DefaultContext.FlushResourceBarriers();
-
-		DefaultContext.ClearState();
-		DefaultContext.FlushCommands();
-
-		Device->GetTextureAllocator().CleanUpAllocations();
-
-		// Only delete free blocks when not used in the last 2 frames, to make sure we are not allocating and releasing
-		// the same blocks every frame.
-		uint64 BufferPoolDeletionFrameLag = 20;
-		Device->GetDefaultBufferAllocator().CleanupFreeBlocks(BufferPoolDeletionFrameLag);
-
-		uint64 FastAllocatorDeletionFrameLag = 10;
-		Device->GetDefaultFastAllocator().CleanupPages(FastAllocatorDeletionFrameLag);
-	}
-
-	UpdateMemoryStats();
-
-	// Stop Timing at the very last moment
-	for (uint32 GPUIndex : GPUMask)
-	{
-		Device = ParentAdapter->GetDevice(GPUIndex);
-#if (RHI_NEW_GPU_PROFILER == 0)
-		Device->GetGPUProfiler().EndFrame();
-#endif
-	}
-
-	// Close the previous frame's timing and start a new one
-	FD3D12DynamicRHI::GetD3DRHI()->FlushTiming(true);
-
-	// Pump the interrupt queue to gather completed events
-	// (required if we're not using an interrupt thread).
-	FD3D12DynamicRHI::GetD3DRHI()->ProcessInterruptQueueUntil(nullptr);
-}
-
-void FD3D12CommandContextBase::UpdateMemoryStats()
+void FD3D12DynamicRHI::UpdateMemoryStats()
 {
 #if PLATFORM_WINDOWS && (STATS || CSV_PROFILER_STATS)
 	SCOPE_CYCLE_COUNTER(STAT_D3DUpdateVideoMemoryStats);
 
-	// Refresh captured memory stats.
-	const FD3DMemoryStats& MemoryStats = ParentAdapter->CollectMemoryStats();
-	UpdateD3DMemoryStatsAndCSV(MemoryStats, true);
+	for (TSharedPtr<FD3D12Adapter> const& Adapter : ChosenAdapters)
+	{
+		// Refresh captured memory stats.
+		const FD3DMemoryStats& MemoryStats = Adapter->CollectMemoryStats();
+		UpdateD3DMemoryStatsAndCSV(MemoryStats, true);
 	
 #if STATS
-	uint64 MaxTexAllocWastage = 0;
-	for (uint32 GPUIndex : GPUMask)
-	{
-		FD3D12Device* Device = ParentAdapter->GetDevice(GPUIndex);
-
+		uint64 MaxTexAllocWastage = 0;
+		for (FD3D12Device* Device : Adapter->GetDevices())
+		{
 #if D3D12RHI_SEGREGATED_TEXTURE_ALLOC && D3D12RHI_SEGLIST_ALLOC_TRACK_WASTAGE
-		uint64 TotalAllocated;
-		uint64 TotalUnused;
-		Device->GetTextureAllocator().GetMemoryStats(TotalAllocated, TotalUnused);
-		MaxTexAllocWastage = FMath::Max(MaxTexAllocWastage, TotalUnused);
-		SET_MEMORY_STAT(STAT_D3D12TextureAllocatorAllocated, TotalAllocated);
-		SET_MEMORY_STAT(STAT_D3D12TextureAllocatorUnused, TotalUnused);
+			uint64 TotalAllocated;
+			uint64 TotalUnused;
+			Device->GetTextureAllocator().GetMemoryStats(TotalAllocated, TotalUnused);
+			MaxTexAllocWastage = FMath::Max(MaxTexAllocWastage, TotalUnused);
+			SET_MEMORY_STAT(STAT_D3D12TextureAllocatorAllocated, TotalAllocated);
+			SET_MEMORY_STAT(STAT_D3D12TextureAllocatorUnused, TotalUnused);
 #endif
 
-		Device->GetDefaultBufferAllocator().UpdateMemoryStats();
-		ParentAdapter->GetUploadHeapAllocator(GPUIndex).UpdateMemoryStats();
-	}
+			Device->GetDefaultBufferAllocator().UpdateMemoryStats();
+			Adapter->GetUploadHeapAllocator(Device->GetGPUIndex()).UpdateMemoryStats();
+		}
 #endif // STATS
+	}
 #endif // PLATFORM_WINDOWS && (STATS || CSV_PROFILER_STATS)
 }
 
diff --git a/Engine/Source/Runtime/D3D12RHI/Private/D3D12CommandContext.h b/Engine/Source/Runtime/D3D12RHI/Private/D3D12CommandContext.h
--- a/Engine/Source/Runtime/D3D12RHI/Private/D3D12CommandContext.h
+++ b/Engine/Source/Runtime/D3D12RHI/Private/D3D12CommandContext.h
@@ -315,12 +315,16 @@ protected:
 	FD3D12Payload* GetPayload(EPhase Phase)
 	{
 		if (Payloads.Num() == 0 || Phase < CurrentPhase)
-			Payloads.Add(new FD3D12Payload(Device, QueueType));
+		{
+			NewPayload();
+		}
 
 		CurrentPhase = Phase;
 		return Payloads.Last();
 	}
 
+	void NewPayload();
+
 	uint32 ActiveQueries = 0;
 
 	// Open the command list if it's not already open.
@@ -471,10 +475,6 @@ public:
 	void RHIBeginDrawingViewport(FRHIViewport* Viewport, FRHITexture* RenderTargetRHI) final override;
 	void RHIEndDrawingViewport(FRHIViewport* Viewport, bool bPresent, bool bLockToVsync) final override;
 
-	void RHIEndFrame() final override;
-
-	virtual void UpdateMemoryStats();
-
 	FRHIGPUMask GetGPUMask() const { return GPUMask; }
 	FRHIGPUMask GetPhysicalGPUMask() const { return PhysicalGPUMask; }
 
@@ -696,8 +696,6 @@ public:
 
 	virtual void RHIClearMRTImpl(bool* bClearColorArray, int32 NumClearColors, const FLinearColor* ColorArray, bool bClearDepth, float Depth, bool bClearStencil, uint32 Stencil);
 
-	void RHIBeginFrame() final override;
-
 	virtual void RHIBeginRenderPass(const FRHIRenderPassInfo& InInfo, const TCHAR* InName)
 	{
 		FRHISetRenderTargetsInfo RTInfo;
@@ -1048,11 +1046,6 @@ public:
 		ContextRedirect(RHISetShadingRate(ShadingRate, Combiner));
 	}
 
-	FORCEINLINE virtual void RHIBeginFrame() final override
-	{
-		ContextRedirect(RHIBeginFrame());
-	}
-
 	virtual void RHIBeginRenderPass(const FRHIRenderPassInfo& InInfo, const TCHAR* InName) final override
 	{
 		ContextRedirect(RHIBeginRenderPass(InInfo, InName));
diff --git a/Engine/Source/Runtime/D3D12RHI/Private/D3D12Device.cpp b/Engine/Source/Runtime/D3D12RHI/Private/D3D12Device.cpp
--- a/Engine/Source/Runtime/D3D12RHI/Private/D3D12Device.cpp
+++ b/Engine/Source/Runtime/D3D12RHI/Private/D3D12Device.cpp
@@ -519,7 +519,7 @@ void FD3D12Device::BlockUntilIdle()
 	{
 		FD3D12SyncPointRef SyncPoint = FD3D12SyncPoint::Create(ED3D12SyncPointType::GPUAndCPU);
 
-		FD3D12Payload* Payload = new FD3D12Payload(this, (ED3D12QueueType)QueueTypeIndex);
+		FD3D12Payload* Payload = new FD3D12Payload(GetQueue((ED3D12QueueType)QueueTypeIndex));
 		Payload->SyncPointsToSignal.Add(SyncPoint);
 		Payload->bAlwaysSignal = true;
 
diff --git a/Engine/Source/Runtime/D3D12RHI/Private/D3D12DirectCommandListManager.cpp b/Engine/Source/Runtime/D3D12RHI/Private/D3D12DirectCommandListManager.cpp
--- a/Engine/Source/Runtime/D3D12RHI/Private/D3D12DirectCommandListManager.cpp
+++ b/Engine/Source/Runtime/D3D12RHI/Private/D3D12DirectCommandListManager.cpp
@@ -138,11 +138,14 @@ FD3D12ManualFence::FD3D12ManualFence(FD3D12Adapter* InParent)
 {
 	for (FD3D12Device* Device : Parent->GetDevices())
 	{
-		FFencePair& Pair = FencePairs.Emplace_GetRef();
-		Pair.Context = &Device->GetDefaultCommandContext();
+		for (FD3D12Queue& Queue : Device->GetQueues())
+		{
+			TRefCountPtr<ID3D12Fence> Fence;
+			VERIFYD3D12RESULT(Parent->GetD3DDevice()->CreateFence(0, D3D12_FENCE_FLAG_NONE, IID_PPV_ARGS(Fence.GetInitReference())));
+			Fence->SetName(TEXT("Manual Fence"));
 
-		VERIFYD3D12RESULT(Parent->GetD3DDevice()->CreateFence(0, D3D12_FENCE_FLAG_NONE, IID_PPV_ARGS(Pair.Fence.GetInitReference())));
-		Pair.Fence->SetName(TEXT("Manual Fence"));
+			Fences.Add(&Queue, MoveTemp(Fence));
+		}
 	}
 }
 
@@ -152,10 +155,10 @@ uint64 FD3D12ManualFence::GetCompletedFenceValue(bool bUpdateCachedFenceValue)
 	{
 		uint64 MinFenceValue = TNumericLimits<uint64>::Max();
 
-		for (FFencePair& Pair : FencePairs)
+		for (auto& [Queue, Fence] : Fences)
 		{
 			MinFenceValue = FMath::Min(
-				Pair.Fence->GetCompletedValue(),
+				Fence->GetCompletedValue(),
 				MinFenceValue
 			);
 		}
@@ -166,16 +169,26 @@ uint64 FD3D12ManualFence::GetCompletedFenceValue(bool bUpdateCachedFenceValue)
 	return CompletedFenceValue;
 }
 
-void FD3D12ManualFence::AdvanceFrame()
+void FD3D12ManualFence::AdvanceTOP()
+{
+	check(IsInRenderingThread());
+	NextFenceValueTOP.Increment();
+}
+
+void FD3D12ManualFence::AdvanceBOP()
 {
-	FRHICommandListImmediate& RHICmdList = FRHICommandListImmediate::Get();
+	const uint64 NextValue = ++NextFenceValueBOP;
+
+	TArray<FD3D12Payload*> Payloads;
+	Payloads.Reserve(Fences.Num());
 
-	const uint64 NextValue = NextFenceValue.Increment();
-	RHICmdList.EnqueueLambda([this, NextValue](FRHICommandListImmediate&)
+	for (auto& [Queue, Fence] : Fences)
 	{
-		for (FFencePair& Pair : FencePairs)
-		{
-			Pair.Context->SignalManualFence(Pair.Fence, NextValue);
-		}
-	});
+		FD3D12Payload* Payload = new FD3D12Payload(*Queue);
+		Payload->ManualFencesToSignal.Emplace(Fence.GetReference(), NextValue);
+
+		Payloads.Add(Payload);
+	}
+
+	FD3D12DynamicRHI::GetD3DRHI()->SubmitPayloads(MoveTemp(Payloads));
 }
diff --git a/Engine/Source/Runtime/D3D12RHI/Private/D3D12DirectCommandListManager.h b/Engine/Source/Runtime/D3D12RHI/Private/D3D12DirectCommandListManager.h
--- a/Engine/Source/Runtime/D3D12RHI/Private/D3D12DirectCommandListManager.h
+++ b/Engine/Source/Runtime/D3D12RHI/Private/D3D12DirectCommandListManager.h
@@ -2,12 +2,13 @@
 
 #pragma once
 
-#include "Containers/Array.h"
+#include "Containers/Map.h"
 #include "D3D12ThirdParty.h"
 #include "Templates/RefCounting.h"
 
 class FD3D12Adapter;
 class FD3D12CommandContext;
+class FD3D12Queue;
 
 extern int32 GEmitRgpFrameMarkers;
 
@@ -16,15 +17,10 @@ extern int32 GEmitRgpFrameMarkers;
 class FD3D12ManualFence final
 {
 	FD3D12Adapter* const Parent;
-	struct FFencePair
-	{
-		TRefCountPtr<ID3D12Fence> Fence;
-		FD3D12CommandContext* Context;
-	};
-
-	TArray<FFencePair> FencePairs;
+	TMap<FD3D12Queue*, TRefCountPtr<ID3D12Fence>> Fences;
 
-	FThreadSafeCounter NextFenceValue = 0;
+	FThreadSafeCounter NextFenceValueTOP = 0;
+	uint64 NextFenceValueBOP = 0;
 	uint64 CompletedFenceValue = 0;
 
 	FD3D12ManualFence(FD3D12ManualFence const&) = delete;
@@ -47,9 +43,9 @@ public:
 	// Returns the next value to be signaled.
 	uint64 GetNextFenceToSignal() const
 	{
-		return NextFenceValue.GetValue() + 1;
+		return NextFenceValueTOP.GetValue() + 1;
 	}
 
-	// Should only be called by RHIAdvanceFrameFence
-	void AdvanceFrame();
+	void AdvanceTOP();
+	void AdvanceBOP();
 };
diff --git a/Engine/Source/Runtime/D3D12RHI/Private/D3D12GPUProfiler.cpp b/Engine/Source/Runtime/D3D12RHI/Private/D3D12GPUProfiler.cpp
--- a/Engine/Source/Runtime/D3D12RHI/Private/D3D12GPUProfiler.cpp
+++ b/Engine/Source/Runtime/D3D12RHI/Private/D3D12GPUProfiler.cpp
@@ -217,21 +217,10 @@ void FD3D12GPUProfiler::BeginFrame()
 		SetEmitDrawEvents(bOriginalGEmitDrawEvents);
 	}
 	bPreviousLatchedGProfilingGPUHitches = bLatchedGProfilingGPUHitches;
-
-	if (GetEmitDrawEvents())
-	{
-		PushEvent(TEXT("FRAME"), FColor(0, 255, 0, 255));
-	}
 }
 
 void FD3D12GPUProfiler::EndFrame()
 {
-	if (GetEmitDrawEvents())
-	{
-		PopEvent();
-		check(StackDepth == 0);
-	}
-
 	const uint32 GPUIndex = GetParentDevice()->GetGPUIndex();
 
 	// if we have a frame open, close it now.
diff --git a/Engine/Source/Runtime/D3D12RHI/Private/D3D12GPUProfiler.h b/Engine/Source/Runtime/D3D12RHI/Private/D3D12GPUProfiler.h
--- a/Engine/Source/Runtime/D3D12RHI/Private/D3D12GPUProfiler.h
+++ b/Engine/Source/Runtime/D3D12RHI/Private/D3D12GPUProfiler.h
@@ -111,7 +111,9 @@ struct FD3D12GPUProfiler : public FGPUProfiler, public FD3D12DeviceChild
 
 	FD3D12GPUProfiler(FD3D12Device* Parent)
 		: FD3D12DeviceChild(Parent)
-	{}
+	{
+		BeginFrame();
+	}
 
 	virtual FGPUProfilerEventNode* CreateEventNode(const TCHAR* InName, FGPUProfilerEventNode* InParent) override
 	{
diff --git a/Engine/Source/Runtime/D3D12RHI/Private/D3D12RHI.cpp b/Engine/Source/Runtime/D3D12RHI/Private/D3D12RHI.cpp
--- a/Engine/Source/Runtime/D3D12RHI/Private/D3D12RHI.cpp
+++ b/Engine/Source/Runtime/D3D12RHI/Private/D3D12RHI.cpp
@@ -6,6 +6,7 @@
 
 #include "D3D12RHI.h"
 #include "D3D12RHIPrivate.h"
+#include "D3D12RayTracing.h"
 #include "RHIStaticStates.h"
 #include "OneColorShader.h"
 #include "DataDrivenShaderPlatformInfo.h"
@@ -391,7 +392,7 @@ void FD3D12DynamicRHI::EnqueueEndOfPipeTask(TUniqueFunction<void()> TaskFunc, TU
 
 	ForEachQueue([&](FD3D12Queue& Queue)
 	{
-		FD3D12Payload* Payload = new FD3D12Payload(Queue.Device, Queue.QueueType);
+		FD3D12Payload* Payload = new FD3D12Payload(Queue);
 
 		FD3D12SyncPointRef SyncPoint = FD3D12SyncPoint::Create(ED3D12SyncPointType::GPUAndCPU);
 		Payload->SyncPointsToSignal.Emplace(SyncPoint);
@@ -526,6 +527,87 @@ void FD3D12DynamicRHI::RHIProcessDeleteQueue()
 	}
 }
 
+void FD3D12DynamicRHI::RHIEndFrame_RenderThread(FRHICommandListImmediate& RHICmdList)
+{
+	// Close the GPU profiler frame
+#if (RHI_NEW_GPU_PROFILER == 0)
+	RHICmdList.EnqueueLambdaMultiPipe(GetEnabledRHIPipelines(), FRHICommandListBase::EThreadFence::Enabled, TEXT("D3D12 EndFrame"),
+		[this](FD3D12ContextArray const& Contexts)
+	{
+		for (auto& Adapter : ChosenAdapters)
+		{
+			for (auto& Device : Adapter->GetDevices())
+			{
+				Device->GetGPUProfiler().EndFrame();
+			}
+		}
+	});
+#endif
+
+	for (auto& Adapter : ChosenAdapters)
+	{
+		Adapter->GetFrameFence().AdvanceTOP();
+	}
+
+	// Base implementation flushes all prior work and results in a bottom-of-pipe call to RHIEndFrame() on the RHI thread.
+	FDynamicRHI::RHIEndFrame_RenderThread(RHICmdList);
+
+	// Start the next GPU profiler frame
+	RHICmdList.EnqueueLambdaMultiPipe(GetEnabledRHIPipelines(), FRHICommandListBase::EThreadFence::Enabled, TEXT("D3D12 BeginFrame"),
+		[this](FD3D12ContextArray const& Contexts)
+	{
+		for (auto& Adapter : ChosenAdapters)
+		{
+			for (auto& Device : Adapter->GetDevices())
+			{
+				Device->GetDefaultBufferAllocator().BeginFrame(Contexts);
+				Device->GetTextureAllocator().BeginFrame(Contexts);
+
+#if D3D12_RHI_RAYTRACING
+				// @todo dev-pr - explicit use of graphics context - nothing is synchronizing async compute - needs refactor
+				Device->GetRayTracingCompactionRequestHandler()->Update(*static_cast<FD3D12CommandContext*>(Contexts[ERHIPipeline::Graphics]));
+#endif // D3D12_RHI_RAYTRACING
+
+#if (RHI_NEW_GPU_PROFILER == 0)
+				Device->GetGPUProfiler().BeginFrame();
+#endif
+			}
+		}
+	});
+}
+
+void FD3D12DynamicRHI::RHIEndFrame()
+{
+	for (auto& Adapter : ChosenAdapters)
+	{
+		Adapter->EndFrame();
+
+		for (auto& Device : Adapter->GetDevices())
+		{
+			Device->GetTextureAllocator().CleanUpAllocations();
+
+			// Only delete free blocks when not used in the last 2 frames, to make sure we are not allocating and releasing
+			// the same blocks every frame.
+			uint64 BufferPoolDeletionFrameLag = 20;
+			Device->GetDefaultBufferAllocator().CleanupFreeBlocks(BufferPoolDeletionFrameLag);
+
+			uint64 FastAllocatorDeletionFrameLag = 10;
+			Device->GetDefaultFastAllocator().CleanupPages(FastAllocatorDeletionFrameLag);
+		}
+
+		Adapter->GetFrameFence().AdvanceBOP();
+	}
+
+	UpdateMemoryStats();
+
+	// Close the previous frame's timing and start a new one
+	FlushTiming(true);
+
+	// Pump the interrupt queue to gather completed events
+	// (required if we're not using an interrupt thread).
+	ProcessInterruptQueueUntil(nullptr);
+}
+
 TArray<FD3D12MinimalAdapterDesc> FD3D12DynamicRHI::RHIGetAdapterDescs() const
 {
 	TArray<FD3D12MinimalAdapterDesc> Result;
@@ -838,7 +920,7 @@ void FD3D12DynamicRHI::RHIRunOnQueue(ED3D12RHIRunOnQueueType QueueType, TFunctio
 	FGraphEventRef SubmissionEvent;
 
 	TArray<FD3D12Payload*> Payloads;
-	FD3D12Payload* Payload = new FD3D12Payload(GetRHIDevice(0), (QueueType == ED3D12RHIRunOnQueueType::Graphics) ?  ED3D12QueueType::Direct : ED3D12QueueType::Copy);
+	FD3D12Payload* Payload = new FD3D12Payload(GetRHIDevice(0)->GetQueue((QueueType == ED3D12RHIRunOnQueueType::Graphics) ? ED3D12QueueType::Direct : ED3D12QueueType::Copy));
 	Payloads.Add(Payload);
 
 	Payload->PreExecuteCallback = MoveTemp(CodeToRun);
diff --git a/Engine/Source/Runtime/D3D12RHI/Private/D3D12RHIPrivate.h b/Engine/Source/Runtime/D3D12RHI/Private/D3D12RHIPrivate.h
--- a/Engine/Source/Runtime/D3D12RHI/Private/D3D12RHIPrivate.h
+++ b/Engine/Source/Runtime/D3D12RHI/Private/D3D12RHIPrivate.h
@@ -215,6 +215,8 @@ private:
 	FCriticalSection ObjectsToDeleteCS;
 	TArray<FD3D12DeferredDeleteObject> ObjectsToDelete;
 
+	virtual void UpdateMemoryStats();
+
 public:
 	template <typename ...Args>
 	void DeferredDelete(Args&&... InArgs)
@@ -289,7 +291,9 @@ public:
 		return false;
 	}
 
-	virtual void RHIBeginFrame(FRHICommandListImmediate& RHICmdList) final override;
+	virtual void RHIEndFrame_RenderThread(FRHICommandListImmediate& RHICmdList) final override;
+	virtual void RHIEndFrame() final override;
+
 	virtual FSamplerStateRHIRef RHICreateSamplerState(const FSamplerStateInitializerRHI& Initializer) final override;
 	virtual FRasterizerStateRHIRef RHICreateRasterizerState(const FRasterizerStateInitializerRHI& Initializer) final override;
 	virtual FDepthStencilStateRHIRef RHICreateDepthStencilState(const FDepthStencilStateInitializerRHI& Initializer) final override;
@@ -369,7 +373,6 @@ public:
 	virtual void RHIGetDisplaysInformation(FDisplayInformationArray& OutDisplayInformation) final override;
 	virtual uint64 RHIComputePrecachePSOHash(const FGraphicsPipelineStateInitializer& Initializer) final override;
 	virtual bool RHIMatchPrecachePSOInitializers(const FGraphicsPipelineStateInitializer& LHS, const FGraphicsPipelineStateInitializer& RHS) final override;
-	virtual void RHIAdvanceFrameFence() final override;
 	virtual void RHIAdvanceFrameForGetViewportBackBuffer(FRHIViewport* Viewport) final override;
 	virtual void RHIFlushResources() final override;
 	virtual uint32 RHIGetGPUFrameCycles(uint32 GPUIndex = 0) final override;
diff --git a/Engine/Source/Runtime/D3D12RHI/Private/D3D12Submission.cpp b/Engine/Source/Runtime/D3D12RHI/Private/D3D12Submission.cpp
--- a/Engine/Source/Runtime/D3D12RHI/Private/D3D12Submission.cpp
+++ b/Engine/Source/Runtime/D3D12RHI/Private/D3D12Submission.cpp
@@ -481,7 +481,7 @@ FD3D12DynamicRHI::FProcessResult FD3D12DynamicRHI::ProcessSubmissionQueue()
 								// Barrier command list will run on a different queue.
 								if (!BarrierQueue.PayloadToSubmit)
 								{
-									BarrierQueue.PayloadToSubmit = new FD3D12Payload(BarrierCommandList->Device, BarrierCommandList->QueueType);
+									BarrierQueue.PayloadToSubmit = new FD3D12Payload(BarrierQueue);
 									QueuesWithPayloads.Add(&BarrierQueue);
 								}
 								
@@ -1433,8 +1433,8 @@ FD3D12DynamicRHI::FProcessResult FD3D12DynamicRHI::ProcessInterruptQueue()
 	return Result;
 }
 
-FD3D12PayloadBase::FD3D12PayloadBase(FD3D12Device* const Device, ED3D12QueueType const QueueType)
-	: Queue(Device->GetQueue(QueueType))
+FD3D12PayloadBase::FD3D12PayloadBase(FD3D12Queue& Queue)
+	: Queue(Queue)
 {}
 
 FD3D12PayloadBase::~FD3D12PayloadBase()
diff --git a/Engine/Source/Runtime/D3D12RHI/Private/D3D12Submission.h b/Engine/Source/Runtime/D3D12RHI/Private/D3D12Submission.h
--- a/Engine/Source/Runtime/D3D12RHI/Private/D3D12Submission.h
+++ b/Engine/Source/Runtime/D3D12RHI/Private/D3D12Submission.h
@@ -272,7 +272,7 @@ struct FD3D12PayloadBase
 	}
 
 protected:
-	FD3D12PayloadBase(FD3D12Device* Device, ED3D12QueueType QueueType);
+	FD3D12PayloadBase(FD3D12Queue& Queue);
 };
 
 #include COMPILED_PLATFORM_HEADER(D3D12Submission.h)
diff --git a/Engine/Source/Runtime/D3D12RHI/Private/D3D12Viewport.cpp b/Engine/Source/Runtime/D3D12RHI/Private/D3D12Viewport.cpp
--- a/Engine/Source/Runtime/D3D12RHI/Private/D3D12Viewport.cpp
+++ b/Engine/Source/Runtime/D3D12RHI/Private/D3D12Viewport.cpp
@@ -781,12 +781,6 @@ void FD3D12CommandContextBase::RHIEndDrawingViewport(FRHIViewport* ViewportRHI,
 	}
 }
 
-void FD3D12DynamicRHI::RHIAdvanceFrameFence()
-{
-	check(IsInRenderingThread());
-	GetAdapter().GetFrameFence().AdvanceFrame();
-}
-
 void FD3D12DynamicRHI::RHIAdvanceFrameForGetViewportBackBuffer(FRHIViewport* ViewportRHI)
 {
 	check(IsInRenderingThread());
diff --git a/Engine/Source/Runtime/D3D12RHI/Private/Windows/WindowsD3D12Submission.h b/Engine/Source/Runtime/D3D12RHI/Private/Windows/WindowsD3D12Submission.h
--- a/Engine/Source/Runtime/D3D12RHI/Private/Windows/WindowsD3D12Submission.h
+++ b/Engine/Source/Runtime/D3D12RHI/Private/Windows/WindowsD3D12Submission.h
@@ -7,8 +7,8 @@
 // Standard Windows implementation, used to mark the type as 'final'.
 struct FD3D12Payload final : public FD3D12PayloadBase
 {
-	FD3D12Payload(FD3D12Device* Device, ED3D12QueueType QueueType)
-		: FD3D12PayloadBase(Device, QueueType)
+	FD3D12Payload(FD3D12Queue& Queue)
+		: FD3D12PayloadBase(Queue)
 	{
 	}
 };
diff --git a/Engine/Source/Runtime/Engine/Private/Commandlets/Commandlet.cpp b/Engine/Source/Runtime/Engine/Private/Commandlets/Commandlet.cpp
--- a/Engine/Source/Runtime/Engine/Private/Commandlets/Commandlet.cpp
+++ b/Engine/Source/Runtime/Engine/Private/Commandlets/Commandlet.cpp
@@ -115,7 +115,6 @@ void CommandletHelpers::TickEngine(UWorld* InWorld, double InDeltaTime)
 			ENQUEUE_RENDER_COMMAND(BeginFrame)([](FRHICommandListImmediate& RHICmdList)
 			{
 				GFrameNumberRenderThread++;
-				RHICmdList.BeginFrame();
 				FCoreDelegates::OnBeginFrameRT.Broadcast();
 			});
 
diff --git a/Engine/Source/Runtime/Engine/Private/UnrealClient.cpp b/Engine/Source/Runtime/Engine/Private/UnrealClient.cpp
--- a/Engine/Source/Runtime/Engine/Private/UnrealClient.cpp
+++ b/Engine/Source/Runtime/Engine/Private/UnrealClient.cpp
@@ -1441,7 +1441,6 @@ static void HighResScreenshotBeginFrame(FDummyViewport* DummyViewport)
 		GFrameCounterRenderThread = CurrentFrameCounter;
 		GFrameNumberRenderThread++;
 		GPU_STATS_BEGINFRAME(RHICmdList);
-		RHICmdList.BeginFrame();
 		FCoreDelegates::OnBeginFrameRT.Broadcast();
 		if (DummyViewport)
 		{
diff --git a/Engine/Source/Runtime/Launch/Private/LaunchEngineLoop.cpp b/Engine/Source/Runtime/Launch/Private/LaunchEngineLoop.cpp
--- a/Engine/Source/Runtime/Launch/Private/LaunchEngineLoop.cpp
+++ b/Engine/Source/Runtime/Launch/Private/LaunchEngineLoop.cpp
@@ -5465,7 +5465,6 @@ static inline void BeginFrameRenderThread(FRHICommandListImmediate& RHICmdList,
 	if ( !FApp::CanEverRender() )
 	{
 		GFrameNumberRenderThread++;
-		RHICmdList.BeginFrame();
 		return;
 	}
 
@@ -5499,7 +5498,6 @@ static inline void BeginFrameRenderThread(FRHICommandListImmediate& RHICmdList,
 #endif
 
 	GPU_STATS_BEGINFRAME(RHICmdList);
-	RHICmdList.BeginFrame();
 	FCoreDelegates::OnBeginFrameRT.Broadcast();
 
 	RHICmdList.EnqueueLambda([CurrentFrameCounter](FRHICommandListImmediate& InRHICmdList)
@@ -5536,8 +5534,6 @@ static inline void EndFrameRenderThread(FRHICommandListImmediate& RHICmdList, ui
 	});
 
 	FCoreDelegates::OnEndFrameRT.Broadcast();
-	RHICmdList.EndFrame();
-
 	GVisualizeTexture.EndFrameRenderThread();
 
 	GPU_STATS_ENDFRAME(RHICmdList);
@@ -5547,6 +5543,8 @@ static inline void EndFrameRenderThread(FRHICommandListImmediate& RHICmdList, ui
 	GRHIFrameBreadcrumb.Reset();
 #endif
 
+	RHICmdList.EndFrame();
+
 	TRACE_END_FRAME(TraceFrameType_Rendering);
 }
 
diff --git a/Engine/Source/Runtime/MoviePlayer/Private/DefaultGameMoviePlayer.cpp b/Engine/Source/Runtime/MoviePlayer/Private/DefaultGameMoviePlayer.cpp
--- a/Engine/Source/Runtime/MoviePlayer/Private/DefaultGameMoviePlayer.cpp
+++ b/Engine/Source/Runtime/MoviePlayer/Private/DefaultGameMoviePlayer.cpp
@@ -536,9 +536,7 @@ void FDefaultGameMoviePlayer::WaitForMovieToFinish(bool bAllowEngineTick)
 				ENQUEUE_RENDER_COMMAND(BeginLoadingMovieFrameAndTickMovieStreamer)(
 					[InMoviePlayer, DeltaTime](FRHICommandListImmediate& RHICmdList)
 					{
-						GFrameNumberRenderThread++;
-						GRHICommandList.GetImmediateCommandList().BeginFrame();
-				
+						GFrameNumberRenderThread++;				
 						InMoviePlayer->TickStreamer(DeltaTime);
 					}
 				);
@@ -711,7 +709,6 @@ void FDefaultGameMoviePlayer::Tick( float DeltaTime )
 			if(SyncMechanism->IsSlateDrawPassEnqueued())
 			{
 				GFrameNumberRenderThread++;
-				GRHICommandList.GetImmediateCommandList().BeginFrame();
 				TickStreamer(DeltaTime);
 				SyncMechanism->ResetSlateDrawPassEnqueued();
 				GRHICommandList.GetImmediateCommandList().EndFrame();
diff --git a/Engine/Source/Runtime/NullDrv/Public/NullRHI.h b/Engine/Source/Runtime/NullDrv/Public/NullRHI.h
--- a/Engine/Source/Runtime/NullDrv/Public/NullRHI.h
+++ b/Engine/Source/Runtime/NullDrv/Public/NullRHI.h
@@ -340,12 +340,6 @@ public:
 		return new FNullTexture(Desc);
 	}
 
-	using FDynamicRHI::RHIBeginFrame;
-	virtual void RHIBeginFrame() final override
-	{
-
-	}
-
 	virtual void RHIEndFrame() final override
 	{
 
diff --git a/Engine/Source/Runtime/OpenGLDrv/Private/OpenGLDevice.cpp b/Engine/Source/Runtime/OpenGLDrv/Private/OpenGLDevice.cpp
--- a/Engine/Source/Runtime/OpenGLDrv/Private/OpenGLDevice.cpp
+++ b/Engine/Source/Runtime/OpenGLDrv/Private/OpenGLDevice.cpp
@@ -217,42 +217,28 @@ FOpenGLContextState& FOpenGLDynamicRHI::GetContextStateForCurrentContext()
 	return *CachedContextState;
 }
 
-void FOpenGLDynamicRHI::RHIBeginFrame()
-{
-#if (RHI_NEW_GPU_PROFILER == 0)
-	GPUProfilingData.BeginFrame(this);
-#endif
-
-#if PLATFORM_ANDROID //adding #if since not sure if this is required for any other platform.
-	PendingState.DepthStencil = 0 ;
-#endif
-
-	OpenGL_PollAllFences();
-}
-
-extern void OpenGLCommands_OnEndFrame();
-
 void FOpenGLDynamicRHI::RHIEndFrame()
 {
 #if (RHI_NEW_GPU_PROFILER == 0)
-	GPUProfilingData.EndFrame();
+	GPUProfilingData->EndFrame();
 #endif
 
 	OpenGL_PollAllFences();
 
+	extern void OpenGLCommands_OnEndFrame();
 	OpenGLCommands_OnEndFrame();
-}
 
-void FOpenGLDynamicRHI::RHIAdvanceFrameFence()
-{
-	RunOnGLRenderContextThread([]()
-	{
-		BeginFrame_UniformBufferPoolCleanup();
-		BeginFrame_VertexBufferCleanup();
-		BeginFrame_QueryBatchCleanup();
-	});
+#if PLATFORM_ANDROID //adding #if since not sure if this is required for any other platform.
+	PendingState.DepthStencil = 0;
+#endif
 
-	OpenGL_PollAllFences();
+	BeginFrame_UniformBufferPoolCleanup();
+	BeginFrame_VertexBufferCleanup();
+	BeginFrame_QueryBatchCleanup();
+
+#if (RHI_NEW_GPU_PROFILER == 0)
+	GPUProfilingData->BeginFrame();
+#endif
 }
 
 #if PLATFORM_ANDROID
@@ -1287,12 +1273,6 @@ FDynamicRHI* FOpenGLDynamicRHIModule::CreateRHI(ERHIFeatureLevel::Type InRequest
 
 
 FOpenGLDynamicRHI::FOpenGLDynamicRHI()
-:	bRevertToSharedContextAfterDrawingViewport(false)
-,	bIsRenderingContextAcquired(false)
-,	PlatformDevice(NULL)
-#if (RHI_NEW_GPU_PROFILER == 0)
-,	GPUProfilingData(this)
-#endif
 {
 	check(Singleton == nullptr);
 	Singleton = this;
@@ -1388,6 +1368,8 @@ FOpenGLDynamicRHI::FOpenGLDynamicRHI()
 
 	PrivateOpenGLDevicePtr = this;
 	GlobalUniformBuffers.AddZeroed(FUniformBufferStaticSlotRegistry::Get().GetSlotCount());
+
+	GPUProfilingData.Emplace();
 }
 
 extern void DestroyShadersAndPrograms();
@@ -1508,7 +1490,7 @@ void FOpenGLDynamicRHI::Cleanup()
 		GIsRHIInitialized = false;
 
 #if (RHI_NEW_GPU_PROFILER == 0)
-		GPUProfilingData.Cleanup();
+		GPUProfilingData->Cleanup();
 #endif
 
 		// Ask all initialized FRenderResources to release their RHI resources.
diff --git a/Engine/Source/Runtime/OpenGLDrv/Private/OpenGLDrv.cpp b/Engine/Source/Runtime/OpenGLDrv/Private/OpenGLDrv.cpp
--- a/Engine/Source/Runtime/OpenGLDrv/Private/OpenGLDrv.cpp
+++ b/Engine/Source/Runtime/OpenGLDrv/Private/OpenGLDrv.cpp
@@ -61,7 +61,7 @@ void FOpenGLDynamicRHI::RHISetExternalGPUTime(uint32 InExternalGPUTime)
 #if RHI_NEW_GPU_PROFILER
 	checkNoEntry(); // @todo - new gpu profiler
 #else
-	GPUProfilingData.ExternalGPUTime = InExternalGPUTime;
+	GPUProfilingData->ExternalGPUTime = InExternalGPUTime;
 #endif
 }
 
@@ -143,9 +143,9 @@ void FOpenGLDynamicRHI::RHIEGLTerminateContext()
 		}
 
 	#if (RHI_NEW_GPU_PROFILER == 0)
-		if (GPUProfilingData.IsProfilingGPU())
+		if (GPUProfilingData->IsProfilingGPU())
 		{
-			GPUProfilingData.PushEvent(GetNameStr(), FColor::White);
+			GPUProfilingData->PushEvent(GetNameStr(), FColor::White);
 		}
 	#endif
 	}
@@ -153,9 +153,9 @@ void FOpenGLDynamicRHI::RHIEGLTerminateContext()
 	void FOpenGLDynamicRHI::RHIEndBreadcrumbGPU(FRHIBreadcrumbNode* Breadcrumb)
 	{
 	#if (RHI_NEW_GPU_PROFILER == 0)
-		if (GPUProfilingData.IsProfilingGPU())
+		if (GPUProfilingData->IsProfilingGPU())
 		{
-			GPUProfilingData.PopEvent();
+			GPUProfilingData->PopEvent();
 		}
 	#endif
 
@@ -226,7 +226,7 @@ bool FOpenGLDynamicRHI::RHIMatchPrecachePSOInitializers(const FGraphicsPipelineS
 
 #if (RHI_NEW_GPU_PROFILER == 0)
 
-void FOpenGLGPUProfiler::BeginFrame(FOpenGLDynamicRHI* InRHI)
+void FOpenGLGPUProfiler::BeginFrame()
 {
 	if (NestedFrameCount++>0)
 	{
@@ -264,7 +264,7 @@ void FOpenGLGPUProfiler::BeginFrame(FOpenGLDynamicRHI* InRHI)
 		{
 			SetEmitDrawEvents(true);  // thwart an attempt to turn this off on the game side
 			bTrackingEvents = true;
-			CurrentEventNodeFrame = new FOpenGLEventNodeFrame(InRHI);
+			CurrentEventNodeFrame = new FOpenGLEventNodeFrame();
 			CurrentEventNodeFrame->StartFrame();
 		}
 	}
@@ -285,11 +285,6 @@ void FOpenGLGPUProfiler::BeginFrame(FOpenGLDynamicRHI* InRHI)
 		CurrentGPUFrameQueryIndex = (CurrentGPUFrameQueryIndex + 1) % MAX_GPUFRAMEQUERIES;
 		DisjointGPUFrameTimeQuery[CurrentGPUFrameQueryIndex].StartTracking();
 	}
-
-	if (GetEmitDrawEvents())
-	{
-		PushEvent(TEXT("FRAME"), FColor(0, 255, 0, 255));
-	}
 }
 
 void FOpenGLGPUProfiler::EndFrame()
@@ -300,11 +295,6 @@ void FOpenGLGPUProfiler::EndFrame()
 		return;
 	}
 
-	if (GetEmitDrawEvents())
-	{
-		PopEvent();
-	}
-
 	if (FrameTiming.IsSupported())
 	{
 		FrameTiming.EndTiming();
diff --git a/Engine/Source/Runtime/OpenGLDrv/Private/OpenGLDrv.h b/Engine/Source/Runtime/OpenGLDrv/Private/OpenGLDrv.h
--- a/Engine/Source/Runtime/OpenGLDrv/Private/OpenGLDrv.h
+++ b/Engine/Source/Runtime/OpenGLDrv/Private/OpenGLDrv.h
@@ -50,15 +50,15 @@ public:
 	 * @param InOpenGLRHI			RHI interface
 	 * @param InBufferSize		Number of buffered measurements
 	 */
-	FOpenGLBufferedGPUTiming(class FOpenGLDynamicRHI* InOpenGLRHI, int32 BufferSize);
+	FOpenGLBufferedGPUTiming(int32 BufferSize);
 
-	void	StartTiming();
+	void StartTiming();
 
 	/**
 	 * End a GPU timing measurement.
 	 * The timing for this particular measurement will be resolved at a later time by the GPU.
 	 */
-	void	EndTiming();
+	void EndTiming();
 
 	/**
 	 * Retrieves the most recently resolved timing measurement.
@@ -79,20 +79,18 @@ private:
 	 */
 	static void PlatformStaticInitialize(void* UserData);
 
-	/** RHI interface */
-	FOpenGLDynamicRHI*					OpenGLRHI;
 	/** Number of timestamps created in 'StartTimestamps' and 'EndTimestamps'. */
-	int32								BufferSize;
+	const int32							BufferSize;
 	/** Current timing being measured on the CPU. */
-	int32								CurrentTimestamp;
+	int32								CurrentTimestamp = -1;
 	/** Number of measurements in the buffers (0 - BufferSize). */
-	int32								NumIssuedTimestamps;
+	int32								NumIssuedTimestamps = 0;
 	/** Timestamps for all StartTimings. */
 	TArray<FOpenGLRenderQuery *>		StartTimestamps;
 	/** Timestamps for all EndTimings. */
 	TArray<FOpenGLRenderQuery *>		EndTimestamps;
 	/** Whether we are currently timing the GPU: between StartTiming() and EndTiming(). */
-	bool								bIsTiming;
+	bool								bIsTiming = false;
 };
 
 /**
@@ -103,13 +101,7 @@ private:
 class FOpenGLDisjointTimeStampQuery
 {
 public:
-	FOpenGLDisjointTimeStampQuery(class FOpenGLDynamicRHI* InOpenGLRHI=NULL);
-
-	void Init(class FOpenGLDynamicRHI* InOpenGLRHI)
-	{
-		OpenGLRHI = InOpenGLRHI;
-		InitResources();
-	}
+	FOpenGLDisjointTimeStampQuery() = default;
 
 	void StartTracking();
 	void EndTracking();
@@ -133,11 +125,9 @@ public:
 
 
 private:
-	bool	bIsResultValid;
-	GLuint	DisjointQuery;
-	uint64	Context;
-
-	FOpenGLDynamicRHI* OpenGLRHI;
+	bool	bIsResultValid = false;
+	GLuint	DisjointQuery = 0;
+	uint64	Context = 0;
 };
 
 /** A single perf event node, which tracks information about a appBeginDrawEvent/appEndDrawEvent range. */
@@ -145,9 +135,9 @@ class FOpenGLEventNode : public FGPUProfilerEventNode
 {
 public:
 
-	FOpenGLEventNode(const TCHAR* InName, FGPUProfilerEventNode* InParent, class FOpenGLDynamicRHI* InRHI)
-	:	FGPUProfilerEventNode(InName, InParent)
-	,	Timing(InRHI, 1)
+	FOpenGLEventNode(const TCHAR* InName, FGPUProfilerEventNode* InParent)
+		: FGPUProfilerEventNode(InName, InParent)
+		, Timing(1)
 	{
 		// Initialize Buffered timestamp queries 
 		Timing.InitResources();
@@ -181,18 +171,16 @@ public:
 class FOpenGLEventNodeFrame : public FGPUProfilerEventNodeFrame
 {
 public:
-	FOpenGLEventNodeFrame(class FOpenGLDynamicRHI* InRHI) :
-		FGPUProfilerEventNodeFrame(),
-		RootEventTiming(InRHI, 1),
-		DisjointQuery(InRHI)
+	FOpenGLEventNodeFrame()
+		: RootEventTiming(1)
+		, DisjointQuery()
 	{
-	  RootEventTiming.InitResources();
-	  DisjointQuery.InitResources();
+		RootEventTiming.InitResources();
+		DisjointQuery.InitResources();
 	}
 
 	~FOpenGLEventNodeFrame()
 	{
-
 		RootEventTiming.ReleaseResources();
 		DisjointQuery.ReleaseResources();
 	}
@@ -227,41 +215,37 @@ struct FOpenGLGPUProfiler : public FGPUProfiler
 	/** Measuring GPU frame time with a disjoint query. */
 	static const int MAX_GPUFRAMEQUERIES = 4;
 	FOpenGLDisjointTimeStampQuery DisjointGPUFrameTimeQuery[MAX_GPUFRAMEQUERIES];
-	int CurrentGPUFrameQueryIndex;
+	int32 CurrentGPUFrameQueryIndex = 0;
 
-	class FOpenGLDynamicRHI* OpenGLRHI;
 	// count the number of beginframe calls without matching endframe calls.
-	int32 NestedFrameCount;
+	int32 NestedFrameCount = 0;
 
-	uint32 ExternalGPUTime;
+	uint32 ExternalGPUTime = 0;
 
 	/** GPU hitch profile histories */
 	TIndirectArray<FOpenGLEventNodeFrame> GPUHitchEventNodeFrames;
 
-	FOpenGLGPUProfiler(class FOpenGLDynamicRHI* InOpenGLRHI)
-	:	FGPUProfiler()
-	,	FrameTiming(InOpenGLRHI, 4)
-	,	CurrentGPUFrameQueryIndex(0)
-	,	OpenGLRHI(InOpenGLRHI)
-	,	NestedFrameCount(0)
-	,	ExternalGPUTime(0)
+	FOpenGLGPUProfiler()
+		: FrameTiming(4)
 	{
 		FrameTiming.InitResources();
 		for (int32 Index = 0; Index < MAX_GPUFRAMEQUERIES; ++Index)
 		{
-			DisjointGPUFrameTimeQuery[Index].Init(OpenGLRHI);
+			DisjointGPUFrameTimeQuery[Index].InitResources();
 		}
+
+		BeginFrame();
 	}
 
 	virtual FGPUProfilerEventNode* CreateEventNode(const TCHAR* InName, FGPUProfilerEventNode* InParent) override
 	{
-		FOpenGLEventNode* EventNode = new FOpenGLEventNode(InName, InParent, OpenGLRHI);
+		FOpenGLEventNode* EventNode = new FOpenGLEventNode(InName, InParent);
 		return EventNode;
 	}
 
 	void Cleanup();
 
-	void BeginFrame(class FOpenGLDynamicRHI* InRHI);
+	void BeginFrame();
 	void EndFrame();
 };
 
@@ -425,8 +409,6 @@ public:
 	virtual void RHIEndRenderQuery(FRHIRenderQuery* RenderQuery) final override;
 	virtual void RHIBeginDrawingViewport(FRHIViewport* Viewport, FRHITexture* RenderTargetRHI) final override;
 	virtual void RHIEndDrawingViewport(FRHIViewport* Viewport, bool bPresent, bool bLockToVsync) final override;
-	using FDynamicRHI::RHIBeginFrame;
-	virtual void RHIBeginFrame() final override;
 	virtual void RHIEndFrame() final override;
 	virtual void RHISetStreamSource(uint32 StreamIndex, FRHIBuffer* VertexBuffer, uint32 Offset) final override;
 	virtual void RHISetRasterizerState(FRHIRasterizerState* NewState) final override;
@@ -769,7 +751,6 @@ public:
 	void LinkComputeProgram(FRHIComputeShader* ComputeShaderRHI);
 
 	FBoundShaderStateRHIRef RHICreateBoundShaderState_OnThisThread(FRHIVertexDeclaration* VertexDeclaration, FRHIVertexShader* VertexShader, FRHIPixelShader* PixelShader, FRHIGeometryShader* GeometryShader, bool FromPSOFileCache);
-	virtual void RHIAdvanceFrameFence() final override;
 
 	virtual void RHIPostExternalCommandsReset() final override;
 
@@ -809,9 +790,9 @@ private:
 	/** A list of all viewport RHIs that have been created. */
 	TArray<FOpenGLViewport*> Viewports;
 	TRefCountPtr<FOpenGLViewport>		DrawingViewport;
-	bool								bRevertToSharedContextAfterDrawingViewport;
+	bool								bRevertToSharedContextAfterDrawingViewport = false;
 
-	bool								bIsRenderingContextAcquired;
+	bool								bIsRenderingContextAcquired = false;
 
 	EPrimitiveType						PrimitiveType = PT_Num;
 
@@ -832,7 +813,7 @@ private:
 	TMap<GLuint, TPair<GLenum, GLenum>> TextureMipLimits;
 
 	/** Underlying platform-specific data */
-	struct FPlatformOpenGLDevice* PlatformDevice;
+	struct FPlatformOpenGLDevice* PlatformDevice = nullptr;
 
 	/** Query list. This is used to inform queries they're no longer valid when OpenGL context they're in gets released from another thread. */
 	TArray<FOpenGLRenderQuery*> Queries;
@@ -854,16 +835,16 @@ private:
 
 #else
 
-	FOpenGLGPUProfiler GPUProfilingData;
+	TOptional<FOpenGLGPUProfiler> GPUProfilingData;
 	friend FOpenGLGPUProfiler;
 
 	void RegisterGPUWork(uint32 NumPrimitives = 0, uint32 NumVertices = 0)
 	{
-		GPUProfilingData.RegisterGPUWork(NumPrimitives, NumVertices);
+		GPUProfilingData->RegisterGPUWork(NumPrimitives, NumVertices);
 	}
 	void RegisterGPUDispatch(FIntVector GroupCount)
 	{
-		GPUProfilingData.RegisterGPUDispatch(GroupCount);
+		GPUProfilingData->RegisterGPUDispatch(GroupCount);
 	}
 
 #endif
diff --git a/Engine/Source/Runtime/OpenGLDrv/Private/OpenGLQuery.cpp b/Engine/Source/Runtime/OpenGLDrv/Private/OpenGLQuery.cpp
--- a/Engine/Source/Runtime/OpenGLDrv/Private/OpenGLQuery.cpp
+++ b/Engine/Source/Runtime/OpenGLDrv/Private/OpenGLQuery.cpp
@@ -21,13 +21,13 @@ static FAutoConsoleVariableRef CVarOpenGLPollRenderQueryResult(
 
 struct FQueryItem
 {
-	FRHIRenderQuery* Query;
+	TRefCountPtr<FRHIRenderQuery> Query;
 	int32 BeginSequence;
 
 	FQueryItem(FRHIRenderQuery* InQueryRHI)
 		: Query(InQueryRHI)
 	{
-		FOpenGLRenderQuery* InQuery = FOpenGLDynamicRHI::ResourceCast(InQueryRHI);
+		FOpenGLRenderQuery* InQuery = FOpenGLDynamicRHI::ResourceCast(InQueryRHI)->GetInnerQuery();
 		BeginSequence = InQuery->TotalBegins.GetValue();
 	}
 };
@@ -99,7 +99,7 @@ struct FGLQueryBatcher
 			{
 				FQueryItem& Item = Batch->BatchContents[IndexInner];
 				FRHIRenderQuery* QueryRHI = Item.Query;
-				FOpenGLRenderQuery* Query = FOpenGLDynamicRHI::ResourceCast(QueryRHI);
+				FOpenGLRenderQuery* Query = FOpenGLDynamicRHI::ResourceCast(QueryRHI)->GetInnerQuery();
 				if (TargetQueryRHI == QueryRHI)
 				{
 					bFoundQuery = true;
@@ -158,7 +158,7 @@ struct FGLQueryBatcher
 			{
 				FQueryItem& Item = Batch->BatchContents[IndexInner];
 				FRHIRenderQuery* QueryRHI = Item.Query;
-				FOpenGLRenderQuery* Query = FOpenGLDynamicRHI::ResourceCast(QueryRHI);
+				FOpenGLRenderQuery* Query = FOpenGLDynamicRHI::ResourceCast(QueryRHI)->GetInnerQuery();
 
 				int32 Begins = Query->TotalBegins.GetValue();
 
@@ -256,7 +256,6 @@ void FOpenGLDynamicRHI::RHIPollOcclusionQueries()
 
 FRenderQueryRHIRef FOpenGLDynamicRHI::RHICreateRenderQuery(ERenderQueryType QueryType)
 {
-
 	check(QueryType == RQT_Occlusion || QueryType == RQT_AbsoluteTime);
 
 	if(QueryType == RQT_AbsoluteTime && FOpenGL::SupportsTimestampQueries() == false)
@@ -264,18 +263,18 @@ FRenderQueryRHIRef FOpenGLDynamicRHI::RHICreateRenderQuery(ERenderQueryType Quer
 		return NULL;
 	}
 
-	return new FOpenGLRenderQuery(QueryType);
+	return new FOpenGLRenderQuery_RHI(FRHICommandListImmediate::Get(), QueryType);
 }
 
 void FOpenGLDynamicRHI::RHIBeginRenderQuery(FRHIRenderQuery* QueryRHI)
 {
 	VERIFY_GL_SCOPE();
 
-	FOpenGLRenderQuery* Query = ResourceCast(QueryRHI);
+	FOpenGLRenderQuery_RHI* Query = ResourceCast(QueryRHI);
 
 	if (Query)
 	{
-		BeginRenderQuery_OnThisThread(Query);
+		BeginRenderQuery_OnThisThread(Query->GetInnerQuery());
 		GBatcher.Add(QueryRHI);
 	}
 }
@@ -284,12 +283,13 @@ void FOpenGLDynamicRHI::RHIEndRenderQuery(FRHIRenderQuery* QueryRHI)
 {
 	VERIFY_GL_SCOPE();
 
-	FOpenGLRenderQuery* Query = ResourceCast(QueryRHI);
+	FOpenGLRenderQuery_RHI* Query = ResourceCast(QueryRHI);
 
 	if (Query)
 	{
-		EndRenderQuery_OnThisThread(Query);
-		if (Query->QueryType == RQT_AbsoluteTime)
+		FOpenGLRenderQuery* InnerQuery = Query->GetInnerQuery();
+		EndRenderQuery_OnThisThread(InnerQuery);
+		if (InnerQuery->QueryType == RQT_AbsoluteTime)
 		{
 			GBatcher.AddSingle(Query);
 		}
@@ -540,7 +540,7 @@ bool FOpenGLDynamicRHI::RHIGetRenderQueryResult(FRHIRenderQuery* QueryRHI, uint6
 {
 	check(IsInRenderingThread() || IsInRHIThread());
 
-	FOpenGLRenderQuery* Query = ResourceCast(QueryRHI);
+	FOpenGLRenderQuery* Query = ResourceCast(QueryRHI)->GetInnerQuery();
 
 	if (!Query)
 	{
@@ -567,13 +567,13 @@ bool FOpenGLDynamicRHI::RHIGetRenderQueryResult(FRHIRenderQuery* QueryRHI, uint6
 				if (IsRunningRHIInDedicatedThread())
 				{
 					// send a command that will wait, so if the RHIT runs out of work, it just blocks and waits for the GPU
-					ALLOC_COMMAND_CL(RHICmdList, FRHICommandGLCommand)([this, QueryRHI]() {GetRenderQueryResult_OnThisThread(ResourceCast(QueryRHI), true); });
+					ALLOC_COMMAND_CL(RHICmdList, FRHICommandGLCommand)([this, QueryRHI]() {GetRenderQueryResult_OnThisThread(ResourceCast(QueryRHI)->GetInnerQuery(), true); });
 					FGraphEventRef Done = RHICmdList.RHIThreadFence(false);
 					ALLOC_COMMAND_CL(RHICmdList, FRHICommandGLCommand)([this, QueryRHI]() {GBatcher.Flush(*this, QueryRHI); });
 					RHICmdList.ImmediateFlush(EImmediateFlushType::DispatchToRHIThread);
 					while (!Done->IsComplete())
 					{
-						FGraphEventRef RHITask = TGraphTask<FPollQueriesRHIThreadTask>::CreateTask().ConstructAndDispatchWhenReady(ResourceCast(QueryRHI), this, false);
+						FGraphEventRef RHITask = TGraphTask<FPollQueriesRHIThreadTask>::CreateTask().ConstructAndDispatchWhenReady(ResourceCast(QueryRHI)->GetInnerQuery(), this, false);
 						FTaskGraphInterface::Get().WaitUntilTaskCompletes(RHITask);
 
 						if (Query->TotalResults.GetValue() == Query->TotalBegins.GetValue())
@@ -587,7 +587,7 @@ bool FOpenGLDynamicRHI::RHIGetRenderQueryResult(FRHIRenderQuery* QueryRHI, uint6
 				}
 				else
 				{
-					ALLOC_COMMAND_CL(RHICmdList, FRHICommandGLCommand)([this, QueryRHI]() {GetRenderQueryResult_OnThisThread(ResourceCast(QueryRHI), true); });
+					ALLOC_COMMAND_CL(RHICmdList, FRHICommandGLCommand)([this, QueryRHI]() {GetRenderQueryResult_OnThisThread(ResourceCast(QueryRHI)->GetInnerQuery(), true); });
 					FGraphEventRef Done = RHICmdList.RHIThreadFence(false);
 					ALLOC_COMMAND_CL(RHICmdList, FRHICommandGLCommand)([this, QueryRHI]() {GBatcher.Flush(*this, QueryRHI); });
 					RHICmdList.ImmediateFlush(EImmediateFlushType::DispatchToRHIThread);
@@ -597,7 +597,7 @@ bool FOpenGLDynamicRHI::RHIGetRenderQueryResult(FRHIRenderQuery* QueryRHI, uint6
 			}
 			else
 			{
-				ALLOC_COMMAND_CL(RHICmdList, FRHICommandGLCommand)([this, QueryRHI]() {GetRenderQueryResult_OnThisThread(ResourceCast(QueryRHI), false); GBatcher.Flush(*this, QueryRHI);  });
+				ALLOC_COMMAND_CL(RHICmdList, FRHICommandGLCommand)([this, QueryRHI]() {GetRenderQueryResult_OnThisThread(ResourceCast(QueryRHI)->GetInnerQuery(), false); GBatcher.Flush(*this, QueryRHI);  });
 			}
 		}	
 	}
@@ -629,44 +629,31 @@ extern void OnQueryCreation( FOpenGLRenderQuery* Query );
 extern void OnQueryDeletion( FOpenGLRenderQuery* Query );
 
 FOpenGLRenderQuery::FOpenGLRenderQuery(ERenderQueryType InQueryType)
-	: Result(0)
-	, bInvalidResource(true)
-	, QueryType(InQueryType)
+	: QueryType(InQueryType)
 {
-	FRHICommandListImmediate::Get().EnqueueLambda(
-		[this](FRHICommandListImmediate&) { AcquireResource(); }
-	);
-}
+	VERIFY_GL_SCOPE();
 
+	bInvalidResource = false;
+	PlatformGetNewRenderQuery(&Resource, &ResourceContext);
+	
+	OnQueryCreation(this);
+}
 
 FOpenGLRenderQuery::~FOpenGLRenderQuery()
 {
 	VERIFY_GL_SCOPE();
-	OnQueryDeletion( this );
+	OnQueryDeletion(this);
 
 	if (Resource && !bInvalidResource)
 	{
 		bInvalidResource = true;
-		ReleaseResource(Resource, ResourceContext);
-	}
-}
 
-void FOpenGLRenderQuery::AcquireResource()
-{
-	VERIFY_GL_SCOPE();
-	bInvalidResource = false;
-	PlatformGetNewRenderQuery(&Resource, &ResourceContext);
-	OnQueryCreation(this);
-}
-void FOpenGLRenderQuery::ReleaseResource(GLuint Resource, uint64 ResourceContext)
-{
-	VERIFY_GL_SCOPE();
-	check(Resource);
-	PlatformReleaseRenderQuery(Resource, ResourceContext);
+		check(Resource);
+		PlatformReleaseRenderQuery(Resource, ResourceContext);
+	}
 }
 
 
-
 void FOpenGLEventQuery::IssueEvent()
 {
 	VERIFY_GL_SCOPE();
@@ -741,12 +728,8 @@ FOpenGLEventQuery::~FOpenGLEventQuery()
  * @param InOpenGLRHI			RHI interface
  * @param InBufferSize		Number of buffered measurements
  */
-FOpenGLBufferedGPUTiming::FOpenGLBufferedGPUTiming( FOpenGLDynamicRHI* InOpenGLRHI, int32 InBufferSize )
-:	OpenGLRHI( InOpenGLRHI )
-,	BufferSize( InBufferSize )
-,	CurrentTimestamp( -1 )
-,	NumIssuedTimestamps( 0 )
-,	bIsTiming( false )
+FOpenGLBufferedGPUTiming::FOpenGLBufferedGPUTiming(int32 InBufferSize)
+	: BufferSize(InBufferSize)
 {
 }
 
@@ -781,7 +764,7 @@ static FOpenGLRenderQuery* GetTimeQuery()
 
 void FOpenGLBufferedGPUTiming::InitResources()
 {
-	StaticInitialize(OpenGLRHI, PlatformStaticInitialize);
+	StaticInitialize(nullptr, PlatformStaticInitialize);
 
 	CurrentTimestamp = 0;
 	NumIssuedTimestamps = 0;
@@ -1016,14 +999,6 @@ uint64 FOpenGLBufferedGPUTiming::GetTiming(bool bGetCurrentResultsAndBlock)
 	return 0;
 }
 
-FOpenGLDisjointTimeStampQuery::FOpenGLDisjointTimeStampQuery(class FOpenGLDynamicRHI* InOpenGLRHI)
-:	bIsResultValid(false)
-,	DisjointQuery(0)
-,	Context(0)
-,	OpenGLRHI(InOpenGLRHI)
-{
-}
-
 void FOpenGLDisjointTimeStampQuery::StartTracking()
 {
 	VERIFY_GL_SCOPE();
@@ -1095,11 +1070,10 @@ bool FOpenGLDisjointTimeStampQuery::GetResult( uint64* OutResult/*=NULL*/ )
 
 void FOpenGLDisjointTimeStampQuery::InitResources()
 {
+	VERIFY_GL_SCOPE();
 	if (IsSupported())
 	{
-		FRHICommandListExecutor::GetImmediateCommandList().EnqueueLambda([this](FRHICommandListImmediate&) { 
-				PlatformGetNewRenderQuery(&DisjointQuery, &Context); 
-		});
+		PlatformGetNewRenderQuery(&DisjointQuery, &Context); 
 	}
 }
 
diff --git a/Engine/Source/Runtime/OpenGLDrv/Private/OpenGLResources.h b/Engine/Source/Runtime/OpenGLDrv/Private/OpenGLResources.h
--- a/Engine/Source/Runtime/OpenGLDrv/Private/OpenGLResources.h
+++ b/Engine/Source/Runtime/OpenGLDrv/Private/OpenGLResources.h
@@ -1158,10 +1158,9 @@ public:
 	uint8 const bMultisampleRenderbuffer : 1;
 };
 
-class FOpenGLRenderQuery : public FRHIRenderQuery
+class FOpenGLRenderQuery final
 {
 public:
-
 	/** The query resource. */
 	GLuint Resource;
 
@@ -1169,7 +1168,7 @@ public:
 	uint64 ResourceContext;
 
 	/** The cached query result. */
-	GLuint64 Result;
+	GLuint64 Result = 0;
 
 	FThreadSafeCounter TotalBegins;
 	FThreadSafeCounter TotalResults;
@@ -1178,16 +1177,32 @@ public:
 	bool bResultWasSuccess;
 
 	/** true if the context the query is in was released from another thread */
-	bool bInvalidResource;
+	bool bInvalidResource = true;
 
 	// todo: memory optimize
-	ERenderQueryType QueryType;
+	const ERenderQueryType QueryType;
 
 	FOpenGLRenderQuery(ERenderQueryType InQueryType);
-	virtual ~FOpenGLRenderQuery();
+	~FOpenGLRenderQuery();
+};
+
+class FOpenGLRenderQuery_RHI : public FRHIRenderQuery
+{
+	TOptional<FOpenGLRenderQuery> Inner;
 
-	void AcquireResource();
-	static void ReleaseResource(GLuint Resource, uint64 ResourceContext);
+public:
+	FOpenGLRenderQuery_RHI(FRHICommandListBase& RHICmdList, ERenderQueryType QueryType)
+	{
+		RHICmdList.EnqueueLambda([this, QueryType](FRHICommandListBase&)
+		{
+			Inner.Emplace(QueryType);
+		});
+	}
+
+	FOpenGLRenderQuery* GetInnerQuery()
+	{
+		return &Inner.GetValue();
+	}
 };
 
 class FOpenGLView : public TIntrusiveLinkedList<FOpenGLView>
@@ -1402,7 +1417,7 @@ struct TOpenGLResourceTraits<FRHIBoundShaderState>
 template<>
 struct TOpenGLResourceTraits<FRHIRenderQuery>
 {
-	typedef FOpenGLRenderQuery TConcreteType;
+	typedef FOpenGLRenderQuery_RHI TConcreteType;
 };
 template<>
 struct TOpenGLResourceTraits<FRHIUniformBuffer>
diff --git a/Engine/Source/Runtime/PreLoadScreen/Private/PreLoadScreenManager.cpp b/Engine/Source/Runtime/PreLoadScreen/Private/PreLoadScreenManager.cpp
--- a/Engine/Source/Runtime/PreLoadScreen/Private/PreLoadScreenManager.cpp
+++ b/Engine/Source/Runtime/PreLoadScreen/Private/PreLoadScreenManager.cpp
@@ -382,7 +382,6 @@ void FPreLoadScreenManager::RenderTick_RenderThread()
 		if (MainWindow.IsValid() && VirtualRenderWindow.IsValid() && !PinnedActivePreloadScreen->IsDone())
 		{
 			GFrameNumberRenderThread++;
-			GRHICommandList.GetImmediateCommandList().BeginFrame();
 			PinnedActivePreloadScreen->RenderTick(DeltaTime);
 			GRHICommandList.GetImmediateCommandList().EndFrame();
 			GRHICommandList.GetImmediateCommandList().ImmediateFlush(EImmediateFlushType::FlushRHIThreadFlushResources);
@@ -542,7 +541,6 @@ void FPreLoadScreenManager::EarlyPlayRenderFrameTick()
 					if (FPreLoadScreenManager::bRenderingEnabled && PinnedActivePreloadScreen && !bHasRenderPreLoadScreenFrame_RenderThread)
 					{
 						GFrameNumberRenderThread++;
-						GRHICommandList.GetImmediateCommandList().BeginFrame();
 
 						bHasRenderPreLoadScreenFrame_RenderThread = true;
 						PinnedActivePreloadScreen->RenderTick(SlateDeltaTime);
diff --git a/Engine/Source/Runtime/RHI/Private/DynamicRHI.cpp b/Engine/Source/Runtime/RHI/Private/DynamicRHI.cpp
--- a/Engine/Source/Runtime/RHI/Private/DynamicRHI.cpp
+++ b/Engine/Source/Runtime/RHI/Private/DynamicRHI.cpp
@@ -449,8 +449,9 @@ void RHIExit()
 	GRHICommandList.CleanupGraphEvents();
 }
 
-void FDynamicRHI::RHIBeginFrame(FRHICommandListImmediate& RHICmdList)
+void FDynamicRHI::RHIEndFrame_RenderThread(FRHICommandListImmediate& RHICmdList)
 {
+	RHICmdList.ImmediateFlush(EImmediateFlushType::DispatchToRHIThread, ERHISubmitFlags::EndFrame);
 }
 
 // Default fallback; will not work for non-8-bit surfaces and it's extremely slow.
diff --git a/Engine/Source/Runtime/RHI/Private/RHI.cpp b/Engine/Source/Runtime/RHI/Private/RHI.cpp
--- a/Engine/Source/Runtime/RHI/Private/RHI.cpp
+++ b/Engine/Source/Runtime/RHI/Private/RHI.cpp
@@ -1294,7 +1294,6 @@ void FRHIDrawStats::Accumulate(FRHIDrawStats& Other)
 	}
 }
 
-// Called from RHIBeginFrame
 RHI_API void FRHIDrawStats::ProcessAsFrameStats()
 {
 #if HAS_GPU_STATS
diff --git a/Engine/Source/Runtime/RHI/Private/RHICommandList.cpp b/Engine/Source/Runtime/RHI/Private/RHICommandList.cpp
--- a/Engine/Source/Runtime/RHI/Private/RHICommandList.cpp
+++ b/Engine/Source/Runtime/RHI/Private/RHICommandList.cpp
@@ -847,11 +847,11 @@ void FRHICommandListExecutor::FSubmitState::Dispatch(FRHICommandListBase* CmdLis
 
 		//
 		// Don't start new translations until all prior submissions have been made.
-		// This is because some RHI commands directly submit to GPU queues from within the RHI (during RHICmdList translation).
+		// This is because work inside RHIEndFrame needs to complete on the RHI thread before any new translates can start.
+		// 
+		// Also some RHI commands directly submit to GPU queues from within the RHI (during RHICmdList translation).
 		// Not waiting for prior submits means these internal submissions can happen out-of-order with respect to other translations.
 		// E.g. some platform implementations of EndDrawingViewport() submit directly to the GPU to perform the flip / present.
-		// 
-		// @todo: remove this if/when Present() is properly pipelined, and any other implicit submits in RHI implementations are removed.
 		//
 		Prereqs.Add(GRHICommandList.LastSubmit);
 
@@ -1034,8 +1034,9 @@ void FRHICommandListExecutor::FSubmitState::Submit()
 		GDynamicRHI->RHIProcessDeleteQueue();
 	}
 
-	if (EnumHasAnyFlags(SubmitFlags, ERHISubmitFlags::ProcessStats))
+	if (EnumHasAnyFlags(SubmitFlags, ERHISubmitFlags::EndFrame))
 	{
+		GDynamicRHI->RHIEndFrame();
 		GRHICommandList.FrameDrawStats.ProcessAsFrameStats();
 	}
 
@@ -1619,53 +1620,10 @@ void FRHICommandListImmediate::EndDrawingViewport(FRHIViewport* Viewport, bool b
 	RHIAdvanceFrameForGetViewportBackBuffer(Viewport);
 }
 
-void FRHICommandListImmediate::BeginFrame()
-{
-	check(IsImmediate() && IsInRenderingThread());
-
-	{
-		QUICK_SCOPE_CYCLE_COUNTER(BeginFrame_Flush);
-		CSV_SCOPED_TIMING_STAT(RHITFlushes, BeginFrame);
-
-		// Use the ERHISubmitFlags::ProcessStats flag to delineate stat data between the previous and next frames.
-		ImmediateFlush(EImmediateFlushType::DispatchToRHIThread, ERHISubmitFlags::ProcessStats);
-	}
-
-	GDynamicRHI->RHIBeginFrame(*this);
-
-	if (Bypass())
-	{
-		GetContext().RHIBeginFrame();
-		return;
-	}
-	ALLOC_COMMAND(FRHICommandBeginFrame)();
-}
-
 void FRHICommandListImmediate::EndFrame()
 {
-	check(IsImmediate() && IsInRenderingThread());
-
-	if (Bypass())
-	{
-		GetContext().RHIEndFrame();
-		GDynamicRHI->RHIAdvanceFrameFence();
-		return;
-	}
-
-	ALLOC_COMMAND(FRHICommandEndFrame)();
-	GDynamicRHI->RHIAdvanceFrameFence();
-
-	if (!IsRunningRHIInSeparateThread())
-	{
-		// if we aren't running an RHIThread, there is no good reason to buffer this frame advance stuff and that complicates state management, so flush everything out now
-		QUICK_SCOPE_CYCLE_COUNTER(EndFrame_Flush);
-		CSV_SCOPED_TIMING_STAT(RHITFlushes, EndFrame);
-		ImmediateFlush(EImmediateFlushType::FlushRHIThread);
-	}
-	else
-	{
-		ImmediateFlush(EImmediateFlushType::DispatchToRHIThread);
-	}
+	check(IsInRenderingThread());
+	GDynamicRHI->RHIEndFrame_RenderThread(*this);
 }
 
 void FRHIComputeCommandList::Transition(TArrayView<const FRHITransitionInfo> Infos, ERHITransitionCreateFlags CreateFlags)
diff --git a/Engine/Source/Runtime/RHI/Private/RHIValidation.cpp b/Engine/Source/Runtime/RHI/Private/RHIValidation.cpp
--- a/Engine/Source/Runtime/RHI/Private/RHIValidation.cpp
+++ b/Engine/Source/Runtime/RHI/Private/RHIValidation.cpp
@@ -938,20 +938,16 @@ FValidationContext::FValidationContext(EType InType)
 	Tracker = &State.TrackerInstance;
 }
 
-void FValidationContext::RHIBeginFrame()
+void FValidationRHI::RHIEndFrame_RenderThread(FRHICommandListImmediate& RHICmdList)
 {
-	State.Reset();
-	RHIContext->RHIBeginFrame();
+	RenderThreadFrameID++;
+	RHI->RHIEndFrame_RenderThread(RHICmdList);
 }
 
-void FValidationContext::RHIEndFrame()
+void FValidationRHI::RHIEndFrame()
 {
-	RHIContext->RHIEndFrame();
-
-	// The RHI thread should always be updated at its own frequency (called from RHI thread if available)
-	// The RenderThread FrameID is update in RHIAdvanceFrameFence which is called on the RenderThread
-	FValidationRHI* ValidateRHI = (FValidationRHI*)GDynamicRHI;
-	ValidateRHI->RHIThreadFrameID++;
+	RHIThreadFrameID++;
+	RHI->RHIEndFrame();
 }
 
 namespace RHIValidation
diff --git a/Engine/Source/Runtime/RHI/Public/DynamicRHI.h b/Engine/Source/Runtime/RHI/Public/DynamicRHI.h
--- a/Engine/Source/Runtime/RHI/Public/DynamicRHI.h
+++ b/Engine/Source/Runtime/RHI/Public/DynamicRHI.h
@@ -151,7 +151,8 @@ public:
 
 	/////// RHI Methods
 
-	RHI_API virtual void RHIBeginFrame(FRHICommandListImmediate& RHICmdList);
+	RHI_API virtual void RHIEndFrame_RenderThread(FRHICommandListImmediate& RHICmdList);
+	virtual void RHIEndFrame() = 0;
 
 	// FlushType: Thread safe
 	virtual FSamplerStateRHIRef RHICreateSamplerState(const FSamplerStateInitializerRHI& Initializer) = 0;
@@ -684,8 +685,6 @@ public:
 	// Check if PSO Initializers are the same used during PSO Precaching (only compare data relevant for the RHI specific PSO)
 	RHI_API virtual bool RHIMatchPrecachePSOInitializers(const FGraphicsPipelineStateInitializer& LHS, const FGraphicsPipelineStateInitializer& RHS);
 
-	virtual void RHIAdvanceFrameFence() {};
-
 	// Only relevant with an RHI thread, this advances the backbuffer for the purpose of GetViewportBackBuffer
 	// FlushType: Thread safe
 	virtual void RHIAdvanceFrameForGetViewportBackBuffer(FRHIViewport* Viewport) = 0;
diff --git a/Engine/Source/Runtime/RHI/Public/RHICommandList.h b/Engine/Source/Runtime/RHI/Public/RHICommandList.h
--- a/Engine/Source/Runtime/RHI/Public/RHICommandList.h
+++ b/Engine/Source/Runtime/RHI/Public/RHICommandList.h
@@ -2326,22 +2326,6 @@ FRHICOMMAND_MACRO(FRHICommandPollOcclusionQueries)
 	RHI_API void Execute(FRHICommandListBase& CmdList);
 };
 
-FRHICOMMAND_MACRO(FRHICommandBeginFrame)
-{
-	FORCEINLINE_DEBUGGABLE FRHICommandBeginFrame()
-	{
-	}
-	RHI_API void Execute(FRHICommandListBase& CmdList);
-};
-
-FRHICOMMAND_MACRO(FRHICommandEndFrame)
-{
-	FORCEINLINE_DEBUGGABLE FRHICommandEndFrame()
-	{
-	}
-	RHI_API void Execute(FRHICommandListBase& CmdList);
-};
-
 FRHICOMMAND_MACRO(FRHICommandBeginDrawingViewport)
 {
 	FRHIViewport* Viewport;
@@ -4521,8 +4505,9 @@ enum class ERHISubmitFlags
 	// If combined with DeleteResources, the pending deletes queue is processed in a loop until all released resources have been deleted.
 	FlushRHIThread = 1 << 2,
 
-	// Accumulates RHI draw stats etc
-	ProcessStats = 1 << 3,
+	// Marks the end of an engine frame. Causes RHI draw stats etc to be accumulated,
+	// and calls RHIEndFrame for platform RHIs to do various cleanup tasks.
+	EndFrame = 1 << 3,
 
 #if CAN_TOGGLE_COMMAND_LIST_BYPASS
 	// Used when toggling RHI command bypass.
@@ -4542,7 +4527,6 @@ class FRHICommandListImmediate : public FRHICommandList
 {
 	friend class FRHICommandListExecutor;
 	friend class FRHICommandListScopedExtendResourceLifetime;
-	friend struct FRHICommandBeginFrame;
 
 	friend void RHI_API RHIResourceLifetimeReleaseRef(FRHICommandListImmediate&, int32);
 
@@ -4580,7 +4564,10 @@ public:
 
 	RHI_API void BeginDrawingViewport(FRHIViewport* Viewport, FRHITexture* RenderTargetRHI);
 	RHI_API void EndDrawingViewport(FRHIViewport* Viewport, bool bPresent, bool bLockToVsync);
-	RHI_API void BeginFrame();
+
+	UE_DEPRECATED(5.5, "FRHICommandListImmediate::BeginFrame() is deprecated and is no longer necessary. Remove any remaining calls to BeginFrame(). There is no replacement.")
+	void BeginFrame() {}
+
 	RHI_API void EndFrame();
 
 	struct FQueuedCommandList
diff --git a/Engine/Source/Runtime/RHI/Public/RHICommandListCommandExecutes.inl b/Engine/Source/Runtime/RHI/Public/RHICommandListCommandExecutes.inl
--- a/Engine/Source/Runtime/RHI/Public/RHICommandListCommandExecutes.inl
+++ b/Engine/Source/Runtime/RHI/Public/RHICommandListCommandExecutes.inl
@@ -19,7 +19,6 @@ struct FComputedBSS;
 struct FComputedUniformBuffer;
 struct FMemory;
 struct FRHICommandBeginDrawingViewport;
-struct FRHICommandBeginFrame;
 struct FRHICommandBeginRenderQuery;
 struct FRHICommandBuildLocalBoundShaderState;
 struct FRHICommandBuildLocalGraphicsPipelineState;
@@ -32,7 +31,6 @@ struct FRHICommandDrawPrimitiveIndirect;
 struct FRHICommandMultiDrawPrimitiveIndirect;
 struct FRHICommandSetDepthBounds;
 struct FRHICommandEndDrawingViewport;
-struct FRHICommandEndFrame;
 struct FRHICommandEndOcclusionQueryBatch;
 struct FRHICommandEndRenderQuery;
 struct FRHICommandSetBlendFactor;
@@ -571,18 +569,6 @@ void FRHICommandSetBindingsOnShaderBindingTable::Execute(FRHICommandListBase& Cm
 	INTERNAL_DECORATOR(RHISetBindingsOnShaderBindingTable)(SBT, RayTracingPipelineState, NumBindings, Bindings, BindingType);
 }
 
-void FRHICommandBeginFrame::Execute(FRHICommandListBase& CmdList)
-{
-	RHISTAT(BeginFrame);
-	INTERNAL_DECORATOR(RHIBeginFrame)();
-}
-
-void FRHICommandEndFrame::Execute(FRHICommandListBase& CmdList)
-{
-	RHISTAT(EndFrame);
-	INTERNAL_DECORATOR(RHIEndFrame)();
-}
-
 void FRHICommandBeginDrawingViewport::Execute(FRHICommandListBase& CmdList)
 {
 	RHISTAT(BeginDrawingViewport);
diff --git a/Engine/Source/Runtime/RHI/Public/RHIContext.h b/Engine/Source/Runtime/RHI/Public/RHIContext.h
--- a/Engine/Source/Runtime/RHI/Public/RHIContext.h
+++ b/Engine/Source/Runtime/RHI/Public/RHIContext.h
@@ -633,12 +633,6 @@ public:
 	// This method is queued with an RHIThread, otherwise it will flush after it is queued; without an RHI thread there is no benefit to queuing this frame advance commands
 	virtual void RHIEndDrawingViewport(FRHIViewport* Viewport, bool bPresent, bool bLockToVsync) = 0;
 
-	// This method is queued with an RHIThread, otherwise it will flush after it is queued; without an RHI thread there is no benefit to queuing this frame advance commands
-	virtual void RHIBeginFrame() = 0;
-
-	// This method is queued with an RHIThread, otherwise it will flush after it is queued; without an RHI thread there is no benefit to queuing this frame advance commands
-	virtual void RHIEndFrame() = 0;
-
 	virtual void RHISetStreamSource(uint32 StreamIndex, FRHIBuffer* VertexBuffer, uint32 Offset) = 0;
 
 	// @param MinX including like Win32 RECT
diff --git a/Engine/Source/Runtime/RHI/Public/RHIValidation.h b/Engine/Source/Runtime/RHI/Public/RHIValidation.h
--- a/Engine/Source/Runtime/RHI/Public/RHIValidation.h
+++ b/Engine/Source/Runtime/RHI/Public/RHIValidation.h
@@ -100,6 +100,9 @@ public:
 
 	/////// RHI Methods
 
+	virtual void RHIEndFrame_RenderThread(FRHICommandListImmediate& RHICmdList) override final;
+	virtual void RHIEndFrame() override final;
+
 	// FlushType: Thread safe
 	virtual FSamplerStateRHIRef RHICreateSamplerState(const FSamplerStateInitializerRHI& Initializer) override final
 	{
@@ -862,13 +865,6 @@ public:
 		RHI->RHIGetDisplaysInformation(OutDisplayInformation);
 	}
 
-	virtual void RHIAdvanceFrameFence() override final
-	{
-		check(IsInRenderingThread());
-		RHI->RHIAdvanceFrameFence();
-		RenderThreadFrameID++;
-	}
-
 	// Only relevant with an RHI thread, this advances the backbuffer for the purpose of GetViewportBackBuffer
 	// FlushType: Thread safe
 	virtual void RHIAdvanceFrameForGetViewportBackBuffer(FRHIViewport* Viewport) override final
diff --git a/Engine/Source/Runtime/RHI/Public/RHIValidationContext.h b/Engine/Source/Runtime/RHI/Public/RHIValidationContext.h
--- a/Engine/Source/Runtime/RHI/Public/RHIValidationContext.h
+++ b/Engine/Source/Runtime/RHI/Public/RHIValidationContext.h
@@ -714,12 +714,6 @@ public:
 		RHIContext->RHIEndDrawingViewport(Viewport, bPresent, bLockToVsync);
 	}
 
-	// This method is queued with an RHIThread, otherwise it will flush after it is queued; without an RHI thread there is no benefit to queuing this frame advance commands
-	virtual void RHIBeginFrame() override final;
-
-	// This method is queued with an RHIThread, otherwise it will flush after it is queued; without an RHI thread there is no benefit to queuing this frame advance commands
-	virtual void RHIEndFrame() override final;
-
 	virtual void RHISetStreamSource(uint32 StreamIndex, FRHIBuffer* VertexBuffer, uint32 Offset) override final
 	{
 		//#todo-rco: Decide if this is needed or not...
diff --git a/Engine/Source/Runtime/Renderer/Private/ReflectionEnvironmentCapture.cpp b/Engine/Source/Runtime/Renderer/Private/ReflectionEnvironmentCapture.cpp
--- a/Engine/Source/Runtime/Renderer/Private/ReflectionEnvironmentCapture.cpp
+++ b/Engine/Source/Runtime/Renderer/Private/ReflectionEnvironmentCapture.cpp
@@ -1267,7 +1267,6 @@ void CaptureSceneIntoScratchCubemap(
 				[](FRHICommandListImmediate& RHICmdList)
 			{
 				GFrameNumberRenderThread++;
-				RHICmdList.BeginFrame();
 			});
 		}
 
diff --git a/Engine/Source/Runtime/SlateRHIRenderer/Private/SlateRHIRenderer.cpp b/Engine/Source/Runtime/SlateRHIRenderer/Private/SlateRHIRenderer.cpp
--- a/Engine/Source/Runtime/SlateRHIRenderer/Private/SlateRHIRenderer.cpp
+++ b/Engine/Source/Runtime/SlateRHIRenderer/Private/SlateRHIRenderer.cpp
@@ -2160,10 +2160,8 @@ void FSlateRHIRenderer::Sync() const
 void FSlateRHIRenderer::BeginFrame() const
 {
 	ENQUEUE_RENDER_COMMAND(SlateRHIBeginFrame)(
-	   [](FRHICommandListImmediate& RHICmdList)
+	   [](FRHICommandListImmediate&)
 	   {
-		   RHICmdList.BeginFrame();
-
 		   // Suspend stat gathering when running modal dialog 'fake' frame loops
 		   GPU_STATS_SUSPENDFRAME();
 	   }
diff --git a/Engine/Source/Runtime/VulkanRHI/Private/VulkanContext.h b/Engine/Source/Runtime/VulkanRHI/Private/VulkanContext.h
--- a/Engine/Source/Runtime/VulkanRHI/Private/VulkanContext.h
+++ b/Engine/Source/Runtime/VulkanRHI/Private/VulkanContext.h
@@ -100,9 +100,6 @@ public:
 	virtual void RHIBeginDrawingViewport(FRHIViewport* Viewport, FRHITexture* RenderTargetRHI) final override;
 	virtual void RHIEndDrawingViewport(FRHIViewport* Viewport, bool bPresent, bool bLockToVsync) final override;
 
-	virtual void RHIBeginFrame() final override;
-	virtual void RHIEndFrame() final override;
-
 	virtual void RHIBeginRenderPass(const FRHIRenderPassInfo& InInfo, const TCHAR* InName) final override;
 	virtual void RHIEndRenderPass() final override;
 	virtual void RHINextSubpass() final override;
@@ -327,3 +324,16 @@ private:
 	friend class FVulkanDevice;
 };
 #endif
+
+struct FVulkanContextArray : public TRHIPipelineArray<FVulkanCommandListContext*>
+{
+	FVulkanContextArray(FRHIContextArray const& Contexts)
+		: TRHIPipelineArray(InPlace, nullptr)
+	{
+		for (ERHIPipeline Pipeline : MakeFlagsRange(ERHIPipeline::All))
+		{
+			IRHIComputeContext* Context = Contexts[Pipeline];
+			(*this)[Pipeline] = Context ? static_cast<FVulkanCommandListContext*>(&Context->GetLowestLevelContext()) : nullptr;
+		}
+	}
+};
diff --git a/Engine/Source/Runtime/VulkanRHI/Private/VulkanRHI.cpp b/Engine/Source/Runtime/VulkanRHI/Private/VulkanRHI.cpp
--- a/Engine/Source/Runtime/VulkanRHI/Private/VulkanRHI.cpp
+++ b/Engine/Source/Runtime/VulkanRHI/Private/VulkanRHI.cpp
@@ -1094,21 +1094,66 @@ void FVulkanDynamicRHI::InitInstance()
 	}
 }
 
-void FVulkanCommandListContext::RHIBeginFrame()
+void FVulkanDynamicRHI::RHIEndFrame_RenderThread(FRHICommandListImmediate& RHICmdList)
 {
-	check(IsImmediate());
+	RHICmdList.EnqueueLambdaMultiPipe(ERHIPipeline::Graphics, FRHICommandListBase::EThreadFence::Enabled, TEXT("Vulkan EndFrame"),
+		[this](FVulkanContextArray const& Contexts)
+	{
+		FVulkanCommandListContext& Context = *Contexts[ERHIPipeline::Graphics];
 
-	extern uint32 GVulkanRHIDeletionFrameNumber;
-	++GVulkanRHIDeletionFrameNumber;
+		check(Context.IsImmediate());
 
 #if (RHI_NEW_GPU_PROFILER == 0)
-	GpuProfiler.BeginFrame();
+		Context.ReadAndCalculateGPUFrameTime();
+		Context.GpuProfiler.EndFrame();
 #endif
 
-	if (GRHISupportsRayTracing)
+		bool bTrimMemory = false;
+		Context.GetCommandBufferManager()->FreeUnusedCmdBuffers(bTrimMemory);
+
+		Context.Device->GetStagingManager().ProcessPendingFree(false, true);
+		Context.Device->GetMemoryManager().ReleaseFreedPages(Context);
+		Context.Device->GetDeferredDeletionQueue().ReleaseResources();
+
+		if (UseVulkanDescriptorCache())
+		{
+			Context.Device->GetDescriptorSetCache().GC();
+		}
+		Context.Device->GetDescriptorPoolsManager().GC();
+
+		Context.Device->ReleaseUnusedOcclusionQueryPools();
+
+		Context.Device->GetPipelineStateCache()->TickLRU();
+
+		++Context.FrameCounter;
+	});
+
+	FDynamicRHI::RHIEndFrame_RenderThread(RHICmdList);
+
+	RHICmdList.EnqueueLambdaMultiPipe(ERHIPipeline::Graphics, FRHICommandListBase::EThreadFence::Enabled, TEXT("Vulkan BeginFrame"),
+		[this](FVulkanContextArray const& Contexts)
 	{
-		Device->GetRayTracingCompactionRequestHandler()->Update(*this);
-	}
+		FVulkanCommandListContext& Context = *Contexts[ERHIPipeline::Graphics];
+
+		check(Context.IsImmediate());
+
+		extern uint32 GVulkanRHIDeletionFrameNumber;
+		++GVulkanRHIDeletionFrameNumber;
+
+#if (RHI_NEW_GPU_PROFILER == 0)
+		Context.GpuProfiler.BeginFrame();
+#endif
+
+		if (GRHISupportsRayTracing)
+		{
+			Context.Device->GetRayTracingCompactionRequestHandler()->Update(Context);
+		}
+	});
+}
+
+void FVulkanDynamicRHI::RHIEndFrame()
+{
+	// @todo dev-pr - refactor RHIEndFrame_RenderThread to reduce use of the immediate command list, and move cleanup work to here.
 }
 
 void FVulkanCommandListContext::RHIBeginDrawingViewport(FRHIViewport* ViewportRHI, FRHITexture* RenderTargetRHI)
@@ -1154,36 +1199,6 @@ void FVulkanCommandListContext::RHIEndDrawingViewport(FRHIViewport* ViewportRHI,
 	WriteBeginTimestamp(CommandBufferManager->GetActiveCmdBuffer());
 }
 
-void FVulkanCommandListContext::RHIEndFrame()
-{
-	check(IsImmediate());
-	//FRCLog::Printf(FString::Printf(TEXT("FVulkanCommandListContext::RHIEndFrame()")));
-	
-#if (RHI_NEW_GPU_PROFILER == 0)
-	ReadAndCalculateGPUFrameTime();
-	GpuProfiler.EndFrame();
-#endif
-
-	bool bTrimMemory = false;
-	GetCommandBufferManager()->FreeUnusedCmdBuffers(bTrimMemory);
-
-	Device->GetStagingManager().ProcessPendingFree(false, true);
-	Device->GetMemoryManager().ReleaseFreedPages(*this);
-	Device->GetDeferredDeletionQueue().ReleaseResources();
-
-	if (UseVulkanDescriptorCache())
-	{
-		Device->GetDescriptorSetCache().GC();
-	}
-	Device->GetDescriptorPoolsManager().GC();
-
-	Device->ReleaseUnusedOcclusionQueryPools();
-
-	Device->GetPipelineStateCache()->TickLRU();
-
-	++FrameCounter;
-}
-
 #if WITH_RHI_BREADCRUMBS
 	void FVulkanCommandListContext::RHIBeginBreadcrumbGPU(FRHIBreadcrumbNode* Breadcrumb)
 	{
diff --git a/Engine/Source/Runtime/VulkanRHI/Private/VulkanUtil.cpp b/Engine/Source/Runtime/VulkanRHI/Private/VulkanUtil.cpp
--- a/Engine/Source/Runtime/VulkanRHI/Private/VulkanUtil.cpp
+++ b/Engine/Source/Runtime/VulkanRHI/Private/VulkanUtil.cpp
@@ -449,6 +449,7 @@ FVulkanGPUProfiler::FVulkanGPUProfiler(FVulkanCommandListContext* InCmd, FVulkan
 	, CmdContext(InCmd)
 	, bBeginFrame(false)
 {
+	BeginFrame();
 }
 
 FVulkanGPUProfiler::~FVulkanGPUProfiler()
@@ -521,11 +522,6 @@ void FVulkanGPUProfiler::BeginFrame()
 		SetEmitDrawEvents(bOriginalGEmitDrawEvents);
 	}
 	bPreviousLatchedGProfilingGPUHitches = bLatchedGProfilingGPUHitches;
-
-	if (GetEmitDrawEvents())
-	{
-		PushEvent(TEXT("FRAME"), FColor(0, 255, 0, 255));
-	}
 }
 
 void FVulkanGPUProfiler::EndFrameBeforeSubmit()
diff --git a/Engine/Source/Runtime/VulkanRHI/Public/VulkanDynamicRHI.h b/Engine/Source/Runtime/VulkanRHI/Public/VulkanDynamicRHI.h
--- a/Engine/Source/Runtime/VulkanRHI/Public/VulkanDynamicRHI.h
+++ b/Engine/Source/Runtime/VulkanRHI/Public/VulkanDynamicRHI.h
@@ -97,6 +97,9 @@ public:
 
 	void InitInstance();
 
+	virtual void RHIEndFrame_RenderThread(FRHICommandListImmediate& RHICmdList) final override;
+	virtual void RHIEndFrame() final override;
+
 	virtual FSamplerStateRHIRef RHICreateSamplerState(const FSamplerStateInitializerRHI& Initializer) final override;
 	virtual FRasterizerStateRHIRef RHICreateRasterizerState(const FRasterizerStateInitializerRHI& Initializer) final override;
 	virtual FDepthStencilStateRHIRef RHICreateDepthStencilState(const FDepthStencilStateInitializerRHI& Initializer) final override;
diff --git a/Engine/Source/Runtime/Windows/D3D11RHI/Private/D3D11RHI.cpp b/Engine/Source/Runtime/Windows/D3D11RHI/Private/D3D11RHI.cpp
--- a/Engine/Source/Runtime/Windows/D3D11RHI/Private/D3D11RHI.cpp
+++ b/Engine/Source/Runtime/Windows/D3D11RHI/Private/D3D11RHI.cpp
@@ -29,21 +29,20 @@ extern void UniformBufferBeginFrame();
 // Has to be .exe module to be correctly detected.
 // extern "C" { _declspec(dllexport) uint32 NvOptimusEnablement = 0x00000001; }
 
-void FD3D11DynamicRHI::RHIBeginFrame()
-{
-	UniformBufferBeginFrame();
-#if (RHI_NEW_GPU_PROFILER == 0)
-	GPUProfilingData.BeginFrame(this);
-#endif
-}
-
 void FD3D11DynamicRHI::RHIEndFrame()
 {
+	// End Frame
 #if (RHI_NEW_GPU_PROFILER == 0)
 	GPUProfilingData.EndFrame();
 #endif
 	UpdateMemoryStats();
 	CurrentComputeShader = nullptr;
+
+	// Begin Frame
+	UniformBufferBeginFrame();
+#if (RHI_NEW_GPU_PROFILER == 0)
+	GPUProfilingData.BeginFrame(this);
+#endif
 }
 
 template <int32 Frequency>
@@ -334,20 +333,10 @@ void FD3DGPUProfiler::BeginFrame(FD3D11DynamicRHI* InRHI)
 	bPreviousLatchedGProfilingGPUHitches = bLatchedGProfilingGPUHitches;
 
 	FrameTiming.StartTiming();
-
-	if (GetEmitDrawEvents())
-	{
-		PushEvent(TEXT("FRAME"), FColor(0, 255, 0, 255));
-	}
 }
 
 void FD3DGPUProfiler::EndFrame()
 {
-	if (GetEmitDrawEvents())
-	{
-		PopEvent();
-	}
-
 	FrameTiming.EndTiming();
 
 	if (FrameTiming.IsSupported())
@@ -473,6 +462,8 @@ FD3DGPUProfiler::FD3DGPUProfiler(class FD3D11DynamicRHI* InD3DRHI)
 {
 	// Initialize Buffered timestamp queries 
 	FrameTiming.InitResource(FRHICommandListImmediate::Get());
+
+	BeginFrame(InD3DRHI);
 }
 
 void FD3DGPUProfiler::PushEvent(const TCHAR* Name, FColor Color)
diff --git a/Engine/Source/Runtime/Windows/D3D11RHI/Private/D3D11RHIPrivate.h b/Engine/Source/Runtime/Windows/D3D11RHI/Private/D3D11RHIPrivate.h
--- a/Engine/Source/Runtime/Windows/D3D11RHI/Private/D3D11RHIPrivate.h
+++ b/Engine/Source/Runtime/Windows/D3D11RHI/Private/D3D11RHIPrivate.h
@@ -524,9 +524,7 @@ public:
 
 	virtual void RHIBeginDrawingViewport(FRHIViewport* Viewport, FRHITexture* RenderTargetRHI) final override;
 	virtual void RHIEndDrawingViewport(FRHIViewport* Viewport, bool bPresent, bool bLockToVsync) final override;
-	using FDynamicRHI::RHIBeginFrame;
-	virtual void RHIBeginFrame() override;
-	virtual void RHIEndFrame() override;
+	virtual void RHIEndFrame() final override;
 	virtual void RHISetStreamSource(uint32 StreamIndex, FRHIBuffer* VertexBuffer, uint32 Offset) final override;
 	virtual void RHISetRasterizerState(FRHIRasterizerState* NewState) final override;
 	virtual void RHISetViewport(float MinX, float MinY, float MinZ, float MaxX, float MaxY, float MaxZ) final override;
