Reinstate memory barriers to fix flickering seen on some malis.

diff --git a/Engine/Source/Runtime/OpenGLDrv/Private/OpenGLES.h b/Engine/Source/Runtime/OpenGLDrv/Private/OpenGLES.h
--- a/Engine/Source/Runtime/OpenGLDrv/Private/OpenGLES.h
+++ b/Engine/Source/Runtime/OpenGLDrv/Private/OpenGLES.h
@@ -417,8 +417,7 @@ struct FOpenGLES : public FOpenGLBase
 
 	static FORCEINLINE void MemoryBarrier(GLbitfield Barriers)
 	{
-		// It seems like barriers do nothing on a GLES drivers
-		//glMemoryBarrier(Barriers);
+		glMemoryBarrier(Barriers);
 	}
 
 	static FORCEINLINE void DispatchCompute(GLuint NumGroupsX, GLuint NumGroupsY, GLuint NumGroupsZ)
