Dev assertions fixed.

OpenGL texture implementation had assertions firing when evicted mips are being updated.
That is not an error and nothing in API forbids this. It seems not to be happening before, but now some CMS code re-uses textures, updating contents -> mips, so this is triggered.

Removed.

#rnx

diff --git a/Engine/Source/Runtime/OpenGLDrv/Private/OpenGLTexture.cpp b/Engine/Source/Runtime/OpenGLDrv/Private/OpenGLTexture.cpp
--- a/Engine/Source/Runtime/OpenGLDrv/Private/OpenGLTexture.cpp
+++ b/Engine/Source/Runtime/OpenGLDrv/Private/OpenGLTexture.cpp
@@ -956,8 +956,6 @@ void* FOpenGLTexture::Lock(uint32 InMipIndex,uint32 ArrayIndex,EResourceLockMode
 	if (IsEvicted())
 	{
 		check(ArrayIndex == 0);
-		// check there's nothing already here?
-		ensure(InMipIndex >= (uint32)EvictionParamsPtr->MipImageData.Num() || EvictionParamsPtr->MipImageData[InMipIndex].Num() == 0);
 		EvictionParamsPtr->SetMipData(InMipIndex, 0, MipBytes);
 		return EvictionParamsPtr->MipImageData[InMipIndex].GetData();
 	}
@@ -2209,14 +2207,10 @@ FTextureEvictionParams::~FTextureEvictionParams()
 void FTextureEvictionParams::SetMipData(uint32 MipIndex, const void* Data, uint32 Bytes)
 {
 	checkf(Bytes, TEXT("FTextureEvictionParams::SetMipData: MipIndex %d, Data %p, Bytes %d)"), MipIndex, Data, Bytes);
+	ensure(MipIndex < (uint32)MipImageData.Num());
 
 	VERIFY_GL_SCOPE();
-	if (MipImageData[MipIndex].Num())
-	{
-		// already have data??
-		checkNoEntry();
-	}
-	else
+	if (MipImageData[MipIndex].Num() == 0)
 	{
 		GTotalMipStoredCount++;
 	}
