OpenGL: Correctly set 'Layered' argiment in a glBindImageTexture call for a Texture2DArray

diff --git a/Engine/Source/Runtime/OpenGLDrv/Private/OpenGLResources.h b/Engine/Source/Runtime/OpenGLDrv/Private/OpenGLResources.h
--- a/Engine/Source/Runtime/OpenGLDrv/Private/OpenGLResources.h
+++ b/Engine/Source/Runtime/OpenGLDrv/Private/OpenGLResources.h
@@ -1233,7 +1233,7 @@ public:
 
 	bool IsLayered() const
 	{
-		return IsTexture() && GetTexture()->GetDesc().Dimension == ETextureDimension::Texture3D;
+		return IsTexture() && (GetTexture()->GetDesc().Dimension == ETextureDimension::Texture3D || GetTexture()->GetDesc().Dimension == ETextureDimension::Texture2DArray);
 	}
 
 	GLint GetLayer() const
