Simplify OpenGL RHI Render Queries
 - The logic for limiting the total number of in-flight queries has been moved out of Android-specific code and is now part of the standard implementation, driven by the GRHIMaximumInFlightQueries global.
 - GL query objects are decoupled from FRHIRenderQuery instances. They are acquired from a pool when the query is used, and returned when the results are read back.
 - A linked list of queries is formed as they are used, allowing the RHI thread to resolve them in age order.
 - Removed RHIPollOcclusionQueries and associated RDG pass code. Queries are polled automatically and periodically by the RHI thread.
 - Removed "bInvalidResource" logic on queries. This existed in case queries are created on one GL context and used on another. Instead, since acquisition of queries is now deferred to the RHI thread, we simply assert the use of the query is on the same context.
 - Reimplement "disjoint" query. glTimerQueryDisjoint is a one-shot check which is not pipelined through the GPU. If it returns true, all unresolved disjoint timestamp queries should be considered invalid, not just the one query for which glTimerQueryDisjoint returned true.

diff --git a/Engine/Source/Runtime/OpenGLDrv/Private/Android/AndroidOpenGL.cpp b/Engine/Source/Runtime/OpenGLDrv/Private/Android/AndroidOpenGL.cpp
--- a/Engine/Source/Runtime/OpenGLDrv/Private/Android/AndroidOpenGL.cpp
+++ b/Engine/Source/Runtime/OpenGLDrv/Private/Android/AndroidOpenGL.cpp
@@ -566,200 +566,6 @@ void FPlatformOpenGLDevice::SetupCurrentContext()
 }
 
 //--------------------------------
-#define VIRTUALIZE_QUERIES (1)
-
-static int32 GMaxmimumOcclusionQueries = 4000;
-
-#if VIRTUALIZE_QUERIES
-// These data structures could be better, but it would be tricky. InFlightVirtualQueries.Remove(QueryId) is a drag
-TArray<GLuint> UsableRealQueries;
-TArray<int32> InFlightVirtualQueries;
-TArray<GLuint> VirtualToRealMap;
-TArray<GLuint64> VirtualResults;
-TArray<int32> FreeVirtuals;
-TArray<GLuint> QueriesBeganButNotEnded;
-#endif
-
-#define QUERY_CHECK(x) check(x)
-//#define QUERY_CHECK(x) if (!(x)) { FPlatformMisc::LocalPrint(TEXT("Failed a check on line:\n")); FPlatformMisc::LocalPrint(TEXT( PREPROCESSOR_TO_STRING(__LINE__))); FPlatformMisc::LocalPrint(TEXT("\n")); *((int*)3) = 13; }
-
-#define CHECK_QUERY_ERRORS (DO_CHECK)
-
-void PlatformGetNewRenderQuery(GLuint* OutQuery, uint64* OutQueryContext)
-{
-#if CHECK_QUERY_ERRORS
-	GLenum Err = glGetError();
-	while (Err != GL_NO_ERROR)
-	{
-		Err = glGetError();
-	}
-#endif
-
-	*OutQueryContext = 0;
-	VERIFY_GL_SCOPE();
-
-#if !VIRTUALIZE_QUERIES
-	FOpenGLES::GenQueries(1, OutQuery);
-#if CHECK_QUERY_ERRORS
-	Err = glGetError();
-	if (Err != GL_NO_ERROR)
-	{
-		FPlatformMisc::LowLevelOutputDebugStringf(TEXT("GenQueries Failed, glError %d (0x%x)"), Err, Err);
-		*(char*)3 = 0;
-	}
-#endif
-#else
-
-
-	if (!UsableRealQueries.Num() && !InFlightVirtualQueries.Num())
-	{
-		GRHIMaximumReccommendedOustandingOcclusionQueries = GMaxmimumOcclusionQueries;
-		UE_LOG(LogRHI, Log, TEXT("AndroidOpenGL: Using a maximum of %d occlusion queries."), GMaxmimumOcclusionQueries);
-
-		UsableRealQueries.AddDefaulted(GMaxmimumOcclusionQueries);
-		glGenQueries(GMaxmimumOcclusionQueries, &UsableRealQueries[0]);
-#if CHECK_QUERY_ERRORS
-		Err = glGetError();
-		if (Err != GL_NO_ERROR)
-		{
-			FPlatformMisc::LowLevelOutputDebugStringf(TEXT("GenQueries Failed, glError %d (0x%x)"), Err, Err);
-			*(char*)3 = 0;
-		}
-#endif
-		VirtualToRealMap.Add(0); // this is null, it is not a real query and never will be
-		VirtualResults.Add(0); // this is null, it is not a real query and never will be
-	}
-
-	if (FreeVirtuals.Num())
-	{
-		*OutQuery = FreeVirtuals.Pop();
-		return;
-	}
-	*OutQuery = VirtualToRealMap.Num();
-	VirtualToRealMap.Add(0);
-	VirtualResults.Add(0);
-#endif
-}
-
-void PlatformReleaseRenderQuery(GLuint Query, uint64 QueryContext)
-{
-#if !VIRTUALIZE_QUERIES
-	glDeleteQueries(1, &Query);
-#else
-	GLuint RealIndex = VirtualToRealMap[Query];
-	if (RealIndex)
-	{
-		GLuint OutResult;
-		// still in use, wait for it now.
-		FAndroidOpenGL::GetQueryObject(Query, FAndroidOpenGL::QM_Result, &OutResult);
-		QUERY_CHECK(!VirtualToRealMap[Query]);
-	}
-	FreeVirtuals.Add(Query);
-#endif
-}
-
-void FAndroidOpenGL::GetQueryObject(GLuint QueryId, EQueryMode QueryMode, GLuint64 *OutResult)
-{
-	GLuint Result;
-	GetQueryObject(QueryId, QueryMode, &Result);
-	*OutResult = Result;
-}
-
-void FAndroidOpenGL::GetQueryObject(GLuint QueryId, EQueryMode QueryMode, GLuint* OutResult)
-{
-	GLenum QueryName = (QueryMode == QM_Result) ? GL_QUERY_RESULT : GL_QUERY_RESULT_AVAILABLE;
-	VERIFY_GL_SCOPE();
-
-	FRenderThreadIdleScope IdleScope(ERenderThreadIdleTypes::WaitingForGPUQuery, QueryName == GL_QUERY_RESULT);
-
-#if !VIRTUALIZE_QUERIES
-#if CHECK_QUERY_ERRORS
-	GLenum Err = glGetError();
-	while (Err != GL_NO_ERROR)
-	{
-		Err = glGetError();
-	}
-#endif
-
-	glGetQueryObjectuiv(QueryId, QueryName, OutResult);
-#else
-	GLuint RealIndex = VirtualToRealMap[QueryId];
-	if (!RealIndex)
-	{
-		if (QueryName == GL_QUERY_RESULT_AVAILABLE)
-		{
-			*OutResult = GL_TRUE;
-		}
-		else
-		{
-			*OutResult = VirtualResults[QueryId];
-		}
-		return;
-	}
-
-	if (QueryName == GL_QUERY_RESULT)
-	{
-		{
-			QUICK_SCOPE_CYCLE_COUNTER(STAT_FAndroidOpenGL_GetQueryObject_Remove);
-			int NumRem = InFlightVirtualQueries.Remove(QueryId);
-			QUERY_CHECK(NumRem == 1);
-		}
-		UsableRealQueries.Add(RealIndex);
-		VirtualToRealMap[QueryId] = 0;
-	}
-
-#if CHECK_QUERY_ERRORS
-	GLenum Err = glGetError();
-	while (Err != GL_NO_ERROR)
-	{
-		Err = glGetError();
-	}
-#endif
-
-	glGetQueryObjectuiv(RealIndex, QueryName, OutResult);
-
-	if (QueryName == GL_QUERY_RESULT)
-	{
-		VirtualResults[QueryId] = *OutResult;
-	}
-#endif
-
-#if CHECK_QUERY_ERRORS
-	Err = glGetError();
-	QUERY_CHECK(Err == GL_NO_ERROR);
-#endif
-}
-
-GLuint FAndroidOpenGL::MakeVirtualQueryReal(GLuint Query)
-{
-#if !VIRTUALIZE_QUERIES
-	return Query;
-#else
-	GLuint RealIndex = VirtualToRealMap[Query];
-	if (RealIndex)
-	{
-		GLuint OutResult;
-		QUICK_SCOPE_CYCLE_COUNTER(STAT_FAndroidOpenGL_BeginQuery_RecycleWait);
-		// still in use, wait for it now.
-		FAndroidOpenGL::GetQueryObject(Query, QM_Result, &OutResult);
-		QUERY_CHECK(!VirtualToRealMap[Query]);
-	}
-	if (!UsableRealQueries.Num())
-	{
-		QUERY_CHECK(InFlightVirtualQueries.Num() + QueriesBeganButNotEnded.Num() == GMaxmimumOcclusionQueries);
-		QUERY_CHECK(InFlightVirtualQueries.Num()); // if this fires, then it means the nesting of begins is too deep.
-		GLuint OutResult;
-		QUICK_SCOPE_CYCLE_COUNTER(STAT_FAndroidOpenGL_BeginQuery_FreeWait);
-		FAndroidOpenGL::GetQueryObject(InFlightVirtualQueries[0], QM_Result, &OutResult);
-		QUERY_CHECK(UsableRealQueries.Num());
-	}
-	RealIndex = UsableRealQueries.Pop();
-	VirtualToRealMap[Query] = RealIndex;
-	VirtualResults[Query] = 0;
-
-	return RealIndex;
-#endif
-}
 
 bool FAndroidOpenGL::SupportsFramebufferSRGBEnable()
 {	
@@ -768,64 +574,6 @@ bool FAndroidOpenGL::SupportsFramebufferSRGBEnable()
 	return bMobileUseHWsRGBEncoding;
 }
 
-void FAndroidOpenGL::BeginQuery(GLenum QueryType, GLuint Query)
-{
-	QUERY_CHECK(QueryType == UGL_ANY_SAMPLES_PASSED || SupportsDisjointTimeQueries());
-#if CHECK_QUERY_ERRORS
-	GLenum Err = glGetError();
-	while (Err != GL_NO_ERROR)
-	{
-		Err = glGetError();
-	}
-#endif
-
-	VERIFY_GL_SCOPE();
-
-#if !VIRTUALIZE_QUERIES
-	glBeginQuery(QueryType, Query);
-#else
-	GLuint RealIndex = MakeVirtualQueryReal(Query);
-	QueriesBeganButNotEnded.Add(Query);
-	glBeginQuery(QueryType, RealIndex);
-#endif
-#if CHECK_QUERY_ERRORS
-	Err = glGetError();
-
-	QUERY_CHECK(Err == GL_NO_ERROR);
-#endif
-}
-
-void FAndroidOpenGL::EndQuery(GLenum QueryType)
-{
-	QUERY_CHECK(QueryType == UGL_ANY_SAMPLES_PASSED || SupportsDisjointTimeQueries());
-
-#if CHECK_QUERY_ERRORS
-	GLenum Err = glGetError();
-	while (Err != GL_NO_ERROR)
-	{
-		Err = glGetError();
-	}
-#endif
-
-	VERIFY_GL_SCOPE();
-
-	if (QueryType == UGL_ANY_SAMPLES_PASSED)
-	{
-		//return;
-	}
-
-#if VIRTUALIZE_QUERIES
-	InFlightVirtualQueries.Add(QueriesBeganButNotEnded.Pop());
-#endif
-	glEndQuery(QueryType);
-
-#if CHECK_QUERY_ERRORS
-	Err = glGetError();
-
-	QUERY_CHECK(Err == GL_NO_ERROR);
-#endif
-}
-
 FAndroidOpenGL::EImageExternalType FAndroidOpenGL::ImageExternalType = FAndroidOpenGL::EImageExternalType::None;
 
 extern bool AndroidThunkCpp_GetMetaDataBoolean(const FString& Key);
@@ -918,7 +666,7 @@ void FAndroidOpenGL::ProcessExtensions(const FString& ExtensionsString)
 			}
 		}
 
-		GMaxmimumOcclusionQueries = 510;
+		GRHIMaximumInFlightQueries = 510;
 		// This is to avoid a bug in Adreno drivers that define GL_ARM_shader_framebuffer_fetch_depth_stencil even when device does not support this extension
 		// OpenGL ES 3.1 [REDACTED_EMAIL] (GIT@I1af360237c)
 		bRequiresARMShaderFramebufferFetchDepthStencilUndef = !bSupportsShaderDepthStencilFetch;
diff --git a/Engine/Source/Runtime/OpenGLDrv/Private/Android/AndroidOpenGL.h b/Engine/Source/Runtime/OpenGLDrv/Private/Android/AndroidOpenGL.h
--- a/Engine/Source/Runtime/OpenGLDrv/Private/Android/AndroidOpenGL.h
+++ b/Engine/Source/Runtime/OpenGLDrv/Private/Android/AndroidOpenGL.h
@@ -13,27 +13,7 @@
 
 struct FAndroidOpenGL : public FOpenGLES
 {
-	static FORCEINLINE bool HasHardwareHiddenSurfaceRemoval() { return bHasHardwareHiddenSurfaceRemoval; };
-
-	// Optional:
-	static void QueryTimestampCounter(GLuint QueryID) 
-	{
-	};
-
-	static GLuint MakeVirtualQueryReal(GLuint QueryID);
-
-	static FORCEINLINE void GenQueries(GLsizei NumQueries, GLuint* QueryIDs)
-	{
-		*(char*)3 = 0; // this is virtualized and should not be called
-	}
-
-	static void GetQueryObject(GLuint QueryId, EQueryMode QueryMode, GLuint *OutResult);
-
-	static void GetQueryObject(GLuint QueryId, EQueryMode QueryMode, GLuint64* OutResult);
-
-	static void BeginQuery(GLenum QueryType, GLuint QueryId);
-
-	static void EndQuery(GLenum QueryType);
+	static FORCEINLINE bool HasHardwareHiddenSurfaceRemoval() { return bHasHardwareHiddenSurfaceRemoval; }
 
 	static bool SupportsFramebufferSRGBEnable();
 
diff --git a/Engine/Source/Runtime/OpenGLDrv/Private/Android/AndroidOpenGLFramePacer.cpp b/Engine/Source/Runtime/OpenGLDrv/Private/Android/AndroidOpenGLFramePacer.cpp
--- a/Engine/Source/Runtime/OpenGLDrv/Private/Android/AndroidOpenGLFramePacer.cpp
+++ b/Engine/Source/Runtime/OpenGLDrv/Private/Android/AndroidOpenGLFramePacer.cpp
@@ -422,8 +422,6 @@ bool FAndroidOpenGLFramePacer::SwapBuffers(bool bLockToVsync)
 						{
 							break;
 						}
-
-						GetDynamicRHI<FOpenGLDynamicRHI>()->RHIPollOcclusionQueries();
 					}
 				}
 			}
diff --git a/Engine/Source/Runtime/OpenGLDrv/Private/Linux/LinuxOpenGLPlatform.cpp b/Engine/Source/Runtime/OpenGLDrv/Private/Linux/LinuxOpenGLPlatform.cpp
--- a/Engine/Source/Runtime/OpenGLDrv/Private/Linux/LinuxOpenGLPlatform.cpp
+++ b/Engine/Source/Runtime/OpenGLDrv/Private/Linux/LinuxOpenGLPlatform.cpp
@@ -93,10 +93,6 @@ private:
 	bool				bSameDCAndContext;
 };
 
-
-
-void Linux_DeleteQueriesForCurrentContext(SDL_HGLContext hGLContext);
-
 /**
  * Create a dummy window used to construct OpenGL contexts.
  */
@@ -174,8 +170,6 @@ void Linux_PlatformCreateOpenGLContextCore(FPlatformOpenGLContext* OutContext)
 
 void PlatformReleaseOpenGLContext(FPlatformOpenGLDevice* Device, FPlatformOpenGLContext* Context);
 
-extern void OnQueryInvalidation(void);
-
 /** Platform specific OpenGL device. */
 struct FPlatformOpenGLDevice
 {
@@ -239,7 +233,6 @@ struct FPlatformOpenGLDevice
 
 		Linux_ContextMakeCurrent(NULL, NULL);
 
-		OnQueryInvalidation();
 		PlatformReleaseOpenGLContext(this, &RenderingContext);
 		PlatformReleaseOpenGLContext(this, &SharedContext);
 
@@ -301,7 +294,6 @@ void PlatformReleaseOpenGLContext(FPlatformOpenGLDevice* Device, FPlatformOpenGL
 		{
 			FScopeContext ScopeContext(Context);
 
-			Linux_DeleteQueriesForCurrentContext(Context->hGLContext);
 			glBindVertexArray(0);
 			glDeleteVertexArrays(1, &Context->VertexArrayObject);
 
@@ -859,81 +851,6 @@ EOpenGLCurrentContext PlatformOpenGLCurrentContext(FPlatformOpenGLDevice* Device
 
 // =============================================================
 
-struct FOpenGLReleasedQuery
-{
-	SDL_HGLContext	hGLContext;
-	GLuint			Query;
-};
-
-static TArray<FOpenGLReleasedQuery>	ReleasedQueries;
-static FCriticalSection* ReleasedQueriesGuard;
-
-void PlatformGetNewRenderQuery(GLuint* OutQuery, uint64* OutQueryContext)
-{
-	if (!ReleasedQueriesGuard)
-	{
-		ReleasedQueriesGuard = new FCriticalSection;
-	}
-
-	{
-		FScopeLock Lock(ReleasedQueriesGuard);
-
-#ifdef UE_BUILD_DEBUG
-		check(OutQuery && OutQueryContext);
-#endif
-
-		SDL_HGLContext hGLContext = Linux_GetCurrentContext();
-		check(hGLContext);
-
-		GLuint NewQuery = 0;
-
-		// Check for possible query reuse
-		const int32 ArraySize = ReleasedQueries.Num();
-		for (int32 Index = 0; Index < ArraySize; ++Index)
-		{
-			if (ReleasedQueries[Index].hGLContext == hGLContext)
-			{
-				NewQuery = ReleasedQueries[Index].Query;
-				ReleasedQueries.RemoveAtSwap(Index);
-				break;
-			}
-		}
-
-		if (!NewQuery)
-		{
-			FOpenGL::GenQueries(1, &NewQuery);
-		}
-
-		*OutQuery = NewQuery;
-		*OutQueryContext = (uint64)hGLContext;
-
-	}
-
-}
-
-
-
-void PlatformReleaseRenderQuery(GLuint Query, uint64 QueryContext)
-{
-	SDL_HGLContext hGLContext = Linux_GetCurrentContext();
-	if ((uint64)hGLContext == QueryContext)
-	{
-		FOpenGL::DeleteQueries(1, &Query);
-	}
-	else
-	{
-		FScopeLock Lock(ReleasedQueriesGuard);
-#ifdef UE_BUILD_DEBUG
-		check(Query && QueryContext && ReleasedQueriesGuard);
-#endif
-		FOpenGLReleasedQuery ReleasedQuery;
-		ReleasedQuery.hGLContext = (SDL_HGLContext)QueryContext;
-		ReleasedQuery.Query = Query;
-		ReleasedQueries.Add(ReleasedQuery);
-	}
-
-}
-
 bool PlatformContextIsCurrent(uint64 QueryContext)
 {
 	return (uint64)Linux_GetCurrentContext() == QueryContext;
@@ -943,22 +860,3 @@ FOpenGLTexture* PlatformCreateBuiltinBackBuffer(FOpenGLDynamicRHI* OpenGLRHI, ui
 {
 	return nullptr;
 }
-
-void Linux_DeleteQueriesForCurrentContext(SDL_HGLContext hGLContext)
-{
-	if (!ReleasedQueriesGuard)
-	{
-		ReleasedQueriesGuard = new FCriticalSection;
-	}
-
-	FScopeLock Lock(ReleasedQueriesGuard);
-	for (int32 Index = 0; Index < ReleasedQueries.Num(); ++Index)
-	{
-		if (ReleasedQueries[Index].hGLContext == hGLContext)
-		{
-			FOpenGL::DeleteQueries(1, &ReleasedQueries[Index].Query);
-			ReleasedQueries.RemoveAtSwap(Index);
-			--Index;
-		}
-	}
-}
diff --git a/Engine/Source/Runtime/OpenGLDrv/Private/OpenGLBase.h b/Engine/Source/Runtime/OpenGLDrv/Private/OpenGLBase.h
--- a/Engine/Source/Runtime/OpenGLDrv/Private/OpenGLBase.h
+++ b/Engine/Source/Runtime/OpenGLDrv/Private/OpenGLBase.h
@@ -159,7 +159,7 @@ public:
 	static FORCEINLINE void QueryTimestampCounter(GLuint QueryID) UGL_OPTIONAL_VOID
 	static FORCEINLINE void BeginQuery(GLenum QueryType, GLuint QueryId) UGL_OPTIONAL_VOID
 	static FORCEINLINE void EndQuery(GLenum QueryType) UGL_OPTIONAL_VOID
-	static FORCEINLINE void GetQueryObject(GLuint QueryId, EQueryMode QueryMode, uint64 *OutResult) UGL_OPTIONAL_VOID
+	static FORCEINLINE void GetQueryObject(GLuint QueryId, EQueryMode QueryMode, GLuint64 *OutResult) UGL_OPTIONAL_VOID
 	static FORCEINLINE void BindFragDataLocation(GLuint Program, GLuint Color, const GLchar *Name) UGL_OPTIONAL_VOID
 	static FORCEINLINE void ReadBuffer(GLenum Mode) UGL_OPTIONAL_VOID
 	static FORCEINLINE void DrawBuffer(GLenum Mode) UGL_OPTIONAL_VOID
diff --git a/Engine/Source/Runtime/OpenGLDrv/Private/OpenGLCommands.cpp b/Engine/Source/Runtime/OpenGLDrv/Private/OpenGLCommands.cpp
--- a/Engine/Source/Runtime/OpenGLDrv/Private/OpenGLCommands.cpp
+++ b/Engine/Source/Runtime/OpenGLDrv/Private/OpenGLCommands.cpp
@@ -2723,7 +2723,6 @@ void FOpenGLDynamicRHI::RHIBlockUntilGPUIdle()
 	RunOnGLRenderContextThread([&]()
 	{
 		FOpenGL::Flush();
-		RHIPollOcclusionQueries();
 	});
 	
 	// @todo dev-pr use a render query to determine when the GPU is done
@@ -2918,12 +2917,6 @@ void FOpenGLDynamicRHI::RHIFinalizeContext(FRHIFinalizeContextArgs&& Args, TRHIP
 		check(Context == this);
 	}
 
-	// Flush the context to ensure recorded commands will reach the driver.
-	//if (PlatformOpenGLContextValid())
-	//{
-	//	FOpenGL::Flush(); // This gets called way too many times in a frame, which is unnecessary 
-	//}
-
 	// Clear some context state
 	FMemory::Memset(PendingState.BoundUniformBuffers, 0, sizeof(PendingState.BoundUniformBuffers));
 	FMemory::Memset(PendingState.BoundUniformBuffersDynamicOffset, 0u, sizeof(PendingState.BoundUniformBuffersDynamicOffset));
@@ -2931,7 +2924,8 @@ void FOpenGLDynamicRHI::RHIFinalizeContext(FRHIFinalizeContextArgs&& Args, TRHIP
 
 void FOpenGLDynamicRHI::RHISubmitCommandLists(FRHISubmitCommandListsArgs&& Args)
 {
-	// Nothing to do
+	// Periodically poll for completed queries
+	FOpenGLRenderQuery::PollQueryResults();
 }
 
 void FOpenGLDynamicRHI::RHICopyToStagingBuffer(FRHIBuffer* SourceBufferRHI, FRHIStagingBuffer* DestinationStagingBufferRHI, uint32 InOffset, uint32 InNumBytes)
diff --git a/Engine/Source/Runtime/OpenGLDrv/Private/OpenGLDevice.cpp b/Engine/Source/Runtime/OpenGLDrv/Private/OpenGLDevice.cpp
--- a/Engine/Source/Runtime/OpenGLDrv/Private/OpenGLDevice.cpp
+++ b/Engine/Source/Runtime/OpenGLDrv/Private/OpenGLDevice.cpp
@@ -87,28 +87,6 @@ static TAutoConsoleVariable<bool> CVarAllowPSOPrecaching(
 	TEXT("false: GL RHI will disable precaching (even if r.PSOPrecaching=1). "),
 	ECVF_RenderThreadSafe | ECVF_ReadOnly);
 
-void OnQueryCreation( FOpenGLRenderQuery* Query )
-{
-	check(PrivateOpenGLDevicePtr);
-	PrivateOpenGLDevicePtr->RegisterQuery( Query );
-}
-
-void OnQueryDeletion( FOpenGLRenderQuery* Query )
-{
-	if(PrivateOpenGLDevicePtr)
-	{
-		PrivateOpenGLDevicePtr->UnregisterQuery( Query );
-	}
-}
-
-void OnQueryInvalidation( void )
-{
-	if(PrivateOpenGLDevicePtr)
-	{
-		PrivateOpenGLDevicePtr->InvalidateQueries();
-	}
-}
-
 void OnProgramDeletion( GLint ProgramResource )
 {
 	check(PrivateOpenGLDevicePtr);
@@ -177,9 +155,11 @@ bool IsUniformBufferBound( GLuint Buffer )
 
 extern void BeginFrame_UniformBufferPoolCleanup();
 extern void BeginFrame_VertexBufferCleanup();
-extern void BeginFrame_QueryBatchCleanup();
-extern void OpenGL_PollAllFences();
 
+EOpenGLCurrentContext FOpenGLDynamicRHI::GetCurrentContext()
+{
+	return PlatformOpenGLCurrentContext(FOpenGLDynamicRHI::Get().PlatformDevice);
+}
 
 FOpenGLContextState& FOpenGLDynamicRHI::GetContextStateForCurrentContext()
 {
@@ -223,8 +203,6 @@ void FOpenGLDynamicRHI::RHIEndFrame(const FRHIEndFrameArgs& Args)
 	GPUProfilingData->EndFrame();
 #endif
 
-	OpenGL_PollAllFences();
-
 	extern void OpenGLCommands_OnEndFrame();
 	OpenGLCommands_OnEndFrame();
 
@@ -234,7 +212,6 @@ void FOpenGLDynamicRHI::RHIEndFrame(const FRHIEndFrameArgs& Args)
 
 	BeginFrame_UniformBufferPoolCleanup();
 	BeginFrame_VertexBufferCleanup();
-	BeginFrame_QueryBatchCleanup();
 
 #if (RHI_NEW_GPU_PROFILER == 0)
 	GPUProfilingData->BeginFrame();
@@ -860,6 +837,8 @@ static void InitRHICapabilitiesForGL()
 	GMaxRHIFeatureLevel = ERHIFeatureLevel::ES3_1;
 	GMaxRHIShaderPlatform = FOpenGL::GetShaderPlatform();
 
+	GRHIMaximumInFlightQueries = 4000;
+
 	// Log all supported extensions.
 #if PLATFORM_WINDOWS
 	bool bWindowsSwapControlExtensionPresent = false;
@@ -1476,6 +1455,8 @@ void FOpenGLDynamicRHI::Shutdown()
 
 	Cleanup();
 
+	FOpenGLRenderQuery::Cleanup();
+
 	DestroyShadersAndPrograms();
 	PlatformDestroyOpenGLDevice(PlatformDevice);
 
@@ -1547,18 +1528,6 @@ void FOpenGLDynamicRHI::RHIReleaseThreadOwnership()
 	CachedContextState = nullptr;
 }
 
-void FOpenGLDynamicRHI::RegisterQuery( FOpenGLRenderQuery* Query )
-{
-	FScopeLock Lock(&QueriesListCriticalSection);
-	Queries.Add(Query);
-}
-
-void FOpenGLDynamicRHI::UnregisterQuery( FOpenGLRenderQuery* Query )
-{
-	FScopeLock Lock(&QueriesListCriticalSection);
-	Queries.RemoveSingleSwap(Query);
-}
-
 void* FOpenGLDynamicRHI::RHIGetNativeDevice()
 {
 	return PlatformDevice;
@@ -1569,18 +1538,6 @@ void* FOpenGLDynamicRHI::RHIGetNativeInstance()
 	return nullptr;
 }
 
-void FOpenGLDynamicRHI::InvalidateQueries( void )
-{
-	{
-		FScopeLock Lock(&QueriesListCriticalSection);
-		PendingState.RunningOcclusionQuery = 0;
-		for( int32 Index = 0; Index < Queries.Num(); ++Index )
-		{
-			Queries[Index]->bInvalidResource = true;
-		}
-	}
-}
-
 void FOpenGLDynamicRHI::SetCustomPresent(FRHICustomPresent* InCustomPresent)
 {
 	FScopeLock lock(&CustomPresentSection);
diff --git a/Engine/Source/Runtime/OpenGLDrv/Private/OpenGLDrv.cpp b/Engine/Source/Runtime/OpenGLDrv/Private/OpenGLDrv.cpp
--- a/Engine/Source/Runtime/OpenGLDrv/Private/OpenGLDrv.cpp
+++ b/Engine/Source/Runtime/OpenGLDrv/Private/OpenGLDrv.cpp
@@ -452,11 +452,6 @@ void FOpenGLGPUProfiler::EndFrame()
 
 void FOpenGLGPUProfiler::Cleanup()
 {
-	for (int32 Index = 0; Index < MAX_GPUFRAMEQUERIES; ++Index)
-	{
-		DisjointGPUFrameTimeQuery[Index].ReleaseResources();
-	}
-
 	FrameTiming.ReleaseResources();
 	NestedFrameCount = 0;
 }
diff --git a/Engine/Source/Runtime/OpenGLDrv/Private/OpenGLDrv.h b/Engine/Source/Runtime/OpenGLDrv/Private/OpenGLDrv.h
--- a/Engine/Source/Runtime/OpenGLDrv/Private/OpenGLDrv.h
+++ b/Engine/Source/Runtime/OpenGLDrv/Private/OpenGLDrv.h
@@ -34,6 +34,7 @@ class FOpenGLDynamicRHI;
 class FResourceBulkDataInterface;
 struct FOpenGLResourceBinder;
 struct Rect;
+enum EOpenGLCurrentContext : int8;
 
 template<class T> struct TOpenGLResourceTraits;
 
@@ -106,11 +107,14 @@ public:
 	void StartTracking();
 	void EndTracking();
 	bool IsResultValid();
-	bool GetResult(uint64* OutResult=NULL);
+
+	bool GetResult(uint64* OutResult);
+
 	static uint64 GetTimingFrequency()
 	{
 		return 1000000000ull;
 	}
+
 	static bool IsSupported()
 	{
 #if UE_BUILD_SHIPPING
@@ -120,14 +124,9 @@ public:
 #endif
 	}
 
-	void InitResources();
-	void ReleaseResources();
-
-
 private:
 	bool	bIsResultValid = false;
-	GLuint	DisjointQuery = 0;
-	uint64	Context = 0;
+	FOpenGLRenderQuery DisjointQuery { FOpenGLRenderQuery::EType::Disjoint };
 };
 
 /** A single perf event node, which tracks information about a appBeginDrawEvent/appEndDrawEvent range. */
@@ -176,13 +175,11 @@ public:
 		, DisjointQuery()
 	{
 		RootEventTiming.InitResources();
-		DisjointQuery.InitResources();
 	}
 
 	~FOpenGLEventNodeFrame()
 	{
 		RootEventTiming.ReleaseResources();
-		DisjointQuery.ReleaseResources();
 	}
 
 	/** Start this frame of per tracking */
@@ -229,11 +226,6 @@ struct FOpenGLGPUProfiler : public FGPUProfiler
 		: FrameTiming(4)
 	{
 		FrameTiming.InitResources();
-		for (int32 Index = 0; Index < MAX_GPUFRAMEQUERIES; ++Index)
-		{
-			DisjointGPUFrameTimeQuery[Index].InitResources();
-		}
-
 		BeginFrame();
 	}
 
@@ -263,6 +255,8 @@ public:
 
 	friend class FOpenGLViewport;
 
+	static EOpenGLCurrentContext GetCurrentContext();
+
 	/** Initialization constructor. */
 	FOpenGLDynamicRHI();
 
@@ -382,7 +376,6 @@ public:
 	virtual EPixelFormat RHIPreferredPixelFormatHint(EPixelFormat PreferredPixelFormat) final override;
 	virtual void RHITick(float DeltaTime) final override;
 	virtual void RHIBlockUntilGPUIdle() final override;
-	virtual void RHIPollOcclusionQueries() final override;
 	virtual bool RHIGetAvailableResolutions(FScreenResolutionArray& Resolutions, bool bIgnoreRefreshRate) final override;
 	virtual void RHIGetSupportedResolution(uint32& Width, uint32& Height) final override;
 	virtual void* RHIGetNativeDevice() final override;
@@ -405,8 +398,12 @@ public:
 	virtual void RHISetMultipleViewports(uint32 Count, const FViewportBounds* Data) final override;
 	virtual void RHIClearUAVFloat(FRHIUnorderedAccessView* UnorderedAccessViewRHI, const FVector4f& Values) final override;
 	virtual void RHIClearUAVUint(FRHIUnorderedAccessView* UnorderedAccessViewRHI, const FUintVector4& Values) final override;
+
+	virtual void RHIBeginRenderQuery_TopOfPipe(FRHICommandListBase& RHICmdList, FRHIRenderQuery* RenderQuery) override final;
+	virtual void RHIEndRenderQuery_TopOfPipe(FRHICommandListBase& RHICmdList, FRHIRenderQuery* RenderQuery) override final;
 	virtual void RHIBeginRenderQuery(FRHIRenderQuery* RenderQuery) final override;
 	virtual void RHIEndRenderQuery(FRHIRenderQuery* RenderQuery) final override;
+
 	virtual void RHIBeginDrawingViewport(FRHIViewport* Viewport, FRHITexture* RenderTargetRHI) final override;
 	virtual void RHIEndDrawingViewport(FRHIViewport* Viewport, bool bPresent, bool bLockToVsync) final override;
 	virtual void RHIEndFrame(const FRHIEndFrameArgs& Args) final override;
@@ -561,19 +558,6 @@ public:
 		return ( ContextState.UniformBufferBound == Buffer );
 	}
 
-	/** Add query to Queries list upon its creation. */
-	void RegisterQuery( FOpenGLRenderQuery* Query );
-
-	/** Remove query from Queries list upon its deletion. */
-	void UnregisterQuery( FOpenGLRenderQuery* Query );
-
-	/** Inform all queries about the need to recreate themselves after OpenGL context they're in gets deleted. */
-	void InvalidateQueries();
-
-	void BeginRenderQuery_OnThisThread(FOpenGLRenderQuery* Query);
-	void EndRenderQuery_OnThisThread(FOpenGLRenderQuery* Query);
-	void GetRenderQueryResult_OnThisThread(FOpenGLRenderQuery* Query, bool bWait);
-
 	FOpenGLSamplerState* GetPointSamplerState() const { return (FOpenGLSamplerState*)PointSamplerState.GetReference(); }
 
 	void InitializeGLTexture(FOpenGLTexture* Texture, const void* BulkDataPtr, uint64 BulkDataSize);
@@ -815,12 +799,6 @@ private:
 	/** Underlying platform-specific data */
 	struct FPlatformOpenGLDevice* PlatformDevice = nullptr;
 
-	/** Query list. This is used to inform queries they're no longer valid when OpenGL context they're in gets released from another thread. */
-	TArray<FOpenGLRenderQuery*> Queries;
-
-	/** A critical section to protect modifications and iteration over Queries list */
-	FCriticalSection QueriesListCriticalSection;
-
 #if RHI_NEW_GPU_PROFILER
 
 	void RegisterGPUWork(uint32 NumPrimitives = 0, uint32 NumVertices = 0)
diff --git a/Engine/Source/Runtime/OpenGLDrv/Private/OpenGLDrvPrivate.h b/Engine/Source/Runtime/OpenGLDrv/Private/OpenGLDrvPrivate.h
--- a/Engine/Source/Runtime/OpenGLDrv/Private/OpenGLDrvPrivate.h
+++ b/Engine/Source/Runtime/OpenGLDrv/Private/OpenGLDrvPrivate.h
@@ -117,7 +117,7 @@ DECLARE_CYCLE_STAT_EXTERN(TEXT("Uniform bind time"),STAT_OpenGLUniformBindTime,S
 DECLARE_CYCLE_STAT_EXTERN(TEXT("VBO setup time"),STAT_OpenGLVBOSetupTime,STATGROUP_OpenGLRHI, );
 #endif
 
-enum EOpenGLCurrentContext
+enum EOpenGLCurrentContext : int8
 {
 	CONTEXT_Other = -2,
 	CONTEXT_Invalid = -1,
@@ -162,21 +162,6 @@ int32 PlatformGlGetError();
  */
 EOpenGLCurrentContext PlatformOpenGLCurrentContext(FPlatformOpenGLDevice* Device);
 
-/**
- * Get new occlusion query from current context. This is provided from a cache inside the context
- * if some entries are in there, and created otherwise. All other released queries present in the
- * cache are deleted at the same time (as this is the earliest possible occasion when the context
- * they're in is present).
- */
-void PlatformGetNewRenderQuery( GLuint* OutQuery, uint64* OutQueryContext );
-
-/**
- * Release occlusion query. If the current context is the one it was created in, it's deleted from
- * OpenGL immediately, otherwise it's stored on the list, and will be deleted when its context is
- * a current one.
- */
-void PlatformReleaseRenderQuery( GLuint Query, uint64 QueryContext );
-
 /**
  * Check if the query's OpenGL context is a current one. If it's not, there's no sense in issuing
  * OpenGL command for this query and it must be released and recreated.
diff --git a/Engine/Source/Runtime/OpenGLDrv/Private/OpenGLES.h b/Engine/Source/Runtime/OpenGLDrv/Private/OpenGLES.h
--- a/Engine/Source/Runtime/OpenGLDrv/Private/OpenGLES.h
+++ b/Engine/Source/Runtime/OpenGLDrv/Private/OpenGLES.h
@@ -113,6 +113,8 @@ struct FOpenGLES : public FOpenGLBase
 		glDeleteQueries(NumQueries, QueryIDs);
 	}
 
+	using FOpenGLBase::GetQueryObject;
+
 	static FORCEINLINE void GetQueryObject(GLuint QueryId, EQueryMode QueryMode, GLuint* OutResult)
 	{
 		GLenum QueryName = (QueryMode == QM_Result) ? GL_QUERY_RESULT : GL_QUERY_RESULT_AVAILABLE;
diff --git a/Engine/Source/Runtime/OpenGLDrv/Private/OpenGLQuery.cpp b/Engine/Source/Runtime/OpenGLDrv/Private/OpenGLQuery.cpp
--- a/Engine/Source/Runtime/OpenGLDrv/Private/OpenGLQuery.cpp
+++ b/Engine/Source/Runtime/OpenGLDrv/Private/OpenGLQuery.cpp
@@ -10,649 +10,424 @@
 #include "OpenGLDrvPrivate.h"
 #include "RenderCore.h"
 
-static int32 GOpenGLPollRenderQueryResult = 1;
-static FAutoConsoleVariableRef CVarOpenGLPollRenderQueryResult(
-	TEXT("r.OpenGL.PollRenderQueryResult"),
-	GOpenGLPollRenderQueryResult,
-	TEXT("Whether to poll render query for result until it's ready, otherwise do a blocking call to get result.")
-	TEXT("0: Block, 1: Poll (default)"),
-	ECVF_Default
-	);
-
-struct FQueryItem
-{
-	TRefCountPtr<FRHIRenderQuery> Query;
-	int32 BeginSequence;
-
-	FQueryItem(FRHIRenderQuery* InQueryRHI)
-		: Query(InQueryRHI)
-	{
-		FOpenGLRenderQuery* InQuery = FOpenGLDynamicRHI::ResourceCast(InQueryRHI)->GetInnerQuery();
-		BeginSequence = InQuery->TotalBegins.GetValue();
-	}
-};
+FOpenGLRenderQuery::FActiveQueries FOpenGLRenderQuery::ActiveQueries;
+FOpenGLRenderQuery::FQueryPool FOpenGLRenderQuery::PooledQueries;
 
-struct FGLQueryBatch
+FOpenGLRenderQuery::~FOpenGLRenderQuery()
 {
-	TArray<FQueryItem> BatchContents;
-	uint32 FrameNumberRenderThread;
-	bool bHasFlushedSinceLastWait;
-
-	FGLQueryBatch()
-		: FrameNumberRenderThread(0)
-		, bHasFlushedSinceLastWait(false)
-	{
-
-	}
-};
+	VERIFY_GL_SCOPE();
+	ReleaseGlQuery();
+}
 
-struct FGLQueryBatcher
+void FOpenGLRenderQuery::Link()
 {
-	FGLQueryBatch* NewBatch;
-	TArray<FGLQueryBatch*> Batches;
-	uint32 NextFrameNumberRenderThread;
+	// The renderer might re-use a query without reading its results back first.
+	// Ensure this query is unlinked, so it can be re-linked at the end of the list.
+	Unlink();
 
-	FGLQueryBatcher()
-		: NewBatch(nullptr)
-		, NextFrameNumberRenderThread(1)
+	if (!ActiveQueries.First)
 	{
-	}
+		check(!ActiveQueries.Last);
+		check(Next == nullptr);
 
-	void Add(FRHIRenderQuery* Query)
-	{
-		if (NewBatch && NewBatch->FrameNumberRenderThread)
-		{
-			NewBatch->BatchContents.Add(FQueryItem(Query));
-		}
-	}
-
-	void AddSingle(FRHIRenderQuery* Query)
-	{
-		FGLQueryBatch* SingleBatch = new FGLQueryBatch();
-		SingleBatch->FrameNumberRenderThread = NextFrameNumberRenderThread;
-		SingleBatch->BatchContents.Add(FQueryItem(Query));
-		Batches.Add(SingleBatch);
-	}
-
-	void Waited()
-	{
-		for (int32 Index = 0; Index < Batches.Num(); Index++)
-		{
-			FGLQueryBatch* Batch = Batches[Index];
-			Batch->bHasFlushedSinceLastWait = false;
-		}
+		ActiveQueries.First = this;
+		Prev = &ActiveQueries.First;
 	}
-	void Flush(FOpenGLDynamicRHI& RHI, FRHIRenderQuery* TargetQueryRHI)
+	else
 	{
-		QUICK_SCOPE_CYCLE_COUNTER(STAT_FGLQueryBatcher_FlushScan);
-		bool bFoundQuery = false;
-		for (int32 Index = 0; Index < Batches.Num() && !bFoundQuery; Index++)
-		{
-			FGLQueryBatch* Batch = Batches[Index];
-			if (Batch->bHasFlushedSinceLastWait)
-			{
-				break;
-			}
-			bool bAnyUnfinished = false;
-
-			for (int32 IndexInner = 0; IndexInner < Batch->BatchContents.Num(); IndexInner++)
-			{
-				FQueryItem& Item = Batch->BatchContents[IndexInner];
-				FRHIRenderQuery* QueryRHI = Item.Query;
-				FOpenGLRenderQuery* Query = FOpenGLDynamicRHI::ResourceCast(QueryRHI)->GetInnerQuery();
-				if (TargetQueryRHI == QueryRHI)
-				{
-					bFoundQuery = true;
-				}
+		check(ActiveQueries.Last);
+		check(ActiveQueries.Last->Next == nullptr);
 
-				if (Item.BeginSequence < Query->TotalBegins.GetValue())
-				{
-					// stale entry, was never checked, but was reused
-					Batch->BatchContents.RemoveAtSwap(IndexInner--, EAllowShrinking::No);
-					continue;
-				}
-			
-				RHI.GetRenderQueryResult_OnThisThread(Query, false);
-				if (Query->TotalResults.GetValue() == Query->TotalBegins.GetValue())
-				{
-					Batch->BatchContents.RemoveAtSwap(IndexInner--, EAllowShrinking::No);
-				}
-				else
-				{
-					bAnyUnfinished = true;
-				}
-			}
-			if (!bAnyUnfinished || Batch->BatchContents.Num() == 0)
-			{
-				delete Batch;
-				Batches.RemoveAt(Index--);
-			}
-			else
-			{
-				Batch->bHasFlushedSinceLastWait = true;
-				break;
-			}
-		}
+		ActiveQueries.Last->Next = this;
+		Prev = &ActiveQueries.Last->Next;
 	}
 
-	// this just tries to readback queries until it finds one that is not ready
-	void SoftFlush(FOpenGLDynamicRHI& RHI, bool bResetHasFlushedSinceLastWait = false)
-	{
-		QUICK_SCOPE_CYCLE_COUNTER(STAT_FGLQueryBatcher_SoftFlushScan);
-		for (int32 Index = 0; Index < Batches.Num(); Index++)
-		{
-			FGLQueryBatch* Batch = Batches[Index];
-			if (bResetHasFlushedSinceLastWait)
-			{
-				Batch->bHasFlushedSinceLastWait = false; // we will try a full scan if we get around to initviews
-			}
-
-			if (Batch->FrameNumberRenderThread == NextFrameNumberRenderThread)
-			{
-				// do not scan queries issued this frame, 
-				// on some Android devices this causes stalls in the driver (eg. S7 Adreno with Android 7)
-				break;
-			}
-
-			for (int32 IndexInner = 0; IndexInner < Batch->BatchContents.Num(); IndexInner++)
-			{
-				FQueryItem& Item = Batch->BatchContents[IndexInner];
-				FRHIRenderQuery* QueryRHI = Item.Query;
-				FOpenGLRenderQuery* Query = FOpenGLDynamicRHI::ResourceCast(QueryRHI)->GetInnerQuery();
-
-				int32 Begins = Query->TotalBegins.GetValue();
-
-				if (Item.BeginSequence < Query->TotalBegins.GetValue())
-				{
-					// stale entry, was never checked, but was reused
-					Batch->BatchContents.RemoveAtSwap(IndexInner--, EAllowShrinking::No);
-					continue;
-				}
+	ActiveQueries.Last = this;
+}
 
-				RHI.GetRenderQueryResult_OnThisThread(Query, false);
-				if (Query->TotalResults.GetValue() == Query->TotalBegins.GetValue())
-				{
-					Batch->BatchContents.RemoveAtSwap(IndexInner--, EAllowShrinking::No);
-				}
-			}
-			if (Batch->BatchContents.Num() == 0)
-			{
-				delete Batch;
-				Batches.RemoveAt(Index--);
-			}
-			else
-			{
-				break;
-			}
-		}
-	}
+void FOpenGLRenderQuery::Unlink()
+{
+	if (!IsLinked())
+		return;
 
-	void PerFrameFlush()
+	if (ActiveQueries.Last == this)
 	{
-		NextFrameNumberRenderThread++;
-		for (int32 Index = 0; Index < Batches.Num(); Index++)
+		// This is the last node in the list, so the "ActiveQueries.Last" pointer needs fixing up.
+		if (Prev == &ActiveQueries.First)
 		{
-			FGLQueryBatch* Batch = Batches[Index];
-			if (Batch->FrameNumberRenderThread <= NextFrameNumberRenderThread - 5)
-			{
-				delete Batch;
-				Batches.RemoveAt(Index--);
-			}
+			// This is also the first node in the list, meaning there's only 1 node total.
+			// Just clear the "ActiveQueries.Last" pointer.
+			ActiveQueries.Last = nullptr;
 		}
-	}
-
-	void StartNewBatch(FOpenGLDynamicRHI& RHI)
-	{
-		check(!NewBatch);
-		NewBatch = new FGLQueryBatch();
-		NewBatch->FrameNumberRenderThread = NextFrameNumberRenderThread;
-	}
-
-	void EndBatch(FOpenGLDynamicRHI& RHI)
-	{
-		QUICK_SCOPE_CYCLE_COUNTER(STAT_FGLQueryBatcher_EndBatch);
-		SoftFlush(RHI, true);
-		if (NewBatch)
+		else
 		{
-			Batches.Add(NewBatch);
-			NewBatch = nullptr;
+			//
+			// There's at least one real node before us.
+			// 
+			// "Prev" points to the "Next" member field of the previous node.
+			// Subtract the "Next" field offset to get the actual previous node address.
+			//
+			ActiveQueries.Last = reinterpret_cast<FOpenGLRenderQuery*>(reinterpret_cast<uintptr_t>(Prev) - offsetof(FOpenGLRenderQuery, Next));
 		}
 	}
 
-} GBatcher;
+	if (Next) { Next->Prev = Prev; }
+	if (Prev) { *Prev = Next; }
 
-void BeginFrame_QueryBatchCleanup()
-{
-	GBatcher.PerFrameFlush();
+	Next = nullptr;
+	Prev = nullptr;
 }
 
-void BeginOcclusionQueryBatch(uint32 NumOcclusionQueries)
+void FOpenGLRenderQuery::CheckContext()
 {
-	if (IsRunningRHIInSeparateThread())
+	check(Resource);
+	if (bSharedContext)
 	{
-		GBatcher.StartNewBatch(*GetDynamicRHI<FOpenGLDynamicRHI>());
+		check(FOpenGLDynamicRHI::GetCurrentContext() == CONTEXT_Shared);
 	}
-}
-
-void EndOcclusionQueryBatch()
-{
-	if (IsRunningRHIInSeparateThread())
+	else
 	{
-		GBatcher.EndBatch(*GetDynamicRHI<FOpenGLDynamicRHI>());
+		check(FOpenGLDynamicRHI::GetCurrentContext() == CONTEXT_Rendering);
 	}
 }
 
-void OpenGL_PollAllFences();
-
-void FOpenGLDynamicRHI::RHIPollOcclusionQueries()
+void FOpenGLRenderQuery::AcquireGlQuery()
 {
-	if (IsRunningRHIInSeparateThread())
+	if (Resource != 0)
 	{
-		GBatcher.SoftFlush(*GetDynamicRHI<FOpenGLDynamicRHI>());
-
-		OpenGL_PollAllFences();
+		// Already acquired
+		return;
 	}
-}
-
-FRenderQueryRHIRef FOpenGLDynamicRHI::RHICreateRenderQuery(ERenderQueryType QueryType)
-{
-	check(QueryType == RQT_Occlusion || QueryType == RQT_AbsoluteTime);
 
-	if(QueryType == RQT_AbsoluteTime && FOpenGL::SupportsTimestampQueries() == false)
+	EOpenGLCurrentContext Context = FOpenGLDynamicRHI::GetCurrentContext();
+	if (Context == CONTEXT_Shared)
 	{
-		return NULL;
+		// Don't do any pooling on the shared context.
+		bSharedContext = true;
+		FOpenGL::GenQueries(1, &Resource);
 	}
-
-	return new FOpenGLRenderQuery_RHI(FRHICommandListImmediate::Get(), QueryType);
-}
-
-void FOpenGLDynamicRHI::RHIBeginRenderQuery(FRHIRenderQuery* QueryRHI)
-{
-	VERIFY_GL_SCOPE();
-
-	FOpenGLRenderQuery_RHI* Query = ResourceCast(QueryRHI);
-
-	if (Query)
+	else
 	{
-		BeginRenderQuery_OnThisThread(Query->GetInnerQuery());
-		GBatcher.Add(QueryRHI);
-	}
-}
+		check(Context == CONTEXT_Rendering);
 
-void FOpenGLDynamicRHI::RHIEndRenderQuery(FRHIRenderQuery* QueryRHI)
-{
-	VERIFY_GL_SCOPE();
+		while (ActiveQueries.First && ActiveQueries.Count >= GRHIMaximumInFlightQueries)
+		{
+			// We can't start another query until more become available, due to the query count limit.
+			// Block for results on the oldest in-flight queries.
+			ActiveQueries.First->CacheResult(true);
+		}
 
-	FOpenGLRenderQuery_RHI* Query = ResourceCast(QueryRHI);
+		ActiveQueries.Count++;
 
-	if (Query)
-	{
-		FOpenGLRenderQuery* InnerQuery = Query->GetInnerQuery();
-		EndRenderQuery_OnThisThread(InnerQuery);
-		if (InnerQuery->QueryType == RQT_AbsoluteTime)
+		if (PooledQueries[Type].Num())
+		{
+			Resource = PooledQueries[Type].Pop();
+		}
+		else
 		{
-			GBatcher.AddSingle(Query);
+			FOpenGL::GenQueries(1, &Resource);
 		}
 	}
 }
 
-void FOpenGLDynamicRHI::BeginRenderQuery_OnThisThread(FOpenGLRenderQuery* Query)
+void FOpenGLRenderQuery::ReleaseGlQuery()
 {
-	VERIFY_GL_SCOPE();
-
-	int32 NewVal = Query->TotalBegins.Increment();
-	Query->TotalResults.Set(NewVal - 1);
-	Query->Result = 0;
-	Query->bResultWasSuccess = false;
+	CheckContext();
 
-	if (Query->QueryType == RQT_Occlusion)
+	if (bSharedContext)
 	{
-		check(PendingState.RunningOcclusionQuery == 0);
-
-		if (!Query->bInvalidResource && !PlatformContextIsCurrent(Query->ResourceContext))
-		{
-			PlatformReleaseRenderQuery(Query->Resource, Query->ResourceContext);
-			Query->bInvalidResource = true;
-		}
-
-		if (Query->bInvalidResource)
-		{
-			PlatformGetNewRenderQuery(&Query->Resource, &Query->ResourceContext);
-			Query->bInvalidResource = false;
-		}
-
-		GLenum QueryType = FOpenGL::SupportsExactOcclusionQueries() ? UGL_SAMPLES_PASSED : UGL_ANY_SAMPLES_PASSED;
-		FOpenGL::BeginQuery(QueryType, Query->Resource);
-		PendingState.RunningOcclusionQuery = Query->Resource;
+		// Don't do any pooling on the shared context.
+		FOpenGL::DeleteQueries(1, &Resource);
+		bSharedContext = false;
 	}
 	else
 	{
-		// not supported/needed for RQT_AbsoluteTime
-		check(0);
+		check(ActiveQueries.Count > 0);
+		ActiveQueries.Count--;
+
+		PooledQueries[Type].Add(Resource);
 	}
+	
+	Resource = 0;
+
+	Unlink();
 }
 
-void FOpenGLDynamicRHI::EndRenderQuery_OnThisThread(FOpenGLRenderQuery* Query)
+void FOpenGLRenderQuery::Begin()
 {
 	VERIFY_GL_SCOPE();
 
-	if (Query)
-	{
-		if (Query->QueryType == RQT_Occlusion)
-		{
-			if (!Query->bInvalidResource && !PlatformContextIsCurrent(Query->ResourceContext))
-			{
-				PlatformReleaseRenderQuery(Query->Resource, Query->ResourceContext);
-				Query->Resource = 0;
-				Query->bInvalidResource = true;
-			}
+	check(Resource == 0);
+	AcquireGlQuery();
 
-			if (!Query->bInvalidResource)
-			{
-				check(PendingState.RunningOcclusionQuery == Query->Resource);
-				PendingState.RunningOcclusionQuery = 0;
-				GLenum QueryType = FOpenGL::SupportsExactOcclusionQueries() ? UGL_SAMPLES_PASSED : UGL_ANY_SAMPLES_PASSED;
-				FOpenGL::EndQuery(QueryType);
-			}
-		}
-		else if (Query->QueryType == RQT_AbsoluteTime)
-		{
-			int32 NewVal = Query->TotalBegins.Increment();
-			Query->TotalResults.Set(NewVal - 1);
-			Query->Result = 0;
-			Query->bResultWasSuccess = false;
+	CheckContext();
 
-			if (!Query->bInvalidResource && !PlatformContextIsCurrent(Query->ResourceContext))
-			{
-				PlatformReleaseRenderQuery(Query->Resource, Query->ResourceContext);
-				Query->Resource = 0;
-				Query->bInvalidResource = true;
-			}
+	switch(Type)
+	{
+	default:
+	case EType::Timestamp:
+		checkNoEntry();
+		break;
 
-			// query can be silently invalidated in GetRenderQueryResult
-			if (Query->bInvalidResource)
-			{
-				PlatformGetNewRenderQuery(&Query->Resource, &Query->ResourceContext);
-				Query->bInvalidResource = false;
-			}
+	case EType::Occlusion:
+		FOpenGL::BeginQuery(
+			FOpenGL::SupportsExactOcclusionQueries()
+				? UGL_SAMPLES_PASSED
+				: UGL_ANY_SAMPLES_PASSED
+			, Resource
+		);
+		break;
 
-			FOpenGL::QueryTimestampCounter(Query->Resource);
-		}
-	}
+	case EType::Disjoint:
+		FOpenGL::BeginQuery(UGL_TIME_ELAPSED, Resource);
+		break;
+	};
 }
 
-static void GetRenderQueryResult(FOpenGLRenderQuery* Query)
+void FOpenGLRenderQuery::End()
 {
 	VERIFY_GL_SCOPE();
-	if (Query->QueryType == RQT_AbsoluteTime)
-	{
-		FOpenGL::GetQueryObject(Query->Resource, FOpenGL::QM_Result, &Query->Result);
-	}
-	else
+	AcquireGlQuery();
+	
+	CheckContext();
+
+	switch (Type)
 	{
-		GLuint Result32 = 0;
-		FOpenGL::GetQueryObject(Query->Resource, FOpenGL::QM_Result, &Result32);
-		Query->Result = Result32 * (FOpenGL::SupportsExactOcclusionQueries() ? 1 : 500000); // half a mega pixel display
+	case EType::Occlusion:
+		check(Resource);
+		FOpenGL::EndQuery(FOpenGL::SupportsExactOcclusionQueries()
+			? UGL_SAMPLES_PASSED
+			: UGL_ANY_SAMPLES_PASSED
+		);
+		break;
+
+	case EType::Timestamp:
+		FOpenGL::QueryTimestampCounter(Resource);
+		break;
+
+	case EType::Disjoint:
+		FOpenGL::EndQuery(UGL_TIME_ELAPSED);
+		break;
 	}
-	Query->bResultWasSuccess = true;
-	Query->TotalResults.Increment();
+
+	BOPCounter++;
+
+	Link();
 }
 
-void FOpenGLDynamicRHI::GetRenderQueryResult_OnThisThread(FOpenGLRenderQuery* Query, bool bWait)
+bool FOpenGLRenderQuery::CacheResult(bool bWait)
 {
-	if (Query->TotalResults.GetValue() == Query->TotalBegins.GetValue())
+	if (BOPCounter == LastCachedBOPCounter.load(std::memory_order_relaxed))
 	{
-		return;
+		// Value has been cached and no newer query operation has started.
+		check(!IsLinked());
+		return true;
 	}
-	check(Query->TotalResults.GetValue() + 1 == Query->TotalBegins.GetValue());
 
-	VERIFY_GL_SCOPE();
+	CheckContext();
 
-	if (!Query->bInvalidResource && !PlatformContextIsCurrent(Query->ResourceContext))
+	if (!bWait)
 	{
-		PlatformReleaseRenderQuery(Query->Resource, Query->ResourceContext);
-		Query->Resource = 0;
-		Query->bInvalidResource = true;
-	}
+		// If we don't want to wait, we need to check if the result is available first.
+		GLuint IsAvailable = GL_FALSE;
+		FOpenGL::GetQueryObject(Resource, FOpenGL::QM_ResultAvailable, &IsAvailable);
 
-	// Check if the query is valid first
-	if (Query->bInvalidResource)
-	{
-		Query->Result = 0;
-		Query->TotalResults.Increment();
+		if (IsAvailable == GL_FALSE)
+		{
+			// Not ready yet.
+			return false;
+		}
 	}
-	else
+
+	// Read the result back (and block if its not ready)
+	switch (Type)
 	{
-		// Check if the query is finished
-		GLuint Result = 0;
-		FOpenGL::GetQueryObject(Query->Resource, FOpenGL::QM_ResultAvailable, &Result);
-		if (Result == GL_TRUE)
+	default:
+		checkNoEntry();
+		break;
+	
+	case EType::Occlusion:
 		{
-			GetRenderQueryResult(Query);
+			GLuint Result32 = 0;
+			FOpenGL::GetQueryObject(Resource, FOpenGL::QM_Result, &Result32);
+			SetResult(Result32 * (FOpenGL::SupportsExactOcclusionQueries() ? 1 : 500000)); // half a mega pixel display
 		}
-		else if (bWait) // Isn't the query finished yet, and can we wait for it?
+		break;
+
+	case EType::Timestamp:
 		{
-			SCOPE_CYCLE_COUNTER(STAT_RenderQueryResultTime);
+			GLuint64 Value = 0;
+			FOpenGL::GetQueryObject(Resource, FOpenGL::QM_Result, &Value);
 
-			FRenderThreadIdleScope IdleScope(ERenderThreadIdleTypes::WaitingForGPUQuery);
-			GBatcher.Waited();
-			
-			if (GOpenGLPollRenderQueryResult == 0)
-			{
-				// block in the driver waiting for result
-				GetRenderQueryResult(Query);
-			}
-			else
+			// Convert to microseconds (GL queries are in nanoseconds)
+			SetResult(Value / 1000);
+		}
+		break;
+
+	case EType::Disjoint:
+		{
+			// TimerQueryDisjoint is a one-shot state in the driver, it is not pipelined.
+			// If it returns true, all timers we've submitted after this timer but haven't
+			// yet resolved should be discarded for having invalid data.
+			if (FOpenGL::TimerQueryDisjoint())
 			{
-				// poll result until it's ready
-				double StartTime = FPlatformTime::Seconds();
-				do
+				for (FOpenGLRenderQuery* Other = this; Other; Other = Other->Next)
 				{
-					FPlatformProcess::Sleep(0);	// yield to other threads - some of them may be OpenGL driver's and we'd be starving them
-
-					if (Query->bInvalidResource)
-					{
-						// Query got invalidated while we were sleeping.
-						// Bail out, no sense to wait and generate OpenGL errors,
-						// we're in a new OpenGL context that knows nothing about us.
-						Query->Result = 1000;	// safe value
-						Result = GL_FALSE;
-						bWait = false;
-						Query->bResultWasSuccess = true;
-						break;
-					}
-
-					FOpenGL::GetQueryObject(Query->Resource, FOpenGL::QM_ResultAvailable, &Result);
-
-					// timer queries are used for Benchmarks which can stall a bit more
-					double TimeoutValue = (Query->QueryType == RQT_AbsoluteTime) ? 2.0 : 0.5;
-
-					if ((FPlatformTime::Seconds() - StartTime) > TimeoutValue)
+					if (Other->Type == EType::Disjoint)
 					{
-						UE_LOG(LogRHI, Log, TEXT("Timed out while waiting for GPU to catch up. (%.1f s)"), TimeoutValue);
-						break;
+						Other->SetResult(InvalidDisjointMask);
 					}
-				} while (Result == GL_FALSE);
-				
-				if (Result == GL_TRUE)
-				{
-					GetRenderQueryResult(Query);
-				}
-				else
-				{
-					Query->Result = 0;
-					Query->TotalResults.Increment();
 				}
 			}
+			else
+			{
+				GLuint64 Value;
+				FOpenGL::GetQueryObject(Resource, FOpenGL::QM_Result, &Value);
+
+				// Convert to microseconds (GL queries are in nanoseconds)
+				SetResult(Value / 1000);
+			}
 		}
+		break;
 	}
+
+	return true;
 }
 
-class FPollQueriesRHIThreadTask
+void FOpenGLRenderQuery::SetResult(uint64 Value)
 {
-	FOpenGLRenderQuery* Query;
-	FOpenGLDynamicRHI* RHI;
-	bool bWait;
-
-public:
+	Result = Value;
+	ReleaseGlQuery();
 
-	FPollQueriesRHIThreadTask(FOpenGLRenderQuery* InQuery, FOpenGLDynamicRHI* InRHI, bool bInWait)
-		: Query(InQuery)
-		, RHI(InRHI)
-		, bWait(bInWait)
-	{
-	}
+	LastCachedBOPCounter.store(BOPCounter, std::memory_order_release);
+}
 
-	FORCEINLINE TStatId GetStatId() const
+bool FOpenGLRenderQuery_RHI::GetResult(bool bWait, uint64& OutResult)
+{
+	if (TOPCounter == LastCachedBOPCounter.load(std::memory_order_acquire))
 	{
-		RETURN_QUICK_DECLARE_CYCLE_STAT(FPollQueriesRHIThreadTask, STATGROUP_TaskGraphTasks);
+		// Early return for queries we already have the result for.
+		check(!IsLinked());
+		OutResult = FOpenGLRenderQuery::GetResult();
+		return true;
 	}
 
-	ENamedThreads::Type GetDesiredThread()
+	if (!bWait)
 	{
-		return ENamedThreads::SetTaskPriority(ENamedThreads::RHIThread, ENamedThreads::HighTaskPriority);
+		//
+		// The query has not yet completed, and we don't want to wait for the query result.
+		// Return. The RHI thread will poll for results later.
+		//
+		OutResult = 0;
+		return false;
 	}
 
-	static ESubsequentsMode::Type GetSubsequentsMode() { return ESubsequentsMode::TrackSubsequents; }
+	FRHICommandListImmediate& RHICmdList = FRHICommandListImmediate::Get();
 
-	void DoTask(ENamedThreads::Type CurrentThread, const FGraphEventRef& MyCompletionGraphEvent)
+	//
+	// The query has not yet completed, and we want to wait for results.
+	// Append an RHI thread command that will force a readback of the GL query, then flush the RHI thread.
+	//	
+	RHICmdList.EnqueueLambda([Counter = TOPCounter, this](FRHICommandListImmediate&)
 	{
-		check(IsInRHIThread());
-		check(IsRunningRHIInDedicatedThread() && IsInRHIThread()); // this should never be used on a platform that doesn't support the RHI thread, and it can't quite work when running the RHI stuff on task threads
-		if (bWait)
-		{
-			RHI->GetRenderQueryResult_OnThisThread(Query, true); // we must get this one if bWait is true;
-			RHI->RHIPollOcclusionQueries(); // finish any other ones, but don't wait
-		}
-		else
+		if (BOPCounter != Counter)
 		{
-			RHI->GetRenderQueryResult_OnThisThread(Query, false);
-			if (Query->TotalResults.GetValue() == Query->TotalBegins.GetValue())
-			{
-				RHI->RHIPollOcclusionQueries(); // If the target query was ready, then go ahead and scan to see what else is ready.
-			}
+			// Query result is not ready yet. Fetch it...
+			CacheResult(true);
 		}
-	}
-};
+	});
 
+	// Wait for the above lambda to execute
+	RHICmdList.ImmediateFlush(EImmediateFlushType::FlushRHIThread);
 
-bool FOpenGLDynamicRHI::RHIGetRenderQueryResult(FRHIRenderQuery* QueryRHI, uint64& OutResult, bool bWait, uint32 GPUIndex)
-{
-	check(IsInRenderingThread() || IsInRHIThread());
+	checkf(TOPCounter == BOPCounter, TEXT("Attempting to get data from an RHI render query which was never issued."));
+	check(!IsLinked());
 
-	FOpenGLRenderQuery* Query = ResourceCast(QueryRHI)->GetInnerQuery();
+	OutResult = FOpenGLRenderQuery::GetResult();
+	return true;
+}
 
-	if (!Query)
+void FOpenGLRenderQuery::PollQueryResults()
+{
+	if (ActiveQueries.First)
 	{
-		// If timer queries are unsupported, just make sure that OutResult does not contain any random values.
-		OutResult = 0;
-		return false;
-	}
+		TRACE_CPUPROFILER_EVENT_SCOPE(PollQueryResults);
 
-	FRHICommandListImmediate& RHICmdList = FRHICommandListExecutor::GetImmediateCommandList();
+		do
+		{
+			if (!ActiveQueries.First->CacheResult(/*bWait = */ false))
+				break;
+		}
+		while (ActiveQueries.First);
+	}
+}
 
-	const bool bCanRunOnThisThread = RHICmdList.Bypass() || (!IsRunningRHIInSeparateThread() && IsInRenderingThread()) || IsInRHIThread();
+void FOpenGLRenderQuery::Cleanup()
+{
+	VERIFY_GL_SCOPE();
+	check(ActiveQueries.Count == 0);
 
-	if (Query->TotalResults.GetValue() != Query->TotalBegins.GetValue())
+	for (auto& Array : PooledQueries)
 	{
-		if (bCanRunOnThisThread)
+		for (GLuint Resource : Array)
 		{
-			GetRenderQueryResult_OnThisThread(Query, bWait);
+			FOpenGL::DeleteQueries(1, &Resource);
 		}
-		else
-		{
-			if (bWait)
-			{
-				QUICK_SCOPE_CYCLE_COUNTER(STAT_WaitForRHIThreadOcclusionReadback);
-				if (IsRunningRHIInDedicatedThread())
-				{
-					// send a command that will wait, so if the RHIT runs out of work, it just blocks and waits for the GPU
-					ALLOC_COMMAND_CL(RHICmdList, FRHICommandGLCommand)([this, QueryRHI]() {GetRenderQueryResult_OnThisThread(ResourceCast(QueryRHI)->GetInnerQuery(), true); });
-					FGraphEventRef Done = RHICmdList.RHIThreadFence(false);
-					ALLOC_COMMAND_CL(RHICmdList, FRHICommandGLCommand)([this, QueryRHI]() {GBatcher.Flush(*this, QueryRHI); });
-					RHICmdList.ImmediateFlush(EImmediateFlushType::DispatchToRHIThread);
-					while (!Done->IsComplete())
-					{
-						FGraphEventRef RHITask = TGraphTask<FPollQueriesRHIThreadTask>::CreateTask().ConstructAndDispatchWhenReady(ResourceCast(QueryRHI)->GetInnerQuery(), this, false);
-						FTaskGraphInterface::Get().WaitUntilTaskCompletes(RHITask);
-
-						if (Query->TotalResults.GetValue() == Query->TotalBegins.GetValue())
-						{
-							break;
-						}
-						// We want to keep the RHIT working, but we want keep checking between command lists so that we can get the results as soon as the GPU has them
 
-						// this isn't really a spin, the ping-pong between threads will not consume CPU (usually a bad thing, not here).
-					}
-				}
-				else
-				{
-					ALLOC_COMMAND_CL(RHICmdList, FRHICommandGLCommand)([this, QueryRHI]() {GetRenderQueryResult_OnThisThread(ResourceCast(QueryRHI)->GetInnerQuery(), true); });
-					FGraphEventRef Done = RHICmdList.RHIThreadFence(false);
-					ALLOC_COMMAND_CL(RHICmdList, FRHICommandGLCommand)([this, QueryRHI]() {GBatcher.Flush(*this, QueryRHI); });
-					RHICmdList.ImmediateFlush(EImmediateFlushType::DispatchToRHIThread);
-					FRHICommandListExecutor::WaitOnRHIThreadFence(Done);
-				}
-				check(Query->TotalResults.GetValue() == Query->TotalBegins.GetValue());
-			}
-			else
-			{
-				ALLOC_COMMAND_CL(RHICmdList, FRHICommandGLCommand)([this, QueryRHI]() {GetRenderQueryResult_OnThisThread(ResourceCast(QueryRHI)->GetInnerQuery(), false); GBatcher.Flush(*this, QueryRHI);  });
-			}
-		}	
+		Array.Reset();
 	}
-	if (Query->TotalResults.GetValue() == Query->TotalBegins.GetValue() && Query->bResultWasSuccess)
+}
+
+FRenderQueryRHIRef FOpenGLDynamicRHI::RHICreateRenderQuery(ERenderQueryType QueryType)
+{
+	check(QueryType == RQT_Occlusion || QueryType == RQT_AbsoluteTime);
+	if (QueryType == RQT_AbsoluteTime && FOpenGL::SupportsTimestampQueries() == false)
 	{
-		if (Query->QueryType == RQT_AbsoluteTime)
-		{
-#if RHI_NEW_GPU_PROFILER
-			checkNoEntry(); // @todo - new gpu profiler
-#else
-			// GetTimingFrequency is the number of ticks per second
-			uint64 Div = FMath::Max(1llu, FOpenGLBufferedGPUTiming::GetTimingFrequency() / (1000 * 1000));
-
-			// convert from GPU specific timestamp to micro sec (1 / 1 000 000 s) which seems a reasonable resolution
-			OutResult = Query->Result / Div;
-#endif
-		}
-		else
-		{
-			OutResult = Query->Result;
-		}
-		return true;
+		return nullptr;
 	}
-	OutResult = 0;
-	return false;
+
+	return new FOpenGLRenderQuery_RHI(QueryType);
 }
 
-extern void OnQueryCreation( FOpenGLRenderQuery* Query );
-extern void OnQueryDeletion( FOpenGLRenderQuery* Query );
+void FOpenGLDynamicRHI::RHIBeginRenderQuery_TopOfPipe(FRHICommandListBase& RHICmdList, FRHIRenderQuery* RenderQuery)
+{
+	if (!RenderQuery)
+		return;
 
-FOpenGLRenderQuery::FOpenGLRenderQuery(ERenderQueryType InQueryType)
-	: QueryType(InQueryType)
+	FDynamicRHI::RHIBeginRenderQuery_TopOfPipe(RHICmdList, RenderQuery);
+}
+
+void FOpenGLDynamicRHI::RHIEndRenderQuery_TopOfPipe(FRHICommandListBase& RHICmdList, FRHIRenderQuery* RenderQuery)
 {
-	VERIFY_GL_SCOPE();
+	if (!RenderQuery)
+		return;
 
-	bInvalidResource = false;
-	PlatformGetNewRenderQuery(&Resource, &ResourceContext);
-	
-	OnQueryCreation(this);
+	ResourceCast(RenderQuery)->End_TopOfPipe();
+	FDynamicRHI::RHIEndRenderQuery_TopOfPipe(RHICmdList, RenderQuery);
 }
 
-FOpenGLRenderQuery::~FOpenGLRenderQuery()
+void FOpenGLDynamicRHI::RHIBeginRenderQuery(FRHIRenderQuery* RenderQuery)
 {
-	VERIFY_GL_SCOPE();
-	OnQueryDeletion(this);
+	ResourceCast(RenderQuery)->Begin();
+}
 
-	if (Resource && !bInvalidResource)
-	{
-		bInvalidResource = true;
+void FOpenGLDynamicRHI::RHIEndRenderQuery(FRHIRenderQuery* RenderQuery)
+{
+	FOpenGLRenderQuery_RHI* Query = ResourceCast(RenderQuery);
+	Query->End();
+}
 
-		check(Resource);
-		PlatformReleaseRenderQuery(Resource, ResourceContext);
+bool FOpenGLDynamicRHI::RHIGetRenderQueryResult(FRHIRenderQuery* QueryRHI, uint64& OutResult, bool bWait, uint32 GPUIndex)
+{
+	if (!QueryRHI)
+	{
+		OutResult = 0;
+		return true;
 	}
-}
 
+	FOpenGLRenderQuery_RHI* Query = ResourceCast(QueryRHI);
+	return Query->GetResult(bWait, OutResult);
+}
 
 void FOpenGLEventQuery::IssueEvent()
 {
@@ -759,7 +534,7 @@ static FOpenGLRenderQuery* GetTimeQuery()
 	{
 		return TimerQueryPool.Pop();
 	}
-	return new FOpenGLRenderQuery(RQT_AbsoluteTime);
+	return new FOpenGLRenderQuery(FOpenGLRenderQuery::EType::Timestamp);
 }
 
 void FOpenGLBufferedGPUTiming::InitResources()
@@ -791,12 +566,12 @@ void FOpenGLBufferedGPUTiming::ReleaseResources()
 {
 	VERIFY_GL_SCOPE();
 
-	for(FOpenGLRenderQuery* Query : StartTimestamps)
+	for (FOpenGLRenderQuery* Query : StartTimestamps)
 	{
 		TimerQueryPool.Add(Query);
 	}
 
-	for(FOpenGLRenderQuery* Query : EndTimestamps)
+	for (FOpenGLRenderQuery* Query : EndTimestamps)
 	{
 		TimerQueryPool.Add(Query);
 	}
@@ -816,22 +591,8 @@ void FOpenGLBufferedGPUTiming::StartTiming()
 	if ( GIsSupported && !bIsTiming )
 	{
 		int32 NewTimestampIndex = (CurrentTimestamp + 1) % BufferSize;
-		FOpenGLRenderQuery* TimerQuery = StartTimestamps[NewTimestampIndex];
-		{
-			if (!TimerQuery->bInvalidResource && !PlatformContextIsCurrent(TimerQuery->ResourceContext))
-			{
-				PlatformReleaseRenderQuery(TimerQuery->Resource, TimerQuery->ResourceContext);
-				TimerQuery->bInvalidResource = true;
-			}
+		StartTimestamps[NewTimestampIndex]->End();
 
-			if (TimerQuery->bInvalidResource)
-			{
-				PlatformGetNewRenderQuery(&TimerQuery->Resource, &TimerQuery->ResourceContext);
-				TimerQuery->bInvalidResource = false;
-			}
-		}
-
-		FOpenGL::QueryTimestampCounter(StartTimestamps[NewTimestampIndex]->Resource);
 		CurrentTimestamp = NewTimestampIndex;
 		bIsTiming = true;
 	}
@@ -848,23 +609,8 @@ void FOpenGLBufferedGPUTiming::EndTiming()
 	if ( GIsSupported && bIsTiming )
 	{
 		checkSlow( CurrentTimestamp >= 0 && CurrentTimestamp < BufferSize );
+		EndTimestamps[CurrentTimestamp]->End();
 
-		FOpenGLRenderQuery* TimerQuery = EndTimestamps[CurrentTimestamp];
-		{
-			if (!TimerQuery->bInvalidResource && !PlatformContextIsCurrent(TimerQuery->ResourceContext))
-			{
-				PlatformReleaseRenderQuery(TimerQuery->Resource, TimerQuery->ResourceContext);
-				TimerQuery->bInvalidResource = true;
-			}
-
-			if (TimerQuery->bInvalidResource && PlatformOpenGLContextValid())
-			{
-				PlatformGetNewRenderQuery(&TimerQuery->Resource, &TimerQuery->ResourceContext);
-				TimerQuery->bInvalidResource = false;
-			}
-		}
-
-		FOpenGL::QueryTimestampCounter(EndTimestamps[CurrentTimestamp]->Resource);
 		NumIssuedTimestamps = FMath::Min<int32>(NumIssuedTimestamps + 1, BufferSize);
 		bIsTiming = false;
 	}
@@ -878,60 +624,29 @@ void FOpenGLBufferedGPUTiming::EndTiming()
  */
 uint64 FOpenGLBufferedGPUTiming::GetTiming(bool bGetCurrentResultsAndBlock)
 {
-
 	VERIFY_GL_SCOPE();
 
-	if ( GIsSupported )
+	if (GIsSupported)
 	{
-		checkSlow( CurrentTimestamp >= 0 && CurrentTimestamp < BufferSize );
-		GLuint64 StartTime, EndTime;
-
+		checkSlow(CurrentTimestamp >= 0 && CurrentTimestamp < BufferSize);
 		int32 TimestampIndex = CurrentTimestamp;
 
-		{
-			FOpenGLRenderQuery* EndStamp = EndTimestamps[TimestampIndex];
-			if (!EndStamp->bInvalidResource && !PlatformContextIsCurrent(EndStamp->ResourceContext))
-			{
-				PlatformReleaseRenderQuery(EndStamp->Resource, EndStamp->ResourceContext);
-				EndStamp->bInvalidResource = true;
-			}
-
-			FOpenGLRenderQuery* StartStamp = StartTimestamps[TimestampIndex];
-			if (!StartStamp->bInvalidResource && !PlatformContextIsCurrent(StartStamp->ResourceContext))
-			{
-				PlatformReleaseRenderQuery(StartStamp->Resource, StartStamp->ResourceContext);
-				StartStamp->bInvalidResource = true;
-
-			}
-
-			if(StartStamp->bInvalidResource || EndStamp->bInvalidResource)
-			{
-				UE_LOG(LogRHI, Log, TEXT("timing invalid, since the stamp queries have invalid resources"));
-				return 0.0f;
-			}
-		}
-
 		if (!bGetCurrentResultsAndBlock)
 		{
 			// Quickly check the most recent measurements to see if any of them has been resolved.  Do not flush these queries.
-			for ( int32 IssueIndex = 1; IssueIndex < NumIssuedTimestamps; ++IssueIndex )
+			for (int32 IssueIndex = 1; IssueIndex < NumIssuedTimestamps; ++IssueIndex)
 			{
-				GLuint EndAvailable = GL_FALSE;
-				FOpenGL::GetQueryObject(EndTimestamps[TimestampIndex]->Resource, FOpenGL::QM_ResultAvailable, &EndAvailable);
+				FOpenGLRenderQuery* StartQuery = StartTimestamps[TimestampIndex];
+				FOpenGLRenderQuery* EndQuery = EndTimestamps[TimestampIndex];
 
-				if ( EndAvailable == GL_TRUE )
+				if (StartQuery->CacheResult(false) && EndQuery->CacheResult(false))
 				{
-					GLuint StartAvailable = GL_FALSE;
-					FOpenGL::GetQueryObject(StartTimestamps[TimestampIndex]->Resource, FOpenGL::QM_ResultAvailable, &StartAvailable);
+					uint64 StartTime = StartQuery->GetResult();
+					uint64 EndTime = EndQuery->GetResult();
 
-					if(StartAvailable == GL_TRUE)
+					if (EndTime > StartTime)
 					{
-						FOpenGL::GetQueryObject(EndTimestamps[TimestampIndex]->Resource, FOpenGL::QM_Result, &EndTime);
-						FOpenGL::GetQueryObject(StartTimestamps[TimestampIndex]->Resource, FOpenGL::QM_Result, &StartTime);
-						if (EndTime > StartTime)
-						{
-							return EndTime - StartTime;
-						}
+						return EndTime - StartTime;
 					}
 				}
 
@@ -939,59 +654,53 @@ uint64 FOpenGLBufferedGPUTiming::GetTiming(bool bGetCurrentResultsAndBlock)
 			}
 		}
 
-		if ( NumIssuedTimestamps > 0 || bGetCurrentResultsAndBlock )
+		if (NumIssuedTimestamps > 0 || bGetCurrentResultsAndBlock)
 		{
 			// None of the (NumIssuedTimestamps - 1) measurements were ready yet,
 			// so check the oldest measurement more thoroughly.
 			// This really only happens if occlusion and frame sync event queries are disabled, otherwise those will block until the GPU catches up to 1 frame behind
 			const bool bBlocking = ( NumIssuedTimestamps == BufferSize ) || bGetCurrentResultsAndBlock;
 
-			GLuint EndAvailable = GL_FALSE;
+			FOpenGLRenderQuery* StartQuery = StartTimestamps[TimestampIndex];
+			FOpenGLRenderQuery* EndQuery = EndTimestamps[TimestampIndex];
+
+			bool bHasStart = false, bHasEnd = false;
+
 			{
 				FRenderThreadIdleScope IdleScope(ERenderThreadIdleTypes::WaitingForGPUQuery);
+				SCOPE_CYCLE_COUNTER(STAT_RenderQueryResultTime);
+
 				double StartTimeoutTime = FPlatformTime::Seconds();
 
-				SCOPE_CYCLE_COUNTER( STAT_RenderQueryResultTime );
 				// If we are blocking, retry until the GPU processes the time stamp command
-				do
+				while (true)
 				{
-					FOpenGL::GetQueryObject(EndTimestamps[TimestampIndex]->Resource, FOpenGL::QM_ResultAvailable, &EndAvailable);
-
-					if ((FPlatformTime::Seconds() - StartTimeoutTime) > 0.5)
-					{
-						UE_LOG(LogRHI, Log, TEXT("Timed out while waiting for GPU to catch up. (500 ms) EndTimeStamp"));
-						return 0;
-					}
-				} while ( EndAvailable == GL_FALSE && bBlocking );
-			}
-
-			if ( EndAvailable == GL_TRUE )
-			{
-				GLuint StartAvailable = GL_FALSE;
-				{
-					FRenderThreadIdleScope IdleScope(ERenderThreadIdleTypes::WaitingForGPUQuery);
-					double StartTimeoutTime = FPlatformTime::Seconds();
-
-					do
+					bHasStart = StartQuery->CacheResult(false);
+					bHasEnd = EndQuery->CacheResult(false);
+					
+					if (bBlocking && !(bHasStart && bHasEnd))
 					{
-						FOpenGL::GetQueryObject(StartTimestamps[TimestampIndex]->Resource, FOpenGL::QM_ResultAvailable, &StartAvailable);
-
 						if ((FPlatformTime::Seconds() - StartTimeoutTime) > 0.5)
 						{
-							UE_LOG(LogRHI, Log, TEXT("Timed out while waiting for GPU to catch up. (500 ms) StartTimeStamp"));
+							UE_LOG(LogRHI, Log, TEXT("Timed out while waiting for GPU to catch up. (500 ms)"));
 							return 0;
 						}
-					} while ( StartAvailable == GL_FALSE && bBlocking );
+					}
+					else
+					{
+						break;
+					}
 				}
+			}
+
+			if (bHasStart && bHasEnd)
+			{
+				uint64 StartTime = StartQuery->GetResult();
+				uint64 EndTime = EndQuery->GetResult();
 
-				if(StartAvailable == GL_TRUE)
+				if (EndTime > StartTime)
 				{
-					FOpenGL::GetQueryObject(EndTimestamps[TimestampIndex]->Resource, FOpenGL::QM_Result, &EndTime);
-					FOpenGL::GetQueryObject(StartTimestamps[TimestampIndex]->Resource, FOpenGL::QM_Result, &StartTime);
-					if (EndTime > StartTime)
-					{
-						return EndTime - StartTime;
-					}
+					return EndTime - StartTime;
 				}
 			}
 		}
@@ -1004,15 +713,7 @@ void FOpenGLDisjointTimeStampQuery::StartTracking()
 	VERIFY_GL_SCOPE();
 	if (IsSupported())
 	{
-
-		if (!PlatformContextIsCurrent(Context))
-		{
-			PlatformReleaseRenderQuery(DisjointQuery, Context);
-			PlatformGetNewRenderQuery(&DisjointQuery, &Context);
-		}
-		// Dummy query to reset the driver's internal disjoint status
-		FOpenGL::TimerQueryDisjoint();
-		FOpenGL::BeginQuery(UGL_TIME_ELAPSED, DisjointQuery);
+		DisjointQuery.Begin();
 	}
 }
 
@@ -1020,15 +721,10 @@ void FOpenGLDisjointTimeStampQuery::EndTracking()
 {
 	VERIFY_GL_SCOPE();
 
-	if(IsSupported())
+	if (IsSupported())
 	{
-		FOpenGL::EndQuery( UGL_TIME_ELAPSED );
-
-		// Check if the GPU changed clock frequency since the last time GL_GPU_DISJOINT_EXT was checked.
-		// If so, any timer query will be undefined.
-		bIsResultValid = !FOpenGL::TimerQueryDisjoint();
+		DisjointQuery.End();
 	}
-
 }
 
 bool FOpenGLDisjointTimeStampQuery::IsResultValid()
@@ -1037,53 +733,21 @@ bool FOpenGLDisjointTimeStampQuery::IsResultValid()
 	return bIsResultValid;
 }
 
-bool FOpenGLDisjointTimeStampQuery::GetResult( uint64* OutResult/*=NULL*/ )
+bool FOpenGLDisjointTimeStampQuery::GetResult(uint64* OutResult)
 {
 	VERIFY_GL_SCOPE();
 
 	if (IsSupported())
 	{
-		GLuint Result = 0;
-		FOpenGL::GetQueryObject(DisjointQuery, FOpenGL::QM_ResultAvailable, &Result);
-		const double StartTime = FPlatformTime::Seconds();
+		DisjointQuery.CacheResult(true);
 
-		while (Result == GL_FALSE && (FPlatformTime::Seconds() - StartTime) < 0.5)
-		{
-			FPlatformProcess::Sleep(0.005f);
-			FOpenGL::GetQueryObject(DisjointQuery, FOpenGL::QM_ResultAvailable, &Result);
-		}
+		uint64 Result = DisjointQuery.GetResult();
+		bIsResultValid = (Result & FOpenGLRenderQuery::InvalidDisjointMask) == 0;
 
-		// Presently just discarding the result, because timing is handled by timestamps inside
-		if (Result != GL_FALSE)
-		{
-			GLuint64 ElapsedTime = 0;
-			FOpenGL::GetQueryObject(DisjointQuery, FOpenGL::QM_Result, &ElapsedTime);
-			if (OutResult)
-			{
-				*OutResult = ElapsedTime;
-			}
-		}
-		bIsResultValid = Result != GL_FALSE;
+		*OutResult = Result & (~FOpenGLRenderQuery::InvalidDisjointMask);
 	}
-	return bIsResultValid;
-}
 
-void FOpenGLDisjointTimeStampQuery::InitResources()
-{
-	VERIFY_GL_SCOPE();
-	if (IsSupported())
-	{
-		PlatformGetNewRenderQuery(&DisjointQuery, &Context); 
-	}
-}
-
-void FOpenGLDisjointTimeStampQuery::ReleaseResources()
-{
-	VERIFY_GL_SCOPE();
-	if ( IsSupported() )
-	{
-		PlatformReleaseRenderQuery(DisjointQuery, Context);
-	}
+	return bIsResultValid;
 }
 
 #endif // (RHI_NEW_GPU_PROFILER == 0)
diff --git a/Engine/Source/Runtime/OpenGLDrv/Private/OpenGLRenderTarget.cpp b/Engine/Source/Runtime/OpenGLDrv/Private/OpenGLRenderTarget.cpp
--- a/Engine/Source/Runtime/OpenGLDrv/Private/OpenGLRenderTarget.cpp
+++ b/Engine/Source/Runtime/OpenGLDrv/Private/OpenGLRenderTarget.cpp
@@ -1056,12 +1056,6 @@ void FOpenGLDynamicRHI::RHIBeginRenderPass(const FRHIRenderPassInfo& InInfo, con
 
 	RenderPassInfo = InInfo;
 
-	if (InInfo.NumOcclusionQueries > 0)
-	{
-		extern void BeginOcclusionQueryBatch(uint32);
-		BeginOcclusionQueryBatch(InInfo.NumOcclusionQueries);
-	}
-
 #if PLATFORM_ANDROID
 	if (RenderPassInfo.SubpassHint == ESubpassHint::DeferredShadingSubpass &&
 		 FOpenGL::SupportsPixelLocalStorage() && FOpenGL::SupportsShaderDepthStencilFetch())
@@ -1085,12 +1079,6 @@ void FOpenGLDynamicRHI::RHIBeginRenderPass(const FRHIRenderPassInfo& InInfo, con
 
 void FOpenGLDynamicRHI::RHIEndRenderPass()
 {
-	if (RenderPassInfo.NumOcclusionQueries > 0)
-	{
-		extern void EndOcclusionQueryBatch();
-		EndOcclusionQueryBatch();
-	}
-
 	// End GL_EXT_multisampled_render_to_texture
 	PendingState.NumRenderingSamples = 1;
 
diff --git a/Engine/Source/Runtime/OpenGLDrv/Private/OpenGLResources.h b/Engine/Source/Runtime/OpenGLDrv/Private/OpenGLResources.h
--- a/Engine/Source/Runtime/OpenGLDrv/Private/OpenGLResources.h
+++ b/Engine/Source/Runtime/OpenGLDrv/Private/OpenGLResources.h
@@ -1158,51 +1158,99 @@ public:
 	uint8 const bMultisampleRenderbuffer : 1;
 };
 
-class FOpenGLRenderQuery final
+class FOpenGLRenderQuery
 {
 public:
-	/** The query resource. */
-	GLuint Resource;
+	enum class EType : uint8
+	{
+		Timestamp,
+		Occlusion,
+		Disjoint,
 
-	/** Identifier of the OpenGL context the query is a part of. */
-	uint64 ResourceContext;
+		Num
+	};
 
-	/** The cached query result. */
-	GLuint64 Result = 0;
+	static constexpr uint64 InvalidDisjointMask = 0x8000000000000000;
 
-	FThreadSafeCounter TotalBegins;
-	FThreadSafeCounter TotalResults;
+private:
+	// Render queries that should be polled by the RHI thread.
+	struct FActiveQueries
+	{
+		FOpenGLRenderQuery* First = nullptr;
+		FOpenGLRenderQuery* Last = nullptr;
+		int32 Count = 0;
+	} static ActiveQueries;
 
-	/** true if the context the query is in was released from another thread */
-	bool bResultWasSuccess;
+	struct FQueryPool : public TStaticArray<TArray<GLuint>, (uint32)EType::Num>
+	{
+		TArray<GLuint>& operator [](EType InType) { return TStaticArray::operator[]((uint32)InType); }
+	} static PooledQueries;
 
-	/** true if the context the query is in was released from another thread */
-	bool bInvalidResource = true;
+	// Linked list pointers. Used to build a list of "active" queries, i.e. queries that need data to be polled from the GPU.
+	FOpenGLRenderQuery** Prev = nullptr;
+	FOpenGLRenderQuery* Next = nullptr;
 
-	// todo: memory optimize
-	const ERenderQueryType QueryType;
+	uint64 Result = 0;
 
-	FOpenGLRenderQuery(ERenderQueryType InQueryType);
-	~FOpenGLRenderQuery();
-};
+	/** The query resource. */
+	GLuint Resource = 0;
 
-class FOpenGLRenderQuery_RHI : public FRHIRenderQuery
-{
-	TOptional<FOpenGLRenderQuery> Inner;
+	uint8 bSharedContext : 1 = false;
+
+protected:
+	EType const Type : 2;
+
+	std::atomic<uint8> LastCachedBOPCounter = 0;
+	uint8 BOPCounter = 0;
+	uint8 TOPCounter = 0;
 
 public:
-	FOpenGLRenderQuery_RHI(FRHICommandListBase& RHICmdList, ERenderQueryType QueryType)
+	FOpenGLRenderQuery(EType Type)
+		: Type(Type)
+	{}
+
+	~FOpenGLRenderQuery();
+
+	void AcquireGlQuery();
+	void ReleaseGlQuery();
+
+	bool IsLinked() const { return Prev != nullptr; }
+
+	void Begin();
+	void End();
+
+	uint64 GetResult() const
 	{
-		RHICmdList.EnqueueLambda([this, QueryType](FRHICommandListBase&)
-		{
-			Inner.Emplace(QueryType);
-		});
+		return Result;
 	}
 
-	FOpenGLRenderQuery* GetInnerQuery()
+	bool CacheResult(bool bWait);
+
+	static void PollQueryResults();
+	static void Cleanup();
+
+private:
+	void Link();
+	void Unlink();
+
+	void SetResult(uint64 Value);
+
+	void CheckContext();
+};
+
+class FOpenGLRenderQuery_RHI : public FRHIRenderQuery, public FOpenGLRenderQuery
+{
+public:
+	FOpenGLRenderQuery_RHI(ERenderQueryType QueryType)
+		: FOpenGLRenderQuery(QueryType == RQT_Occlusion ? EType::Occlusion : EType::Timestamp)
+	{}
+
+	void End_TopOfPipe()
 	{
-		return &Inner.GetValue();
+		TOPCounter++;
 	}
+
+	bool GetResult(bool bWait, uint64& OutResult);
 };
 
 class FOpenGLView : public TIntrusiveLinkedList<FOpenGLView>
diff --git a/Engine/Source/Runtime/OpenGLDrv/Private/OpenGLState.h b/Engine/Source/Runtime/OpenGLDrv/Private/OpenGLState.h
--- a/Engine/Source/Runtime/OpenGLDrv/Private/OpenGLState.h
+++ b/Engine/Source/Runtime/OpenGLDrv/Private/OpenGLState.h
@@ -430,7 +430,6 @@ struct FOpenGLRHIState final : public FOpenGLCommonState
 	GLuint							ZeroFilledDummyUniformBuffer;
 	uint32							RenderTargetWidth;
 	uint32							RenderTargetHeight;
-	GLuint							RunningOcclusionQuery;
 	bool							bAlphaToCoverageEnabled;
 
 	// Pending framebuffer setup
@@ -491,7 +490,6 @@ struct FOpenGLRHIState final : public FOpenGLCommonState
 	,	ZeroFilledDummyUniformBuffer(0)
 	,	RenderTargetWidth(0)
 	,	RenderTargetHeight(0)
-	,	RunningOcclusionQuery(0)
 	,	bAlphaToCoverageEnabled(false)
 	,	NumRenderingSamples(1)
 	,	FirstNonzeroRenderTarget(-1)
diff --git a/Engine/Source/Runtime/OpenGLDrv/Private/Windows/WindowsOpenGLPlatform.cpp b/Engine/Source/Runtime/OpenGLDrv/Private/Windows/WindowsOpenGLPlatform.cpp
--- a/Engine/Source/Runtime/OpenGLDrv/Private/Windows/WindowsOpenGLPlatform.cpp
+++ b/Engine/Source/Runtime/OpenGLDrv/Private/Windows/WindowsOpenGLPlatform.cpp
@@ -98,8 +98,6 @@ private:
 	bool				bSameDC;
 };
 
-void DeleteQueriesForCurrentContext(HGLRC Context);
-
 /**
  * A dummy wndproc.
  */
@@ -254,8 +252,6 @@ static void PlatformCreateOpenGLContextCore(FPlatformOpenGLContext* OutContext,
 
 void PlatformReleaseOpenGLContext(FPlatformOpenGLDevice* Device, FPlatformOpenGLContext* Context);
 
-extern void OnQueryInvalidation(void);
-
 /** Platform specific OpenGL device. */
 struct FPlatformOpenGLDevice
 {
@@ -313,7 +309,6 @@ struct FPlatformOpenGLDevice
 
 		ContextMakeCurrent(NULL, NULL);
 
-		OnQueryInvalidation();
 		PlatformReleaseOpenGLContext(this, &RenderingContext);
 		PlatformReleaseOpenGLContext(this, &SharedContext);
 
@@ -389,7 +384,6 @@ void PlatformReleaseOpenGLContext(FPlatformOpenGLDevice* Device, FPlatformOpenGL
 			bActiveContextWillBeReleased = ScopeContext.ContextWasAlreadyActive();
 			bSharedDC = ScopeContext.ContextsShareSameDC();
 
-			DeleteQueriesForCurrentContext(Context->OpenGLContext);
 			glBindVertexArray(0);
 			glDeleteVertexArrays(1, &Context->VertexArrayObject);
 
@@ -864,97 +858,6 @@ EOpenGLCurrentContext PlatformOpenGLCurrentContext(FPlatformOpenGLDevice* Device
 
 // =============================================================
 
-struct FOpenGLReleasedQuery
-{
-	HGLRC	Context;
-	GLuint	Query;
-};
-
-static TArray<FOpenGLReleasedQuery>	ReleasedQueries;
-static FCriticalSection* ReleasedQueriesGuard;
-
-void PlatformGetNewRenderQuery(GLuint* OutQuery, uint64* OutQueryContext)
-{
-	if (!ReleasedQueriesGuard)
-	{
-		ReleasedQueriesGuard = new FCriticalSection;
-	}
-
-	{
-		FScopeLock Lock(ReleasedQueriesGuard);
-
-#ifdef UE_BUILD_DEBUG
-		check(OutQuery && OutQueryContext);
-#endif
-
-		HGLRC Context = GetCurrentContext();
-		check(Context);
-
-		GLuint NewQuery = 0;
-
-		// Check for possible query reuse
-		const int32 ArraySize = ReleasedQueries.Num();
-		for (int32 Index = 0; Index < ArraySize; ++Index)
-		{
-			if (ReleasedQueries[Index].Context == Context)
-			{
-				NewQuery = ReleasedQueries[Index].Query;
-				ReleasedQueries.RemoveAtSwap(Index);
-				break;
-			}
-		}
-
-		if (!NewQuery)
-		{
-			FOpenGL::GenQueries(1, &NewQuery);
-		}
-
-		*OutQuery = NewQuery;
-		*OutQueryContext = (uint64)Context;
-	}
-}
-
-void PlatformReleaseRenderQuery(GLuint Query, uint64 QueryContext)
-{
-	HGLRC Context = GetCurrentContext();
-	if ((uint64)Context == QueryContext)
-	{
-		FOpenGL::DeleteQueries(1, &Query);
-	}
-	else
-	{
-		FScopeLock Lock(ReleasedQueriesGuard);
-#ifdef UE_BUILD_DEBUG
-		check(Query && QueryContext && ReleasedQueriesGuard);
-#endif
-		FOpenGLReleasedQuery ReleasedQuery;
-		ReleasedQuery.Context = (HGLRC)QueryContext;
-		ReleasedQuery.Query = Query;
-		ReleasedQueries.Add(ReleasedQuery);
-	}
-}
-
-void DeleteQueriesForCurrentContext(HGLRC Context)
-{
-	if (!ReleasedQueriesGuard)
-	{
-		ReleasedQueriesGuard = new FCriticalSection;
-	}
-
-	{
-		FScopeLock Lock(ReleasedQueriesGuard);
-		for (int32 Index = 0; Index < ReleasedQueries.Num(); ++Index)
-		{
-			if (ReleasedQueries[Index].Context == Context)
-			{
-				FOpenGL::DeleteQueries(1, &ReleasedQueries[Index].Query);
-				ReleasedQueries.RemoveAtSwap(Index);
-				--Index;
-			}
-		}
-	}
-}
-
 bool PlatformContextIsCurrent(uint64 QueryContext)
 {
 	return (uint64)GetCurrentContext() == QueryContext;
diff --git a/Engine/Source/Runtime/RHI/Public/RHICommandList.h b/Engine/Source/Runtime/RHI/Public/RHICommandList.h
--- a/Engine/Source/Runtime/RHI/Public/RHICommandList.h
+++ b/Engine/Source/Runtime/RHI/Public/RHICommandList.h
@@ -2343,11 +2343,6 @@ FRHICOMMAND_MACRO(FRHICommandPostExternalCommandsReset)
 	RHI_EXECUTE_API void Execute(FRHICommandListBase& CmdList);
 };
 
-FRHICOMMAND_MACRO(FRHICommandPollOcclusionQueries)
-{
-	RHI_EXECUTE_API void Execute(FRHICommandListBase& CmdList);
-};
-
 FRHICOMMAND_MACRO(FRHICommandBeginDrawingViewport)
 {
 	FRHIViewport* Viewport;
@@ -3885,14 +3880,9 @@ public:
 		ALLOC_COMMAND(FRHICommandCalibrateTimers)(CalibrationQuery);
 	}
 
+	UE_DEPRECATED(5.5, "PollOcclusionQueries is no longer necessary. It has been deprecated and there is no replacement. Remove any remaining calls to PollOcclusionQueries.")
 	FORCEINLINE_DEBUGGABLE void PollOcclusionQueries()
 	{
-		if (Bypass())
-		{
-			GetContext().RHIPollOcclusionQueries();
-			return;
-		}
-		ALLOC_COMMAND(FRHICommandPollOcclusionQueries)();
 	}
 
 	FORCEINLINE_DEBUGGABLE void BeginRenderPass(const FRHIRenderPassInfo& InInfo, const TCHAR* Name)
diff --git a/Engine/Source/Runtime/RHI/Public/RHICommandListCommandExecutes.inl b/Engine/Source/Runtime/RHI/Public/RHICommandListCommandExecutes.inl
--- a/Engine/Source/Runtime/RHI/Public/RHICommandListCommandExecutes.inl
+++ b/Engine/Source/Runtime/RHI/Public/RHICommandListCommandExecutes.inl
@@ -476,12 +476,6 @@ void FRHICommandPostExternalCommandsReset::Execute(FRHICommandListBase& CmdList)
 	INTERNAL_DECORATOR(RHIPostExternalCommandsReset)();
 }
 
-void FRHICommandPollOcclusionQueries::Execute(FRHICommandListBase& CmdList)
-{
-	RHISTAT(PollOcclusionQueries);
-	INTERNAL_DECORATOR(RHIPollOcclusionQueries)();
-}
-
 void FRHICommandCopyBufferRegion::Execute(FRHICommandListBase& CmdList)
 {
 	INTERNAL_DECORATOR(RHICopyBufferRegion)(DestBuffer, DstOffset, SourceBuffer, SrcOffset, NumBytes);
diff --git a/Engine/Source/Runtime/RHI/Public/RHIContext.h b/Engine/Source/Runtime/RHI/Public/RHIContext.h
--- a/Engine/Source/Runtime/RHI/Public/RHIContext.h
+++ b/Engine/Source/Runtime/RHI/Public/RHIContext.h
@@ -632,12 +632,6 @@ public:
 		/* empty default implementation */
 	}
 
-	// Used for OpenGL to check and see if any occlusion queries can be read back on the RHI thread. If they aren't ready when we need them, then we end up stalling.
-	virtual void RHIPollOcclusionQueries()
-	{
-		/* empty default implementation */
-	}
-
 	// Not all RHIs need this (Mobile specific)
 	virtual void RHIDiscardRenderTargets(bool Depth, bool Stencil, uint32 ColorBitMask) {}
 
diff --git a/Engine/Source/Runtime/RHI/Public/RHIGlobals.h b/Engine/Source/Runtime/RHI/Public/RHIGlobals.h
--- a/Engine/Source/Runtime/RHI/Public/RHIGlobals.h
+++ b/Engine/Source/Runtime/RHI/Public/RHIGlobals.h
@@ -206,8 +206,8 @@ struct FRHIGlobals
 	/** If an RHI is so slow, that it is the limiting factor for the entire frame, we can kick early to try to give it as much as possible. */
 	bool RHIThreadNeedsKicking = false;
 
-	/** If an RHI cannot do an unlimited number of occlusion queries without stalling and waiting for the GPU, this can be used to tune hte occlusion culler to try not to do that. */
-	int32 MaximumReccommendedOustandingOcclusionQueries = MAX_int32;
+	/** The maximum number of in-flight GPU queries the current RHI can handle without stalling and waiting for the GPU. Used to tune the occlusion culler to avoid stalls. */
+	int32 MaximumInFlightQueries = MAX_int32;
 
 	/** Some RHIs can only do visible or not occlusion queries. */
 	bool SupportsExactOcclusionQueries = true;
@@ -736,7 +736,7 @@ extern RHI_API uint32 GGPUFrameTime;
 #define GRHISupportsPipelineStateSortKey                       GRHIGlobals.SupportsPipelineStateSortKey
 #define GSupportsParallelRenderingTasksWithSeparateRHIThread   GRHIGlobals.SupportsParallelRenderingTasksWithSeparateRHIThread
 #define GRHIThreadNeedsKicking                                 GRHIGlobals.RHIThreadNeedsKicking
-#define GRHIMaximumReccommendedOustandingOcclusionQueries      GRHIGlobals.MaximumReccommendedOustandingOcclusionQueries
+#define GRHIMaximumInFlightQueries                             GRHIGlobals.MaximumInFlightQueries
 #define GRHISupportsExactOcclusionQueries                      GRHIGlobals.SupportsExactOcclusionQueries
 #define GSupportsVolumeTextureRendering                        GRHIGlobals.SupportsVolumeTextureRendering
 #define GSupportsSeparateRenderTargetBlendState                GRHIGlobals.SupportsSeparateRenderTargetBlendState
diff --git a/Engine/Source/Runtime/RHI/Public/RHIValidationContext.h b/Engine/Source/Runtime/RHI/Public/RHIValidationContext.h
--- a/Engine/Source/Runtime/RHI/Public/RHIValidationContext.h
+++ b/Engine/Source/Runtime/RHI/Public/RHIValidationContext.h
@@ -687,12 +687,6 @@ public:
 		RHIContext->RHICalibrateTimers(CalibrationQuery);
 	}
 
-	// Used for OpenGL to check and see if any occlusion queries can be read back on the RHI thread. If they aren't ready when we need them, then we end up stalling.
-	virtual void RHIPollOcclusionQueries() override final
-	{
-		RHIContext->RHIPollOcclusionQueries();
-	}
-
 	// Not all RHIs need this (Mobile specific)
 	virtual void RHIDiscardRenderTargets(bool bDepth, bool bStencil, uint32 ColorBitMask) override final
 	{
diff --git a/Engine/Source/Runtime/Renderer/Private/MobileShadingRenderer.cpp b/Engine/Source/Runtime/Renderer/Private/MobileShadingRenderer.cpp
--- a/Engine/Source/Runtime/Renderer/Private/MobileShadingRenderer.cpp
+++ b/Engine/Source/Runtime/Renderer/Private/MobileShadingRenderer.cpp
@@ -227,14 +227,6 @@ static bool PostProcessUsesSceneDepth(const FViewInfo& View)
 	return false;
 }
 
-static void PollOcclusionQueriesPass(FRDGBuilder& GraphBuilder)
-{
-	AddPass(GraphBuilder, RDG_EVENT_NAME("PollOcclusionQueries"), [](FRHICommandListImmediate& RHICmdList)
-	{
-		RHICmdList.PollOcclusionQueries();
-	});
-}
-
 struct FRenderViewContext
 {
 	FViewInfo* ViewInfo;
@@ -1186,8 +1178,6 @@ void FMobileSceneRenderer::Render(FRDGBuilder& GraphBuilder)
 		}
 		ExternalAccessQueue.Submit(GraphBuilder);
 
-		PollOcclusionQueriesPass(GraphBuilder);
-
 		// Custom depth
 		// bShouldRenderCustomDepth has been initialized in InitViews on mobile platform
 		if (bShouldRenderCustomDepth)
@@ -1472,11 +1462,6 @@ void FMobileSceneRenderer::Render(FRDGBuilder& GraphBuilder)
 
 	OnRenderFinish(GraphBuilder, ViewFamilyTexture);
 
-	if (bRendererOutputFinalSceneColor)
-	{
-		PollOcclusionQueriesPass(GraphBuilder);
-	}
-
 	QueueSceneTextureExtractions(GraphBuilder, SceneTextures);
 
 	if (Scene->InstanceCullingOcclusionQueryRenderer)
@@ -1667,7 +1652,7 @@ void FMobileSceneRenderer::RenderForwardSinglePass(FRDGBuilder& GraphBuilder, FM
 		// Opaque and masked
 		RenderMobileBasePass(RHICmdList, View, &PassParameters->InstanceCullingDrawParams);
 		RenderMobileDebugView(RHICmdList, View);
-		RHICmdList.PollOcclusionQueries();
+
 		PostRenderBasePass(RHICmdList, View);
 		// scene depth is read only and can be fetched
 		RHICmdList.NextSubpass();
@@ -1736,7 +1721,7 @@ void FMobileSceneRenderer::RenderForwardMultiPass(FRDGBuilder& GraphBuilder, FMo
 		// Opaque and masked
 		RenderMobileBasePass(RHICmdList, View, &PassParameters->InstanceCullingDrawParams);
 		RenderMobileDebugView(RHICmdList, View);
-		RHICmdList.PollOcclusionQueries();
+
 		PostRenderBasePass(RHICmdList, View);
 	});
 
@@ -2006,7 +1991,7 @@ void FMobileSceneRenderer::RenderDeferredSinglePass(FRDGBuilder& GraphBuilder, F
 				// Opaque and masked
 				RenderMobileBasePass(RHICmdList, View, &PassParameters->InstanceCullingDrawParams);
 				RenderMobileDebugView(RHICmdList, View);
-				RHICmdList.PollOcclusionQueries();
+
 				PostRenderBasePass(RHICmdList, View);
 				// SceneColor + GBuffer write, SceneDepth is read only
 				RHICmdList.NextSubpass();
@@ -2085,7 +2070,7 @@ void FMobileSceneRenderer::RenderDeferredMultiPass(FRDGBuilder& GraphBuilder, FS
 				// Opaque and masked
 				RenderMobileBasePass(RHICmdList, View, &PassParameters->InstanceCullingDrawParams);
 				RenderMobileDebugView(RHICmdList, View);
-				RHICmdList.PollOcclusionQueries();
+
 				PostRenderBasePass(RHICmdList, View);
 
 				if (bDoOcclusionQueires)
diff --git a/Engine/Source/Runtime/Renderer/Private/SceneVisibility.cpp b/Engine/Source/Runtime/Renderer/Private/SceneVisibility.cpp
--- a/Engine/Source/Runtime/Renderer/Private/SceneVisibility.cpp
+++ b/Engine/Source/Runtime/Renderer/Private/SceneVisibility.cpp
@@ -2663,7 +2663,7 @@ bool FGPUOcclusionPacket::OcclusionCullPrimitive(VisitorType& Visitor, FOcclusio
 					}
 					else
 					{
-						if (OcclusionState.NumBufferedFrames > 1 || GRHIMaximumReccommendedOustandingOcclusionQueries < MAX_int32)
+						if (OcclusionState.NumBufferedFrames > 1 || GRHIMaximumInFlightQueries < MAX_int32)
 						{
 							// If there's no occlusion query for the primitive, assume it is whatever it was last frame
 							bIsOccluded = PrimitiveOcclusionHistory->WasOccludedLastFrame;
@@ -2798,7 +2798,7 @@ bool FGPUOcclusionPacket::OcclusionCullPrimitive(VisitorType& Visitor, FOcclusio
 						const FVector BoundOrigin = OcclusionBounds.Origin + View.ViewMatrices.GetPreViewTranslation();
 						const FVector BoundExtent = OcclusionBounds.BoxExtent;
 
-						if (GRHIMaximumReccommendedOustandingOcclusionQueries < MAX_int32 && !bGroupedQuery)
+						if (GRHIMaximumInFlightQueries < MAX_int32 && !bGroupedQuery)
 						{
 							Visitor.AddThrottledOcclusionQuery(FThrottledOcclusionQuery(FPrimitiveOcclusionHistoryKey(PrimitiveId, SubQuery), BoundOrigin, BoundExtent, PrimitiveOcclusionHistory->LastQuerySubmitFrame()));
 						}
@@ -2918,7 +2918,7 @@ void FGPUOcclusionPacket::FProcessVisitor::SubmitThrottledOcclusionQueries()
 
 	const int32 NumRequestedThrottledQueries = SortedQueries.Num();
 	const int32 NumUsedQueries = Packet.View.GroupedOcclusionQueries.GetNumBatchOcclusionQueries();
-	const int32 ThrottleThreshold = GRHIMaximumReccommendedOustandingOcclusionQueries / FMath::Min(Packet.OcclusionState.NumBufferedFrames, 2); // extra RHIT frame does not count
+	const int32 ThrottleThreshold = GRHIMaximumInFlightQueries / FMath::Min(Packet.OcclusionState.NumBufferedFrames, 2); // extra RHIT frame does not count
 
 	int32 NumThrottledQueries = NumRequestedThrottledQueries;
 
