Support mip levels for UAV in OpenGL RHI

diff --git a/Engine/Source/Runtime/OpenGLDrv/Private/OpenGLCommands.cpp b/Engine/Source/Runtime/OpenGLDrv/Private/OpenGLCommands.cpp
--- a/Engine/Source/Runtime/OpenGLDrv/Private/OpenGLCommands.cpp
+++ b/Engine/Source/Runtime/OpenGLDrv/Private/OpenGLCommands.cpp
@@ -637,7 +637,7 @@ void FOpenGLDynamicRHI::SetupTexturesForDraw( FOpenGLContextState& ContextState
 	SetupTexturesForDraw(ContextState, PendingState.BoundShaderState, FOpenGL::GetMaxCombinedTextureImageUnits());
 }
 
-void FOpenGLDynamicRHI::InternalSetShaderImageUAV(GLint UAVIndex, GLenum Format, GLuint Resource, bool bLayered, GLint Layer, GLenum Access)
+void FOpenGLDynamicRHI::InternalSetShaderImageUAV(GLint UAVIndex, GLenum Format, GLuint Resource, bool bLayered, GLint Layer, GLenum Access, GLint Level)
 {
 	VERIFY_GL_SCOPE();
 	PendingState.UAVs[UAVIndex].Format = Format;
@@ -645,6 +645,7 @@ void FOpenGLDynamicRHI::InternalSetShaderImageUAV(GLint UAVIndex, GLenum Format,
 	PendingState.UAVs[UAVIndex].Access = Access;
 	PendingState.UAVs[UAVIndex].Layer = Layer;
 	PendingState.UAVs[UAVIndex].bLayered = bLayered;
+	PendingState.UAVs[UAVIndex].Level = Level;
 }
 
 void FOpenGLDynamicRHI::InternalSetShaderBufferUAV(GLint UAVIndex, GLuint Resource)
@@ -683,12 +684,12 @@ void FOpenGLDynamicRHI::SetupUAVsForProgram(FOpenGLContextState& ContextState, c
 	{
 		if (!NeededBits[UAVStageIndex])
 		{
-			CachedSetupUAVStage(ContextState, UAVStageIndex, 0, 0, false, 0, GL_READ_WRITE);
+			CachedSetupUAVStage(ContextState, UAVStageIndex, 0, 0, false, 0, GL_READ_WRITE, 0);
 		}
 		else
 		{
 			const FUAVStage& UAVStage = PendingState.UAVs[UAVStageIndex];
-			CachedSetupUAVStage(ContextState, UAVStageIndex, UAVStage.Format, UAVStage.Resource, UAVStage.bLayered, UAVStage.Layer, UAVStage.Access);
+			CachedSetupUAVStage(ContextState, UAVStageIndex, UAVStage.Format, UAVStage.Resource, UAVStage.bLayered, UAVStage.Layer, UAVStage.Access, UAVStage.Level);
 		}
 	}
 
@@ -699,7 +700,7 @@ void FOpenGLDynamicRHI::SetupUAVsForProgram(FOpenGLContextState& ContextState, c
 		const int32 NumUAVs = ContextState.UAVs.Num();
 		for (; UAVStageIndex < NumUAVs; ++UAVStageIndex)
 		{
-			CachedSetupUAVStage(ContextState, UAVStageIndex, 0, 0, false, 0, GL_READ_WRITE);
+			CachedSetupUAVStage(ContextState, UAVStageIndex, 0, 0, false, 0, GL_READ_WRITE, 0);
 		}
 	}
 }
@@ -734,7 +735,7 @@ static bool IsImageTextureFormatSupported(GLenum Format)
 	return true;
 }
 
-void FOpenGLDynamicRHI::CachedSetupUAVStage( FOpenGLContextState& ContextState, GLint UAVIndex, GLenum Format, GLuint Resource, bool bLayered, GLint Layer, GLenum Access)
+void FOpenGLDynamicRHI::CachedSetupUAVStage( FOpenGLContextState& ContextState, GLint UAVIndex, GLenum Format, GLuint Resource, bool bLayered, GLint Layer, GLenum Access, GLint Level)
 {
 	VERIFY_GL_SCOPE();
 	
@@ -744,7 +745,8 @@ void FOpenGLDynamicRHI::CachedSetupUAVStage( FOpenGLContextState& ContextState,
 		UAVStage.Resource == Resource &&
 		UAVStage.Access == Access &&
 		UAVStage.Layer == Layer &&
-		UAVStage.bLayered == bLayered)
+		UAVStage.bLayered == bLayered &&
+		UAVStage.Level == Level)
 	{
 		// Nothing's changed, no need to update
 		return;
@@ -771,6 +773,7 @@ void FOpenGLDynamicRHI::CachedSetupUAVStage( FOpenGLContextState& ContextState,
 			UAVStage.Access = GL_READ_WRITE;
 			UAVStage.Layer = 0;
 			UAVStage.bLayered = false;
+			UAVStage.Level = 0;
 		}
 	}
 	else
@@ -791,6 +794,7 @@ void FOpenGLDynamicRHI::CachedSetupUAVStage( FOpenGLContextState& ContextState,
 			UAVStage.Access = GL_READ_WRITE;
 			UAVStage.Layer = 0;
 			UAVStage.bLayered = false;
+			UAVStage.Level = 0;
 			ContextState.StorageBufferBound = Resource;
 		}
 		else // Image
@@ -804,13 +808,14 @@ void FOpenGLDynamicRHI::CachedSetupUAVStage( FOpenGLContextState& ContextState,
 			
 			check(IsImageTextureFormatSupported(Format));
 	
-			FOpenGL::BindImageTexture(UAVIndex, Resource, 0, bLayered ? GL_TRUE : GL_FALSE, Layer, Access, Format);
+			FOpenGL::BindImageTexture(UAVIndex, Resource, Level, bLayered ? GL_TRUE : GL_FALSE, Layer, Access, Format);
 	
 			UAVStage.Format = Format;
 			UAVStage.Resource = Resource;
 			UAVStage.Access = Access;
 			UAVStage.Layer = Layer;
 			UAVStage.bLayered = bLayered;
+			UAVStage.Level = Level;
 		}
 	}
 	
@@ -1029,7 +1034,7 @@ struct FOpenGLResourceBinder
 			if (UAV->Resource)
 			{
 				const GLenum Access = (ShaderFrequency == SF_Compute) ? GL_READ_WRITE : GL_WRITE_ONLY;
-				RHI.InternalSetShaderImageUAV(UAVIndex, UAV->Format, UAV->Resource, UAV->IsLayered(), UAV->GetLayer(), Access);
+				RHI.InternalSetShaderImageUAV(UAVIndex, UAV->Format, UAV->Resource, UAV->IsLayered(), UAV->GetLayer(), Access, UAV->Level);
 			}
 			else
 			{
diff --git a/Engine/Source/Runtime/OpenGLDrv/Private/OpenGLDrv.h b/Engine/Source/Runtime/OpenGLDrv/Private/OpenGLDrv.h
--- a/Engine/Source/Runtime/OpenGLDrv/Private/OpenGLDrv.h
+++ b/Engine/Source/Runtime/OpenGLDrv/Private/OpenGLDrv.h
@@ -497,7 +497,7 @@ public:
 	}
 
 	void CachedSetupTextureStageInner(FOpenGLContextState& ContextState, GLint TextureIndex, GLenum Target, GLuint Resource, GLint BaseMip, GLint NumMips);
-	void CachedSetupUAVStage(FOpenGLContextState& ContextState, GLint UAVIndex, GLenum Format, GLuint Resource, bool bLayered, GLint Layer, GLenum Access);
+	void CachedSetupUAVStage(FOpenGLContextState& ContextState, GLint UAVIndex, GLenum Format, GLuint Resource, bool bLayered, GLint Layer, GLenum Access, GLint Level);
 	void UpdateSRV(FOpenGLShaderResourceView* SRV);
 	FOpenGLContextState& GetContextStateForCurrentContext();
 
@@ -870,7 +870,7 @@ private:
 public:
 	/** Remember what RHI user wants set on a specific OpenGL texture stage, translating from Stage and TextureIndex for stage pair. */
 	void InternalSetShaderTexture(FOpenGLTexture* Texture, FOpenGLShaderResourceView* SRV, GLint TextureIndex, GLenum Target, GLuint Resource, int NumMips, int LimitMip);
-	void InternalSetShaderImageUAV(GLint UAVIndex, GLenum Format, GLuint Resource, bool bLayered, GLint Layer, GLenum Access);
+	void InternalSetShaderImageUAV(GLint UAVIndex, GLenum Format, GLuint Resource, bool bLayered, GLint Layer, GLenum Access, GLint Level);
 	void InternalSetShaderBufferUAV(GLint UAVIndex, GLuint Resource);
 	void InternalSetSamplerStates(GLint TextureIndex, FOpenGLSamplerState* SamplerState);
 	void InitializeGLTextureInternal(FOpenGLTexture* Texture, void const* BulkDataPtr, uint64 BulkDataSize);
diff --git a/Engine/Source/Runtime/OpenGLDrv/Private/OpenGLResources.h b/Engine/Source/Runtime/OpenGLDrv/Private/OpenGLResources.h
--- a/Engine/Source/Runtime/OpenGLDrv/Private/OpenGLResources.h
+++ b/Engine/Source/Runtime/OpenGLDrv/Private/OpenGLResources.h
@@ -1273,6 +1273,7 @@ public:
 	GLuint BufferResource = 0;
 	GLenum Format = 0;
 	uint8  UnrealFormat = 0;
+	uint8  Level = 0;
 
 	uint32 GetBufferSize() const
 	{
diff --git a/Engine/Source/Runtime/OpenGLDrv/Private/OpenGLState.h b/Engine/Source/Runtime/OpenGLDrv/Private/OpenGLState.h
--- a/Engine/Source/Runtime/OpenGLDrv/Private/OpenGLState.h
+++ b/Engine/Source/Runtime/OpenGLDrv/Private/OpenGLState.h
@@ -224,6 +224,7 @@ struct FUAVStage
 	GLenum Access;
 	GLint Layer;
 	bool bLayered;
+	GLint Level;
 	
 	FUAVStage()
 	:	Format(GL_NONE)
@@ -231,6 +232,7 @@ struct FUAVStage
 	,	Access(GL_READ_WRITE)
 	,	Layer(0)
 	,	bLayered(false)
+	,	Level(0)
 	{
 	}
 };
diff --git a/Engine/Source/Runtime/OpenGLDrv/Private/OpenGLUAV.cpp b/Engine/Source/Runtime/OpenGLDrv/Private/OpenGLUAV.cpp
--- a/Engine/Source/Runtime/OpenGLDrv/Private/OpenGLUAV.cpp
+++ b/Engine/Source/Runtime/OpenGLDrv/Private/OpenGLUAV.cpp
@@ -239,7 +239,6 @@ void FOpenGLUnorderedAccessView::UpdateView()
 		auto const Info = ViewDesc.Texture.UAV.GetViewInfo(Texture);
 
 		ensureMsgf(Info.bAllSlices, TEXT("UAV array slice selection is currently unimplemented in OpenGL RHI."));
-		ensureMsgf(Info.MipLevel == 0, TEXT("UAV mip level selection not implemented in OpenGL RHI."));
 
 		const FOpenGLTextureFormat& GLFormat = GOpenGLTextureFormats[Info.Format];
 
@@ -247,6 +246,7 @@ void FOpenGLUnorderedAccessView::UpdateView()
 		this->Resource = Texture->GetResource();
 		this->Format = GLFormat.InternalFormat[0];
 		this->UnrealFormat = Info.Format;
+		this->Level = Info.MipLevel;
 	}
 }
 
@@ -262,6 +262,7 @@ void FOpenGLUnorderedAccessView::Invalidate()
 	BufferResource = 0;
 	Format = 0;
 	UnrealFormat = 0;
+	Level = 0;
 }
 
 FOpenGLUnorderedAccessView::~FOpenGLUnorderedAccessView()
