Ensure opengl fences are polled during command list submission.
 - Fixes late fence results causing low frame rates, observed while using OcclusionFeedback culling.

diff --git a/Engine/Source/Runtime/OpenGLDrv/Private/OpenGLQuery.cpp b/Engine/Source/Runtime/OpenGLDrv/Private/OpenGLQuery.cpp
--- a/Engine/Source/Runtime/OpenGLDrv/Private/OpenGLQuery.cpp
+++ b/Engine/Source/Runtime/OpenGLDrv/Private/OpenGLQuery.cpp
@@ -347,6 +347,8 @@ bool FOpenGLRenderQuery_RHI::GetResult(bool bWait, uint64& OutResult)
 	return true;
 }
 
+void OpenGL_PollAllFences();
+
 void FOpenGLRenderQuery::PollQueryResults()
 {
 	if (ActiveQueries.First)
@@ -360,6 +362,11 @@ void FOpenGLRenderQuery::PollQueryResults()
 		}
 		while (ActiveQueries.First);
 	}
+	EOpenGLCurrentContext Context = FOpenGLDynamicRHI::GetCurrentContext();
+	if (Context == CONTEXT_Rendering)
+	{
+ 		OpenGL_PollAllFences();
+	}
 }
 
 void FOpenGLRenderQuery::Cleanup()
