Add Wait() function to FRHIGPUFence
 - Blocks the caller until the GPU fence has passed. This is a much more robust solution vs using Poll() in a busy loop.
 - Some renderer systems were busy looping on Poll(), which is dangerous since not all platform RHIs will signal fences if the RHI thread isn't making progress.
 - Allowing the platform RHI to implement a Wait() function also solves threading issues in some RHIs (e.g. OpenGL).

Removed FGenericRHIGPUFence
 - All platform RHIs implement the fence type now. We no longer count render thread frames to simulate GPU fences.

diff --git a/Engine/Source/Runtime/Apple/MetalRHI/Private/MetalDynamicRHI.cpp b/Engine/Source/Runtime/Apple/MetalRHI/Private/MetalDynamicRHI.cpp
--- a/Engine/Source/Runtime/Apple/MetalRHI/Private/MetalDynamicRHI.cpp
+++ b/Engine/Source/Runtime/Apple/MetalRHI/Private/MetalDynamicRHI.cpp
@@ -126,7 +126,7 @@ void* FMetalDynamicRHI::RHILockStagingBuffer(FRHIStagingBuffer* StagingBuffer, F
 		FRHICommandListImmediate& RHICmdList = FRHICommandListImmediate::Get();
 		RHICmdList.SubmitAndBlockUntilGPUIdle();
 		
-		ResourceCast(Fence)->WaitCPU();
+		ResourceCast(Fence)->Wait(RHICmdList, FRHIGPUMask::All());
 	}
 	
 	FMetalRHIStagingBuffer* Buffer = ResourceCast(StagingBuffer);
diff --git a/Engine/Source/Runtime/Apple/MetalRHI/Private/MetalDynamicRHI.h b/Engine/Source/Runtime/Apple/MetalRHI/Private/MetalDynamicRHI.h
--- a/Engine/Source/Runtime/Apple/MetalRHI/Private/MetalDynamicRHI.h
+++ b/Engine/Source/Runtime/Apple/MetalRHI/Private/MetalDynamicRHI.h
@@ -179,6 +179,7 @@ public:
 	virtual void RHIReleaseTransition(FRHITransition* Transition) final override;
 
 	virtual FGPUFenceRHIRef RHICreateGPUFence(const FName &Name) final override;
+	virtual void RHIWriteGPUFence_TopOfPipe(FRHICommandListBase& RHICmdList, FRHIGPUFence* FenceRHI) final override;
 
 	virtual FStagingBufferRHIRef RHICreateStagingBuffer() final override;
 	virtual void* RHILockStagingBuffer(FRHIStagingBuffer* StagingBuffer, FRHIGPUFence* Fence, uint32 Offset, uint32 SizeRHI) final override;
diff --git a/Engine/Source/Runtime/Apple/MetalRHI/Private/MetalRenderTarget.cpp b/Engine/Source/Runtime/Apple/MetalRHI/Private/MetalRenderTarget.cpp
--- a/Engine/Source/Runtime/Apple/MetalRHI/Private/MetalRenderTarget.cpp
+++ b/Engine/Source/Runtime/Apple/MetalRHI/Private/MetalRenderTarget.cpp
@@ -287,7 +287,7 @@ void FMetalDynamicRHI::RHIMapStagingSurface(FRHITexture* TextureRHI, FRHIGPUFenc
 	
 	if (FenceRHI && !FenceRHI->Poll())
 	{
-		ResourceCast(FenceRHI)->WaitCPU();
+		ResourceCast(FenceRHI)->Wait(RHICmdList, FRHIGPUMask::FromIndex(GPUIndex));
 	}
 	
 	FMetalSurface* Surface = GetMetalSurfaceFromRHITexture(TextureRHI);
diff --git a/Engine/Source/Runtime/Apple/MetalRHI/Private/MetalUAV.cpp b/Engine/Source/Runtime/Apple/MetalRHI/Private/MetalUAV.cpp
--- a/Engine/Source/Runtime/Apple/MetalRHI/Private/MetalUAV.cpp
+++ b/Engine/Source/Runtime/Apple/MetalRHI/Private/MetalUAV.cpp
@@ -787,12 +787,6 @@ void FMetalCommandBufferFence::Signal(const MTL::CommandBuffer* CommandBuffer)
     Condition->Trigger();
 }
 
-void FMetalGPUFence::WriteInternal(FMetalCommandBuffer* CommandBuffer)
-{
-	Fence = CommandBuffer->GetCompletionFence();
-	check(Fence);
-}
-
 void FMetalRHICommandContext::RHICopyToStagingBuffer(FRHIBuffer* SourceBufferRHI, FRHIStagingBuffer* DestinationStagingBufferRHI, uint32 Offset, uint32 NumBytes)
 {
     MTL_SCOPED_AUTORELEASE_POOL;
@@ -820,42 +814,75 @@ void FMetalRHICommandContext::RHICopyToStagingBuffer(FRHIBuffer* SourceBufferRHI
     CopyFromBufferToBuffer(SourceBuffer->GetCurrentBuffer(), Offset, ReadbackBuffer, 0, NumBytes);
 }
 
-void FMetalRHICommandContext::RHIWriteGPUFence(FRHIGPUFence* FenceRHI)
+void FMetalGPUFence::Clear()
 {
-    MTL_SCOPED_AUTORELEASE_POOL;
-    
-    check(FenceRHI);
-    FMetalGPUFence* Fence = ResourceCast(FenceRHI);
-    Fence->WriteInternal(CurrentEncoder.GetCommandBuffer());
+	Event = nullptr;
 }
 
-FGPUFenceRHIRef FMetalDynamicRHI::RHICreateGPUFence(const FName &Name)
+bool FMetalGPUFence::Poll() const
 {
-    MTL_SCOPED_AUTORELEASE_POOL;
-	return new FMetalGPUFence(Name);
-}
+	if (!Event || !Event->IsComplete())
+		return false;
 
-void FMetalGPUFence::Clear()
-{
-    Fence = nullptr;
+	return Fence->Wait(0);
 }
 
-bool FMetalGPUFence::Poll() const
+void FMetalGPUFence::Wait(FRHICommandListImmediate& RHICmdList, FRHIGPUMask GPUMask) const
 {
-	if (Fence)
+	if (!Event)
+		return;
+
+	// The event is signalled when the RHI thread has updated the Fence field.
+	bool bDispatched = false;
+	if (!Event->IsComplete())
 	{
-		return Fence->Wait(0);
+		RHICmdList.ImmediateFlush(EImmediateFlushType::DispatchToRHIThread);
+		bDispatched = true;
+
+		Event->Wait();
 	}
-	else
+
+	// Now that we have the correct Fence, wait on it.
+	if (!Fence->Wait(0))
 	{
-		return false;
+		if (!bDispatched)
+		{
+			RHICmdList.ImmediateFlush(EImmediateFlushType::DispatchToRHIThread);
+		}
+
+		Fence->Wait(MAX_uint32);
 	}
 }
 
-void FMetalGPUFence::WaitCPU() const
+void FMetalGPUFence::Write_TopOfPipe(FRHICommandListBase& RHICmdList)
 {
-	if (Fence)
+	Event = FGraphEvent::CreateGraphEvent();
+
+	RHICmdList.EnqueueLambda([this, Event = Event](FRHICommandListBase& ExecutingCmdList)
 	{
-		Fence->Wait(MAX_uint32);
-	}
+		MTL_SCOPED_AUTORELEASE_POOL;
+
+		FMetalRHICommandContext& Context = FMetalRHICommandContext::Get(ExecutingCmdList);
+
+		Fence = Context.GetCurrentCommandBuffer()->GetCompletionFence();
+		check(Fence);
+
+		Event->DispatchSubsequents();
+	});
+}
+
+void FMetalDynamicRHI::RHIWriteGPUFence_TopOfPipe(FRHICommandListBase& RHICmdList, FRHIGPUFence* FenceRHI)
+{
+	ResourceCast(FenceRHI)->Write_TopOfPipe(RHICmdList);
+}
+
+void FMetalRHICommandContext::RHIWriteGPUFence(FRHIGPUFence* FenceRHI)
+{
+	checkNoEntry(); // Should never be called
+}
+
+FGPUFenceRHIRef FMetalDynamicRHI::RHICreateGPUFence(const FName &Name)
+{
+    MTL_SCOPED_AUTORELEASE_POOL;
+	return new FMetalGPUFence(Name);
 }
diff --git a/Engine/Source/Runtime/Apple/MetalRHI/Public/MetalResources.h b/Engine/Source/Runtime/Apple/MetalRHI/Public/MetalResources.h
--- a/Engine/Source/Runtime/Apple/MetalRHI/Public/MetalResources.h
+++ b/Engine/Source/Runtime/Apple/MetalRHI/Public/MetalResources.h
@@ -579,15 +579,14 @@ public:
 	~FMetalGPUFence()
 	{}
 
-	virtual void Clear() override final;
+	void Clear() override;
+	bool Poll() const override;
+	void Wait(FRHICommandListImmediate& RHICmdList, FRHIGPUMask GPUMask) const override;
 
-	void WriteInternal(FMetalCommandBuffer* CmdBuffer);
-
-	virtual bool Poll() const override final;
-	
-	void WaitCPU() const;
+	void Write_TopOfPipe(FRHICommandListBase& RHICmdList);
 
 private:
+	FGraphEventRef Event;
     TSharedPtr<FMetalCommandBufferFence, ESPMode::ThreadSafe> Fence = nullptr;
 };
 
diff --git a/Engine/Source/Runtime/D3D12RHI/Private/D3D12DirectCommandListManager.cpp b/Engine/Source/Runtime/D3D12RHI/Private/D3D12DirectCommandListManager.cpp
--- a/Engine/Source/Runtime/D3D12RHI/Private/D3D12DirectCommandListManager.cpp
+++ b/Engine/Source/Runtime/D3D12RHI/Private/D3D12DirectCommandListManager.cpp
@@ -42,13 +42,13 @@ bool FD3D12GPUFence::Poll(FRHIGPUMask GPUMask) const
 	return bHasAnySyncPoint;
 }
 
-void FD3D12GPUFence::WaitCPU()
+void FD3D12GPUFence::Wait(FRHICommandListImmediate& RHICmdList, FRHIGPUMask GPUMask) const
 {
-	for (FD3D12SyncPointRef& SyncPoint : SyncPoints)
+	for (uint32 Index : GPUMask)
 	{
-		if (SyncPoint && !SyncPoint->IsComplete())
+		if (SyncPoints[Index] && !SyncPoints[Index]->IsComplete())
 		{
-			SyncPoint->Wait();
+			SyncPoints[Index]->Wait();
 		}
 	}
 }
diff --git a/Engine/Source/Runtime/D3D12RHI/Private/D3D12RenderTarget.cpp b/Engine/Source/Runtime/D3D12RHI/Private/D3D12RenderTarget.cpp
--- a/Engine/Source/Runtime/D3D12RHI/Private/D3D12RenderTarget.cpp
+++ b/Engine/Source/Runtime/D3D12RHI/Private/D3D12RenderTarget.cpp
@@ -830,7 +830,7 @@ void FD3D12DynamicRHI::RHIMapStagingSurface(FRHITexture* TextureRHI, FRHIGPUFenc
 
 	if (FenceRHI && !FenceRHI->Poll())
 	{
-		ResourceCast(FenceRHI)->WaitCPU();
+		ResourceCast(FenceRHI)->Wait(FRHICommandListImmediate::Get(), FRHIGPUMask::FromIndex(GPUIndex));
 	}
 
 	void* pData;
diff --git a/Engine/Source/Runtime/D3D12RHI/Private/D3D12Resources.h b/Engine/Source/Runtime/D3D12RHI/Private/D3D12Resources.h
--- a/Engine/Source/Runtime/D3D12RHI/Private/D3D12Resources.h
+++ b/Engine/Source/Runtime/D3D12RHI/Private/D3D12Resources.h
@@ -1102,8 +1102,7 @@ public:
 	virtual void Clear() override;
 	virtual bool Poll() const override;
 	virtual bool Poll(FRHIGPUMask GPUMask) const override;
-
-	void WaitCPU();
+	virtual void Wait(FRHICommandListImmediate& RHICmdList, FRHIGPUMask GPUMask) const override;
 
 	TArray<FD3D12SyncPointRef, TInlineAllocator<MAX_NUM_GPUS>> SyncPoints;
 };
diff --git a/Engine/Source/Runtime/NullDrv/Public/NullRHI.h b/Engine/Source/Runtime/NullDrv/Public/NullRHI.h
--- a/Engine/Source/Runtime/NullDrv/Public/NullRHI.h
+++ b/Engine/Source/Runtime/NullDrv/Public/NullRHI.h
@@ -279,6 +279,31 @@ public:
 	}
 
 
+	virtual FGPUFenceRHIRef RHICreateGPUFence(const FName& Name)
+	{
+		class FNullGPUFence final : public FRHIGPUFence
+		{
+		public:
+			FNullGPUFence(const FName& Name)
+				: FRHIGPUFence(Name)
+			{}
+
+			virtual void Clear() override
+			{
+			}
+
+			virtual bool Poll() const override
+			{
+				return true;
+			}
+
+			virtual void Wait(FRHICommandListImmediate& RHICmdList, FRHIGPUMask GPUMask) const override
+			{
+			}
+		};
+		return new FNullGPUFence(Name);
+	}
+
 	virtual void RHIMapStagingSurface(FRHITexture* Texture, FRHIGPUFence* Fence, void*& OutData, int32& OutWidth, int32& OutHeight, uint32 GPUIndex) final override
 	{
 
diff --git a/Engine/Source/Runtime/OpenGLDrv/Private/OpenGLCommands.cpp b/Engine/Source/Runtime/OpenGLDrv/Private/OpenGLCommands.cpp
--- a/Engine/Source/Runtime/OpenGLDrv/Private/OpenGLCommands.cpp
+++ b/Engine/Source/Runtime/OpenGLDrv/Private/OpenGLCommands.cpp
@@ -2930,8 +2930,9 @@ void FOpenGLDynamicRHI::RHIFinalizeContext(FRHIFinalizeContextArgs&& Args, TRHIP
 
 void FOpenGLDynamicRHI::RHISubmitCommandLists(FRHISubmitCommandListsArgs&& Args)
 {
-	// Periodically poll for completed queries
+	// Periodically poll for completed queries and fences
 	FOpenGLRenderQuery::PollQueryResults();
+	FOpenGLGPUFence::PollFences();
 }
 
 void FOpenGLDynamicRHI::RHICopyToStagingBuffer(FRHIBuffer* SourceBufferRHI, FRHIStagingBuffer* DestinationStagingBufferRHI, uint32 InOffset, uint32 InNumBytes)
@@ -2971,13 +2972,6 @@ void FOpenGLDynamicRHI::RHICopyToStagingBuffer(FRHIBuffer* SourceBufferRHI, FRHI
 	glBindBuffer(GL_COPY_WRITE_BUFFER, 0);
 }
 
-void FOpenGLDynamicRHI::RHIWriteGPUFence(FRHIGPUFence* FenceRHI)
-{
-	check(FenceRHI);
-	FOpenGLGPUFence* CopyFence = ResourceCast(FenceRHI);
-	CopyFence->WriteInternal();
-}
-
 void FOpenGLDynamicRHI::RHIPostExternalCommandsReset()
 {
 	auto &RCS = RenderingContextState;
diff --git a/Engine/Source/Runtime/OpenGLDrv/Private/OpenGLDrv.h b/Engine/Source/Runtime/OpenGLDrv/Private/OpenGLDrv.h
--- a/Engine/Source/Runtime/OpenGLDrv/Private/OpenGLDrv.h
+++ b/Engine/Source/Runtime/OpenGLDrv/Private/OpenGLDrv.h
@@ -450,6 +450,7 @@ public:
 
 	// Inline copy
 	virtual void RHICopyToStagingBuffer(FRHIBuffer* SourceBufferRHI, FRHIStagingBuffer* DestinationStagingBufferRHI, uint32 InOffset, uint32 InNumBytes) final override;
+	virtual void RHIWriteGPUFence_TopOfPipe(FRHICommandListBase& RHICmdList, FRHIGPUFence* FenceRHI) final override;
 	virtual void RHIWriteGPUFence(FRHIGPUFence* FenceRHI) final override;
 	virtual void* RHILockStagingBuffer(FRHIStagingBuffer* StagingBuffer, FRHIGPUFence* Fence, uint32 Offset, uint32 SizeRHI) final override;
 	virtual void RHIUnlockStagingBuffer(FRHIStagingBuffer* StagingBuffer) final override;
diff --git a/Engine/Source/Runtime/OpenGLDrv/Private/OpenGLQuery.cpp b/Engine/Source/Runtime/OpenGLDrv/Private/OpenGLQuery.cpp
--- a/Engine/Source/Runtime/OpenGLDrv/Private/OpenGLQuery.cpp
+++ b/Engine/Source/Runtime/OpenGLDrv/Private/OpenGLQuery.cpp
@@ -347,8 +347,6 @@ bool FOpenGLRenderQuery_RHI::GetResult(bool bWait, uint64& OutResult)
 	return true;
 }
 
-void OpenGL_PollAllFences();
-
 void FOpenGLRenderQuery::PollQueryResults()
 {
 	if (ActiveQueries.First)
@@ -362,11 +360,6 @@ void FOpenGLRenderQuery::PollQueryResults()
 		}
 		while (ActiveQueries.First);
 	}
-	EOpenGLCurrentContext Context = FOpenGLDynamicRHI::GetCurrentContext();
-	if (Context == CONTEXT_Rendering)
-	{
- 		OpenGL_PollAllFences();
-	}
 }
 
 void FOpenGLRenderQuery::Cleanup()
@@ -759,134 +752,122 @@ bool FOpenGLDisjointTimeStampQuery::GetResult(uint64* OutResult)
 
 #endif // (RHI_NEW_GPU_PROFILER == 0)
 
-// Fence implementation
-struct FOpenGLGPUFenceProxy
+TQueue<FOpenGLGPUFence::FGLSync, EQueueMode::SingleThreaded> FOpenGLGPUFence::ActiveSyncs;
+
+FOpenGLGPUFence::FOpenGLGPUFence(FName InName)
+	: FRHIGPUFence(InName)
 {
-	FOpenGLGPUFenceProxy()
-		: Fence(0)
-		, bIsSignaled(false)
-	{
-	}
+}
 
-	~FOpenGLGPUFenceProxy()
-	{
-		AllOpenGLGPUFences.RemoveSwap(this);
+void FOpenGLGPUFence::Clear()
+{
+	Event = nullptr;
+}
 
-		if (Fence != 0)
-		{
-			FOpenGL::DeleteSync(Fence);
-		}
-	}
+bool FOpenGLGPUFence::Poll() const
+{
+	return !Event || Event->IsComplete();
+}
 
-	void Write()
+void FOpenGLGPUFence::Wait(FRHICommandListImmediate& RHICmdList, FRHIGPUMask GPUMask) const
+{
+	if (Event && !Event->IsComplete())
 	{
-		if (Fence != 0)
-		{
-			FOpenGL::DeleteSync(Fence);
-		}
-		else
+		//
+		// The fence might get signalled by an earlier RHI command polling them, but we can't be sure that will happen.
+		// The GPU might finish work after the RHI thread has gone idle, and then we'll never see the fence complete.
+		//
+		// Enqueue a command here that will block and wait for the fence if it still hasn't signalled by the time
+		// the RHI thread is done with all prior commands.
+		//
+		RHICmdList.EnqueueLambda([Event = Event](FRHICommandListImmediate&)
 		{
-			AllOpenGLGPUFences.Add(this);
-		}
+			if (!Event->IsComplete())
+			{
+				PollFencesUntil(Event);
+			}
+		});
+		RHICmdList.ImmediateFlush(EImmediateFlushType::DispatchToRHIThread);
 
-		Fence = FOpenGL::FenceSync(GL_SYNC_GPU_COMMANDS_COMPLETE, 0);
-		bIsSignaled = false;
+		Event->Wait();
 	}
-	
-	void Poll()
+}
+
+void FOpenGLGPUFence::PollFencesUntil(FGraphEvent* Target)
+{
+	if (FOpenGLDynamicRHI::GetCurrentContext() != CONTEXT_Rendering && !Target)
 	{
-		if (Fence != 0 && !bIsSignaled)
-		{
-			FOpenGLBase::EFenceResult Result = (FOpenGL::ClientWaitSync(Fence, 0, 0));
-			bIsSignaled = (Result == FOpenGLBase::FR_AlreadySignaled || Result == FOpenGLBase::FR_ConditionSatisfied);
-		}
+		// Don't poll fences if this thread doesn't own the GL context.
+		return;
 	}
 
-	UGLsync Fence;
-	bool bIsSignaled;
+	VERIFY_GL_SCOPE();
 
-	//
-	static TArray<FOpenGLGPUFenceProxy*> AllOpenGLGPUFences;
-	
-	static void PollAllFences()
+	while (FGLSync* Sync = ActiveSyncs.Peek())
 	{
-		QUICK_SCOPE_CYCLE_COUNTER(STAT_FOpenGLGPUFence_PollAllFences);
+		bool const bTarget = Sync->Event == Target;
+
+		// Wait forever if this fence is the one we're looking for specifically, otherwise just poll.
+		uint64 const Timeout = bTarget ? 0xffffffff'ffffffff : 0;
 
-		for (int32 FenceIdx = 0; FenceIdx < AllOpenGLGPUFences.Num(); ++FenceIdx)
+		switch (FOpenGL::ClientWaitSync(Sync->GLSync, 0, Timeout))
 		{
-			AllOpenGLGPUFences[FenceIdx]->Poll();
+		case FOpenGL::FR_AlreadySignaled:
+		case FOpenGL::FR_ConditionSatisfied:
+			break; // Fence completed
+
+		case FOpenGL::FR_TimeoutExpired:
+			return; // Fence is not done
+
+		default:
+			checkNoEntry();
+			[[fallthrough]];
+		case FOpenGL::FR_WaitFailed:
+			// Some error state
+			UE_LOG(LogOpenGL, Fatal, TEXT("Waiting on FGLSync fence 0x%p failed."), Sync);
+			return;
 		}
-	}
-};
 
-TArray<FOpenGLGPUFenceProxy*> FOpenGLGPUFenceProxy::AllOpenGLGPUFences;
+		// The fence has completed. Signal the graph event and remove the node.
+		Sync->Event->DispatchSubsequents();
+		FOpenGL::DeleteSync(Sync->GLSync);
 
-void OpenGL_PollAllFences()
-{
-	VERIFY_GL_SCOPE();
-	FOpenGLGPUFenceProxy::PollAllFences();
-}
+		ActiveSyncs.Pop();
 
-FGPUFenceRHIRef FOpenGLDynamicRHI::RHICreateGPUFence(const FName &Name)
-{
-	return new FOpenGLGPUFence(Name);
-}
+		if (bTarget)
+		{
+			// We found the fence we wanted. Stop polling.
+			return;
+		}
+	}
 
-FOpenGLGPUFence::FOpenGLGPUFence(FName InName)
-	: FRHIGPUFence(InName)
-{
-	Proxy = new FOpenGLGPUFenceProxy();
+	checkf(!Target, TEXT("Attempt to poll for a specific fence, but it was not found in the queue."));
 }
 
-FOpenGLGPUFence::~FOpenGLGPUFence()
+void FOpenGLGPUFence::WriteGPUFence_TopOfPipe(FRHICommandListBase& RHICmdList)
 {
-	VERIFY_GL_SCOPE();
-	delete Proxy;
-}
+	Event = FGraphEvent::CreateGraphEvent();
 
-void FOpenGLGPUFence::Clear()
-{
-	FRHICommandListImmediate& RHICmdList = FRHICommandListImmediate::Get();
-	RHICmdList.EnqueueLambda([Proxy = Proxy](FRHICommandListImmediate&)
+	RHICmdList.EnqueueLambda([Event = Event](FRHICommandListBase&) mutable
 	{
 		VERIFY_GL_SCOPE();
-		delete Proxy;
-	});
 
-	Proxy = new FOpenGLGPUFenceProxy();
+		UGLsync Fence = FOpenGL::FenceSync(GL_SYNC_GPU_COMMANDS_COMPLETE, 0);
+		ActiveSyncs.Enqueue(FGLSync(MoveTemp(Event), Fence));
+	});
 }
 
-bool FOpenGLGPUFence::Poll() const
+void FOpenGLDynamicRHI::RHIWriteGPUFence_TopOfPipe(FRHICommandListBase& RHICmdList, FRHIGPUFence* FenceRHI)
 {
-	if (Proxy->bIsSignaled)
-	{
-		return true;
-	}
-
-	if (!(IsInRenderingThread() || IsInRHIThread()))
-	{
-		ENQUEUE_RENDER_COMMAND(FOpenGLGPUFence_Poll)(
-		[this](FRHICommandListImmediate& RHICmdList)
-		{
-			Poll();
-		});
-	}
-	else
-	{
-		FRHICommandListImmediate& RHICmdList = FRHICommandListImmediate::Get();
-		RHICmdList.EnqueueLambda([Proxy = Proxy](FRHICommandListImmediate&)
-		{
-			VERIFY_GL_SCOPE();
-			check(Proxy != nullptr);
-			Proxy->Poll();
-		});
-	}
+	ResourceCast(FenceRHI)->WriteGPUFence_TopOfPipe(RHICmdList);
+}
 
-	return Proxy->bIsSignaled;
+void FOpenGLDynamicRHI::RHIWriteGPUFence(FRHIGPUFence* FenceRHI)
+{
+	checkNoEntry(); // Should never be called
 }
 
-void FOpenGLGPUFence::WriteInternal()
+FGPUFenceRHIRef FOpenGLDynamicRHI::RHICreateGPUFence(const FName& Name)
 {
-	VERIFY_GL_SCOPE();
-	Proxy->Write();
+	return new FOpenGLGPUFence(Name);
 }
diff --git a/Engine/Source/Runtime/OpenGLDrv/Private/OpenGLResources.h b/Engine/Source/Runtime/OpenGLDrv/Private/OpenGLResources.h
--- a/Engine/Source/Runtime/OpenGLDrv/Private/OpenGLResources.h
+++ b/Engine/Source/Runtime/OpenGLDrv/Private/OpenGLResources.h
@@ -1381,16 +1381,40 @@ private:
 
 class FOpenGLGPUFence final : public FRHIGPUFence
 {
+private:
+	struct FGLSync
+	{
+		// The graph event to trigger when the sync is resolved.
+		FGraphEventRef Event;
+
+		// The GL API sync object to poll / wait on.
+		UGLsync GLSync;
+
+		FGLSync() = default;
+		FGLSync(FGraphEventRef Event, UGLsync GLSync)
+			: Event(MoveTemp(Event))
+			, GLSync(GLSync)
+		{}
+	};
+	
+	static TQueue<FGLSync, EQueueMode::SingleThreaded> ActiveSyncs;
+	static void PollFencesUntil(FGraphEvent* Target);
+
+	FGraphEventRef Event;
+
 public:
 	FOpenGLGPUFence(FName InName);
-	~FOpenGLGPUFence() override;
 
 	void Clear() override;
 	bool Poll() const override;
-	
-	void WriteInternal();
-private:
-	struct FOpenGLGPUFenceProxy* Proxy;
+	void Wait(FRHICommandListImmediate& RHICmdList, FRHIGPUMask GPUMask) const override;
+
+	void WriteGPUFence_TopOfPipe(FRHICommandListBase& RHICmdList);
+
+	static void PollFences()
+	{
+		PollFencesUntil(nullptr);
+	}
 };
 
 class FOpenGLStagingBuffer final : public FRHIStagingBuffer
diff --git a/Engine/Source/Runtime/RHI/Private/RHIGPUReadback.cpp b/Engine/Source/Runtime/RHI/Private/RHIGPUReadback.cpp
--- a/Engine/Source/Runtime/RHI/Private/RHIGPUReadback.cpp
+++ b/Engine/Source/Runtime/RHI/Private/RHIGPUReadback.cpp
@@ -8,36 +8,6 @@
 #include "RHIGPUReadback.h"
 #include "RHITransition.h"
 
-///////////////////////////////////////////////////////////////////////////////
-//////////////////////     FGenericRHIGPUFence    /////////////////////////////
-///////////////////////////////////////////////////////////////////////////////
-
-FGenericRHIGPUFence::FGenericRHIGPUFence(FName InName)
-	: FRHIGPUFence(InName)
-	, InsertedFrameNumber(MAX_uint32)
-{}
-
-void FGenericRHIGPUFence::Clear()
-{
-	InsertedFrameNumber = MAX_uint32;
-}
-
-void FGenericRHIGPUFence::WriteInternal()
-{
-	// GPU generally overlaps the game
-	InsertedFrameNumber = GFrameNumberRenderThread + 1;
-}
-
-bool FGenericRHIGPUFence::Poll() const
-{
-	const uint32 CurrentFrameNumber = GFrameNumberRenderThread;
-	if (CurrentFrameNumber > InsertedFrameNumber)
-	{
-		return true;
-	}
-	return false;
-}
-
 ///////////////////////////////////////////////////////////////////////////////
 ////////////////////     FGenericRHIStagingBuffer    //////////////////////////
 ///////////////////////////////////////////////////////////////////////////////
diff --git a/Engine/Source/Runtime/RHI/Public/DynamicRHI.h b/Engine/Source/Runtime/RHI/Public/DynamicRHI.h
--- a/Engine/Source/Runtime/RHI/Public/DynamicRHI.h
+++ b/Engine/Source/Runtime/RHI/Public/DynamicRHI.h
@@ -225,10 +225,7 @@ public:
 		return nullptr;
 	}
 
-	virtual FGPUFenceRHIRef RHICreateGPUFence(const FName &Name)
-	{
-		return new FGenericRHIGPUFence(Name);
-	}
+	virtual FGPUFenceRHIRef RHICreateGPUFence(const FName &Name) = 0;
 
 	//
 	// Called by the thread recording an RHI command list (via RHICmdList.WriteGPUFence()).
diff --git a/Engine/Source/Runtime/RHI/Public/RHIGPUReadback.h b/Engine/Source/Runtime/RHI/Public/RHIGPUReadback.h
--- a/Engine/Source/Runtime/RHI/Public/RHIGPUReadback.h
+++ b/Engine/Source/Runtime/RHI/Public/RHIGPUReadback.h
@@ -40,6 +40,14 @@ public:
 		return !Fence || Fence->Poll(GPUMask);
 	}
 
+	void Wait(FRHICommandListImmediate& RHICmdList, FRHIGPUMask GPUMask) const
+	{
+		if (Fence)
+		{
+			Fence->Wait(RHICmdList, GPUMask);
+		}
+	}
+
 	/**
 	 * Copy the current state of the resource to the readback data.
 	 * @param RHICmdList The command list to enqueue the copy request on.
diff --git a/Engine/Source/Runtime/RHI/Public/RHIResources.h b/Engine/Source/Runtime/RHI/Public/RHIResources.h
--- a/Engine/Source/Runtime/RHI/Public/RHIResources.h
+++ b/Engine/Source/Runtime/RHI/Public/RHIResources.h
@@ -2089,55 +2089,61 @@ public:
 	uint64 CPUMicroseconds[MAX_NUM_GPUS] = {};
 };
 
-/*
-* Generic GPU fence class.
-* Granularity differs depending on backing RHI - ie it may only represent command buffer granularity.
-* RHI specific fences derive from this to implement real GPU->CPU fencing.
-* The default implementation always returns false for Poll until the next frame from the frame the fence was inserted
-* because not all APIs have a GPU/CPU sync object, we need to fake it.
-*/
 class FRHIGPUFence : public FRHIResource
 {
 public:
-	FRHIGPUFence(FName InName) : FRHIResource(RRT_GPUFence), FenceName(InName) {}
-	virtual ~FRHIGPUFence() {}
+	FRHIGPUFence(FName InName)
+		: FRHIResource(RRT_GPUFence)
+		, FenceName(InName)
+	{}
 
 	virtual void Clear() = 0;
 
 	/**
 	 * Poll the fence to see if the GPU has signaled it.
 	 * @returns True if and only if the GPU fence has been inserted and the GPU has signaled the fence.
+	 * 
+	 * CAUTION: Do not call poll in a loop to block until completion. Some platform RHIs will not signal
+	 *          fences unless the RHI thread continues to make progress and submit GPU work.
+	 *
+	 *          To block for completion, call Wait(), which can only be done from the render thread.
+	 *
+	 *          Use Poll() only to check a fence. If it returns false, continue to do useful work and recheck
+	 *          later, ideally on the following frame.
 	 */
 	virtual bool Poll() const = 0;
 
 	/**
 	 * Poll on a subset of the GPUs that this fence supports.
+	 * 
+	 * CAUTION: Do not call poll in a loop to block until completion. Some platform RHIs will not signal
+	 *          fences unless the RHI thread continues to make progress and submit GPU work.
+	 *
+	 *          To block for completion, call Wait(), which can only be done from the render thread.
+	 *
+	 *          Use Poll() only to check a fence. If it returns false, continue to do useful work and recheck
+	 *          later, ideally on the following frame.
 	 */
-	virtual bool Poll(FRHIGPUMask GPUMask) const { return Poll(); }
+	virtual bool Poll(FRHIGPUMask GPUMask) const
+	{
+		checkf(GPUMask == FRHIGPUMask::GPU0(), TEXT("The current platform RHI does not implement MGPU support for RHI GPU fences."));
+		return Poll();
+	}
+
+	/**
+	 * Blocks the caller until the fence has passed on the GPU. Can only be called from the render thread.
+	 *
+	 * CAUTION: Use of this function may cause a hitch if the fence is not yet signalled and we have to wait for
+	 *          the GPU to catch up. Rendering code should be written to avoid immediately waiting on GPU fences.
+	 */
+	virtual void Wait(FRHICommandListImmediate& RHICmdList, FRHIGPUMask GPUMask) const = 0;
 
 	const FName& GetFName() const { return FenceName; }
 
 	FThreadSafeCounter NumPendingWriteCommands;
 
 protected:
-	FName FenceName;
-};
-
-// Generic implementation of FRHIGPUFence
-class FGenericRHIGPUFence : public FRHIGPUFence
-{
-public:
-	RHI_API FGenericRHIGPUFence(FName InName);
-
-	RHI_API virtual void Clear() final override;
-
-	/** @discussion RHI implementations must be thread-safe and must correctly handle being called before RHIInsertFence if an RHI thread is active. */
-	RHI_API virtual bool Poll() const final override;
-
-	RHI_API void WriteInternal();
-
-private:
-	uint32 InsertedFrameNumber;
+	const FName FenceName;
 };
 
 class FRHIRenderQuery : public FRHIResource
diff --git a/Engine/Source/Runtime/Renderer/Private/SceneOcclusion.cpp b/Engine/Source/Runtime/Renderer/Private/SceneOcclusion.cpp
--- a/Engine/Source/Runtime/Renderer/Private/SceneOcclusion.cpp
+++ b/Engine/Source/Runtime/Renderer/Private/SceneOcclusion.cpp
@@ -1694,7 +1694,7 @@ void FOcclusionFeedback::AdvanceFrame(uint32 InOcclusionFrameCounter)
 	OcclusionBuffer.OcclusionFrameCounter = InOcclusionFrameCounter;
 }
 
-void FOcclusionFeedback::ReadbackResults(FRHICommandList& RHICmdList)
+void FOcclusionFeedback::ReadbackResults(FRHICommandListImmediate& RHICmdList)
 {
 	constexpr uint32 NumBuffers = UE_ARRAY_COUNT(OcclusionBuffers);
 	TArray<uint32, TInlineAllocator<NumBuffers>> ReadyBuffers;
@@ -1735,34 +1735,12 @@ void FOcclusionFeedback::ReadbackResults(FRHICommandList& RHICmdList)
 		if (TailOcclusionBuffer.BatchedPrimitives.Num() > 0 
 			&& ResultsOcclusionFrameCounter < TailOcclusionBuffer.OcclusionFrameCounter)
 		{
-			// We'll do manual wait
-			double StartTime = FPlatformTime::Seconds();
+			// Wait for the results
 			FRenderThreadIdleScope IdleScope(ERenderThreadIdleTypes::WaitingForGPUQuery);
-			ENamedThreads::Type RenderThread_Local = ENamedThreads::GetRenderThread_Local();
+			TailOcclusionBuffer.ReadbackBuffer->Wait(RHICmdList, FRHIGPUMask::All());
+			bReady = true;
 
-			while (!bReady)
-			{
-				FPlatformProcess::SleepNoStats(0);
-				// pump RHIThread to make sure these queries have actually been submitted to the GPU.
-				if (IsInActualRenderingThread())
-				{
-					FTaskGraphInterface::Get().ProcessThreadUntilIdle(RenderThread_Local);
-				}
-				const double TimeoutValue = 0.5;
-				// look for gpu stuck/crashed
-				if ((FPlatformTime::Seconds() - StartTime) > TimeoutValue)
-				{
-					UE_LOG(LogRHI, Log, TEXT("Timed out while waiting for GPU to catch up on occlusion readback. (%.1f s)"), TimeoutValue);
-					break;
-				}
-
-				bReady = TailOcclusionBuffer.ReadbackBuffer->IsReady();
-			}
-
-			if (bReady)
-			{
-				ReadyBuffers.Add(TailBufferIndex);
-			}
+			ReadyBuffers.Add(TailBufferIndex);
 		}
 	}
 
diff --git a/Engine/Source/Runtime/Renderer/Private/ScenePrivate.h b/Engine/Source/Runtime/Renderer/Private/ScenePrivate.h
--- a/Engine/Source/Runtime/Renderer/Private/ScenePrivate.h
+++ b/Engine/Source/Runtime/Renderer/Private/ScenePrivate.h
@@ -686,7 +686,7 @@ public:
 	/** Renders the current batch and resets the batch state. */
 	void SubmitOcclusionDraws(FRHICommandList& RHICmdList, FViewInfo& View);
 
-	void ReadbackResults(FRHICommandList& RHICmdList);
+	void ReadbackResults(FRHICommandListImmediate& RHICmdList);
 	void AdvanceFrame(uint32 OcclusionFrameCounter);
 
 	inline FRDGBuffer* GetGPUFeedbackBuffer() const
diff --git a/Engine/Source/Runtime/VulkanRHI/Private/VulkanRenderTarget.cpp b/Engine/Source/Runtime/VulkanRHI/Private/VulkanRenderTarget.cpp
--- a/Engine/Source/Runtime/VulkanRHI/Private/VulkanRenderTarget.cpp
+++ b/Engine/Source/Runtime/VulkanRHI/Private/VulkanRenderTarget.cpp
@@ -266,16 +266,14 @@ void FVulkanDynamicRHI::RHIMapStagingSurface(FRHITexture* TextureRHI, FRHIGPUFen
 
 	if (FenceRHI && !FenceRHI->Poll())
 	{
-		FRHICommandListExecutor::GetImmediateCommandList().ImmediateFlush(EImmediateFlushType::FlushRHIThread);
+		FRHICommandListImmediate& RHICmdList = FRHICommandListImmediate::Get();
+
+		RHICmdList.ImmediateFlush(EImmediateFlushType::FlushRHIThread);
 		Device->SubmitCommandsAndFlushGPU();
 
 		// SubmitCommandsAndFlushGPU might update fence state if it was tied to a previously submitted command buffer.
 		// Its state will have been updated from Submitted to NeedReset, and would assert in WaitForCmdBuffer (which is not needed in such a case)
-		if (!FenceRHI->Poll())
-		{
-			FVulkanGPUFence* Fence = ResourceCast(FenceRHI);
-			Device->GetImmediateContext().GetCommandBufferManager()->WaitForCmdBuffer(Fence->GetCmdBuffer());
-		}
+		FenceRHI->Wait(RHICmdList, FRHIGPUMask::All());
 	}
 	else
 	{
diff --git a/Engine/Source/Runtime/VulkanRHI/Private/VulkanResources.h b/Engine/Source/Runtime/VulkanRHI/Private/VulkanResources.h
--- a/Engine/Source/Runtime/VulkanRHI/Private/VulkanResources.h
+++ b/Engine/Source/Runtime/VulkanRHI/Private/VulkanResources.h
@@ -1236,7 +1236,7 @@ private:
 	FVulkanDevice* Device;
 };
 
-class FVulkanGPUFence : public FRHIGPUFence
+class FVulkanGPUFence final : public FRHIGPUFence
 {
 public:
 	FVulkanGPUFence(FName InName)
@@ -1244,10 +1244,9 @@ public:
 	{
 	}
 
-	virtual void Clear() final override;
-	virtual bool Poll() const final override;
-
-	FVulkanCmdBuffer* GetCmdBuffer() const { return CmdBuffer; }
+	virtual void Clear() override;
+	virtual bool Poll() const override;
+	virtual void Wait(FRHICommandListImmediate& RHICmdList, FRHIGPUMask GPUMask) const override;
 
 protected:
 	FVulkanCmdBuffer*	CmdBuffer = nullptr;
diff --git a/Engine/Source/Runtime/VulkanRHI/Private/VulkanUtil.cpp b/Engine/Source/Runtime/VulkanRHI/Private/VulkanUtil.cpp
--- a/Engine/Source/Runtime/VulkanRHI/Private/VulkanUtil.cpp
+++ b/Engine/Source/Runtime/VulkanRHI/Private/VulkanUtil.cpp
@@ -736,11 +736,7 @@ void* FVulkanDynamicRHI::RHILockStagingBuffer(FRHIStagingBuffer* StagingBufferRH
 
 		// SubmitCommandsAndFlushGPU might update fence state if it was tied to a previously submitted command buffer.
 		// Its state will have been updated from Submitted to NeedReset, and would assert in WaitForCmdBuffer (which is not needed in such a case)
-		if (!FenceRHI->Poll())
-		{
-			FVulkanGPUFence* Fence = ResourceCast(FenceRHI);
-			Device->GetImmediateContext().GetCommandBufferManager()->WaitForCmdBuffer(Fence->GetCmdBuffer());
-		}
+		FenceRHI->Wait(FRHICommandListImmediate::Get(), FRHIGPUMask::All());
 	}
 
 	return StagingBuffer->Lock(Offset, NumBytes);
@@ -763,6 +759,14 @@ bool FVulkanGPUFence::Poll() const
 	return (CmdBuffer && (FenceSignaledCounter < CmdBuffer->GetFenceSignaledCounter()));
 }
 
+void FVulkanGPUFence::Wait(FRHICommandListImmediate& RHICmdList, FRHIGPUMask GPUMask) const
+{
+	if (!Poll())
+	{
+		FVulkanDynamicRHI::Get().GetDevice()->GetImmediateContext().GetCommandBufferManager()->WaitForCmdBuffer(CmdBuffer);
+	}
+}
+
 FGPUFenceRHIRef FVulkanDynamicRHI::RHICreateGPUFence(const FName& Name)
 {
 	return new FVulkanGPUFence(Name);
diff --git a/Engine/Source/Runtime/Windows/D3D11RHI/Private/D3D11Commands.cpp b/Engine/Source/Runtime/Windows/D3D11RHI/Private/D3D11Commands.cpp
--- a/Engine/Source/Runtime/Windows/D3D11RHI/Private/D3D11Commands.cpp
+++ b/Engine/Source/Runtime/Windows/D3D11RHI/Private/D3D11Commands.cpp
@@ -1434,8 +1434,9 @@ void FD3D11DynamicRHI::RHIFinalizeContext(FRHIFinalizeContextArgs&& Args, TRHIPi
 
 void FD3D11DynamicRHI::RHISubmitCommandLists(FRHISubmitCommandListsArgs&& Args)
 {
-	// Attempt to readback completed queries
+	// Attempt to readback completed queries and fences
 	PollQueryResults();
+	FD3D11GPUFence::PollFences();
 
 #if RHI_NEW_GPU_PROFILER && WITH_RHI_BREADCRUMBS
 	for (IRHIPlatformCommandList* CmdList : Args.CommandLists)
@@ -1666,15 +1667,6 @@ void FD3D11DynamicRHI::RHICopyToStagingBuffer(FRHIBuffer* SourceBufferRHI, FRHIS
 	}
 }
 
-void FD3D11DynamicRHI::RHIWriteGPUFence(FRHIGPUFence* FenceRHI)
-{
-	// @todo-staging Implement real fences for D3D11
-	// D3D11 only has the generic fence for now.
-	FGenericRHIGPUFence* Fence = ResourceCast(FenceRHI);
-	check(Fence);
-	Fence->WriteInternal();
-}
-
 void* FD3D11DynamicRHI::RHILockStagingBuffer(FRHIStagingBuffer* StagingBufferRHI, FRHIGPUFence* Fence, uint32 Offset, uint32 SizeRHI)
 {
 	check(StagingBufferRHI);
@@ -1687,3 +1679,126 @@ void FD3D11DynamicRHI::RHIUnlockStagingBuffer(FRHIStagingBuffer* StagingBufferRH
 	FD3D11StagingBuffer* StagingBuffer = ResourceCast(StagingBufferRHI);
 	StagingBuffer->Unlock();
 }
+
+
+TQueue<FD3D11GPUFence::FD3D11Sync, EQueueMode::SingleThreaded> FD3D11GPUFence::ActiveSyncs;
+
+FD3D11GPUFence::FD3D11GPUFence(FName InName)
+	: FRHIGPUFence(InName)
+{
+}
+
+void FD3D11GPUFence::Clear()
+{
+	Event = nullptr;
+}
+
+bool FD3D11GPUFence::Poll() const
+{
+	return !Event || Event->IsComplete();
+}
+
+void FD3D11GPUFence::Wait(FRHICommandListImmediate& RHICmdList, FRHIGPUMask GPUMask) const
+{
+	if (Event && !Event->IsComplete())
+	{
+		//
+		// The fence might get signalled by an earlier RHI command polling them, but we can't be sure that will happen.
+		// The GPU might finish work after the RHI thread has gone idle, and then we'll never see the fence complete.
+		//
+		// Enqueue a command here that will block and wait for the fence if it still hasn't signalled by the time
+		// the RHI thread is done with all prior commands.
+		//
+		RHICmdList.EnqueueLambda([Event = Event](FRHICommandListImmediate&)
+		{
+			if (!Event->IsComplete())
+			{
+				PollFencesUntil(Event);
+			}
+		});
+		RHICmdList.ImmediateFlush(EImmediateFlushType::DispatchToRHIThread);
+
+		Event->Wait();
+	}
+}
+
+void FD3D11GPUFence::PollFencesUntil(FGraphEvent* Target)
+{
+	FD3D11DynamicRHI& RHI = FD3D11DynamicRHI::Get();
+
+	while (FD3D11Sync* Sync = ActiveSyncs.Peek())
+	{
+		bool const bTarget = Sync->Event == Target;
+
+		// Wait forever if this fence is the one we're looking for specifically, otherwise just poll.
+		UINT const Flags = bTarget ? 0 : D3D11_ASYNC_GETDATA_DONOTFLUSH;
+
+	Retry:
+		BOOL Value = 0;
+		HRESULT Result = RHI.GetDeviceContext()->GetData(Sync->Query, &Value, sizeof(Value), Flags);
+
+		if (Result == S_FALSE || (Result == S_OK && !Value))
+		{
+			// Fence is not done
+			if (bTarget)
+			{
+				// We're waiting for a specific fence. Spin until it passes.
+				goto Retry;
+			}
+			else
+			{
+				return;
+			}
+		}
+
+		VERIFYD3D11RESULT(Result);
+
+		// The fence has completed. Signal the graph event and remove the node.
+		Sync->Event->DispatchSubsequents();
+		ActiveSyncs.Pop();
+
+		if (bTarget)
+		{
+			// We found the fence we wanted. Stop polling.
+			return;
+		}
+	}
+
+	checkf(!Target, TEXT("Attempt to poll for a specific fence, but it was not found in the queue."));
+}
+
+void FD3D11GPUFence::WriteGPUFence_TopOfPipe(FRHICommandListBase& RHICmdList)
+{
+	Event = FGraphEvent::CreateGraphEvent();
+
+	RHICmdList.EnqueueLambda([Event = Event](FRHICommandListBase&) mutable
+	{
+		FD3D11DynamicRHI& RHI = FD3D11DynamicRHI::Get();
+
+		// Insert an event query on the device context
+		D3D11_QUERY_DESC Desc {};
+		Desc.Query = D3D11_QUERY_EVENT;
+
+		TRefCountPtr<ID3D11Query> Query;
+		VERIFYD3D11RESULT(RHI.GetDevice()->CreateQuery(&Desc, Query.GetInitReference()));
+		RHI.GetDeviceContext()->End(Query);
+
+		// Store the query in the list of queries to poll
+		ActiveSyncs.Enqueue(FD3D11Sync(MoveTemp(Event), MoveTemp(Query)));
+	});
+}
+
+void FD3D11DynamicRHI::RHIWriteGPUFence_TopOfPipe(FRHICommandListBase& RHICmdList, FRHIGPUFence* FenceRHI)
+{
+	ResourceCast(FenceRHI)->WriteGPUFence_TopOfPipe(RHICmdList);
+}
+
+void FD3D11DynamicRHI::RHIWriteGPUFence(FRHIGPUFence* FenceRHI)
+{
+	checkNoEntry(); // Should never be called
+}
+
+FGPUFenceRHIRef FD3D11DynamicRHI::RHICreateGPUFence(const FName& Name)
+{
+	return new FD3D11GPUFence(Name);
+}
diff --git a/Engine/Source/Runtime/Windows/D3D11RHI/Private/D3D11RHIPrivate.h b/Engine/Source/Runtime/Windows/D3D11RHI/Private/D3D11RHIPrivate.h
--- a/Engine/Source/Runtime/Windows/D3D11RHI/Private/D3D11RHIPrivate.h
+++ b/Engine/Source/Runtime/Windows/D3D11RHI/Private/D3D11RHIPrivate.h
@@ -433,6 +433,8 @@ public:
 	virtual FComputeShaderRHIRef RHICreateComputeShader(TArrayView<const uint8> Code, const FSHAHash& Hash) final override;
 	virtual FStagingBufferRHIRef RHICreateStagingBuffer() final override;
 	virtual void RHICopyToStagingBuffer(FRHIBuffer* SourceBuffer, FRHIStagingBuffer* DestinationStagingBuffer, uint32 Offset, uint32 NumBytes) final override;
+	virtual FGPUFenceRHIRef RHICreateGPUFence(const FName& Name) final override;
+	virtual void RHIWriteGPUFence_TopOfPipe(FRHICommandListBase& RHICmdList, FRHIGPUFence* FenceRHI) final override;
 	virtual void RHIWriteGPUFence(FRHIGPUFence* Fence) final override;
 	virtual void* RHILockStagingBuffer(FRHIStagingBuffer* StagingBuffer, FRHIGPUFence* Fence, uint32 Offset, uint32 SizeRHI) final override;
     virtual void RHIUnlockStagingBuffer(FRHIStagingBuffer* StagingBuffer) final override;
diff --git a/Engine/Source/Runtime/Windows/D3D11RHI/Public/D3D11Resources.h b/Engine/Source/Runtime/Windows/D3D11RHI/Public/D3D11Resources.h
--- a/Engine/Source/Runtime/Windows/D3D11RHI/Public/D3D11Resources.h
+++ b/Engine/Source/Runtime/Windows/D3D11RHI/Public/D3D11Resources.h
@@ -465,6 +465,44 @@ private:
 	uint32 ShadowBufferSize;
 };
 
+class FD3D11GPUFence final : public FRHIGPUFence
+{
+private:
+	struct FD3D11Sync
+	{
+		// The graph event to trigger when the sync is resolved.
+		FGraphEventRef Event;
+
+		// The D3D11 API query object to poll / wait on.
+		TRefCountPtr<ID3D11Query> Query;
+
+		FD3D11Sync() = default;
+		FD3D11Sync(FGraphEventRef Event, TRefCountPtr<ID3D11Query> Query)
+			: Event(MoveTemp(Event))
+			, Query(MoveTemp(Query))
+		{}
+	};
+
+	static TQueue<FD3D11Sync, EQueueMode::SingleThreaded> ActiveSyncs;
+	static void PollFencesUntil(FGraphEvent* Target);
+
+	FGraphEventRef Event;
+
+public:
+	FD3D11GPUFence(FName InName);
+
+	void Clear() override;
+	bool Poll() const override;
+	void Wait(FRHICommandListImmediate& RHICmdList, FRHIGPUMask GPUMask) const override;
+
+	void WriteGPUFence_TopOfPipe(FRHICommandListBase& RHICmdList);
+
+	static void PollFences()
+	{
+		PollFencesUntil(nullptr);
+	}
+};
+
 namespace D3D11BufferStats
 {
 	void UpdateUniformBufferStats(ID3D11Buffer* Buffer, int64 BufferSize, bool bAllocating);
@@ -564,7 +602,7 @@ struct TD3D11ResourceTraits<FRHIStagingBuffer>
 template<>
 struct TD3D11ResourceTraits<FRHIGPUFence>
 {
-	typedef FGenericRHIGPUFence TConcreteType;
+	typedef FD3D11GPUFence TConcreteType;
 };
 template<>
 struct TD3D11ResourceTraits<FRHIShaderResourceView>
