Fix crashes on shutdown when using OpenGL RHI
 - Some resource destructors were calling FRHICommandListImmediate::Get(), which will break since RHI resources are destroyed by the RHI thread.
 - Acquire the GL context on the game thread when the rendering thread is stopped. The render thread is destroyed before RHIExit is called, which can lead to the game thread attempting to destroy GL resources while not owning the context.
 - Add some extra cleanup to the legacy GPU profiler. Some resources were still alive.

diff --git a/Engine/Source/Runtime/OpenGLDrv/Private/OpenGLBuffer.cpp b/Engine/Source/Runtime/OpenGLDrv/Private/OpenGLBuffer.cpp
--- a/Engine/Source/Runtime/OpenGLDrv/Private/OpenGLBuffer.cpp
+++ b/Engine/Source/Runtime/OpenGLDrv/Private/OpenGLBuffer.cpp
@@ -276,14 +276,8 @@ void FOpenGLStagingBuffer::Initialize()
 
 FOpenGLStagingBuffer::~FOpenGLStagingBuffer()
 {
-	FRHICommandListImmediate& RHICmdList = FRHICommandListImmediate::Get();
-	RHICmdList.EnqueueLambda([&](FRHICommandListImmediate&)
-	{
-		VERIFY_GL_SCOPE();
-		glDeleteBuffers(1, &ShadowBuffer);
-	});
-	RHITHREAD_GLTRACE_BLOCKING;
-	RHICmdList.ImmediateFlush(EImmediateFlushType::FlushRHIThread);
+	VERIFY_GL_SCOPE();
+	glDeleteBuffers(1, &ShadowBuffer);
 }
 
 // If we do not support the BufferStorage extension or if PersistentMapping is set to false, this will send the command to the RHI and flush it
diff --git a/Engine/Source/Runtime/OpenGLDrv/Private/OpenGLDrv.cpp b/Engine/Source/Runtime/OpenGLDrv/Private/OpenGLDrv.cpp
--- a/Engine/Source/Runtime/OpenGLDrv/Private/OpenGLDrv.cpp
+++ b/Engine/Source/Runtime/OpenGLDrv/Private/OpenGLDrv.cpp
@@ -453,7 +453,13 @@ void FOpenGLGPUProfiler::EndFrame()
 void FOpenGLGPUProfiler::Cleanup()
 {
 	FrameTiming.ReleaseResources();
+	GPUHitchEventNodeFrames.Empty();
 	NestedFrameCount = 0;
+
+	for (FOpenGLDisjointTimeStampQuery& DisjointQuery : DisjointGPUFrameTimeQuery)
+	{
+		DisjointQuery.Cleanup();
+	}
 }
 
 /** Start this frame of per tracking */
diff --git a/Engine/Source/Runtime/OpenGLDrv/Private/OpenGLDrv.h b/Engine/Source/Runtime/OpenGLDrv/Private/OpenGLDrv.h
--- a/Engine/Source/Runtime/OpenGLDrv/Private/OpenGLDrv.h
+++ b/Engine/Source/Runtime/OpenGLDrv/Private/OpenGLDrv.h
@@ -102,7 +102,9 @@ private:
 class FOpenGLDisjointTimeStampQuery
 {
 public:
-	FOpenGLDisjointTimeStampQuery() = default;
+	FOpenGLDisjointTimeStampQuery()
+		: DisjointQuery(new FOpenGLRenderQuery { FOpenGLRenderQuery::EType::Disjoint })
+	{}
 
 	void StartTracking();
 	void EndTracking();
@@ -124,9 +126,15 @@ public:
 #endif
 	}
 
+	void Cleanup()
+	{
+		delete DisjointQuery;
+		DisjointQuery = nullptr;
+	}
+
 private:
 	bool	bIsResultValid = false;
-	FOpenGLRenderQuery DisjointQuery { FOpenGLRenderQuery::EType::Disjoint };
+	FOpenGLRenderQuery* DisjointQuery;
 };
 
 /** A single perf event node, which tracks information about a appBeginDrawEvent/appEndDrawEvent range. */
diff --git a/Engine/Source/Runtime/OpenGLDrv/Private/OpenGLQuery.cpp b/Engine/Source/Runtime/OpenGLDrv/Private/OpenGLQuery.cpp
--- a/Engine/Source/Runtime/OpenGLDrv/Private/OpenGLQuery.cpp
+++ b/Engine/Source/Runtime/OpenGLDrv/Private/OpenGLQuery.cpp
@@ -132,6 +132,13 @@ void FOpenGLRenderQuery::AcquireGlQuery()
 
 void FOpenGLRenderQuery::ReleaseGlQuery()
 {
+	if (Resource == 0)
+	{
+		// Already released
+		check(!IsLinked());
+		return;
+	}
+
 	CheckContext();
 
 	if (bSharedContext)
@@ -713,7 +720,7 @@ void FOpenGLDisjointTimeStampQuery::StartTracking()
 	VERIFY_GL_SCOPE();
 	if (IsSupported())
 	{
-		DisjointQuery.Begin();
+		DisjointQuery->Begin();
 	}
 }
 
@@ -723,7 +730,7 @@ void FOpenGLDisjointTimeStampQuery::EndTracking()
 
 	if (IsSupported())
 	{
-		DisjointQuery.End();
+		DisjointQuery->End();
 	}
 }
 
@@ -739,9 +746,9 @@ bool FOpenGLDisjointTimeStampQuery::GetResult(uint64* OutResult)
 
 	if (IsSupported())
 	{
-		DisjointQuery.CacheResult(true);
+		DisjointQuery->CacheResult(true);
 
-		uint64 Result = DisjointQuery.GetResult();
+		uint64 Result = DisjointQuery->GetResult();
 		bIsResultValid = (Result & FOpenGLRenderQuery::InvalidDisjointMask) == 0;
 
 		*OutResult = Result & (~FOpenGLRenderQuery::InvalidDisjointMask);
diff --git a/Engine/Source/Runtime/OpenGLDrv/Private/OpenGLViewport.cpp b/Engine/Source/Runtime/OpenGLDrv/Private/OpenGLViewport.cpp
--- a/Engine/Source/Runtime/OpenGLDrv/Private/OpenGLViewport.cpp
+++ b/Engine/Source/Runtime/OpenGLDrv/Private/OpenGLViewport.cpp
@@ -277,7 +277,7 @@ FOpenGLViewport::FOpenGLViewport(FOpenGLDynamicRHI* InOpenGLRHI,void* InWindowHa
 
 FOpenGLViewport::~FOpenGLViewport()
 {
-	check(IsInRenderingThread() || IsInRHIThread());
+	VERIFY_GL_SCOPE();
 
 	if (bIsFullscreen)
 	{
@@ -288,14 +288,8 @@ FOpenGLViewport::~FOpenGLViewport()
 	BackBuffer.SafeRelease();
 	check(!IsValidRef(BackBuffer));
 
-	FRHICommandListImmediate& RHICmdList = FRHICommandListImmediate::Get();
-	RHICmdList.EnqueueLambda([&](FRHICommandListImmediate&)
-	{
-		FrameSyncEvent = nullptr;
-		PlatformDestroyOpenGLContext(OpenGLRHI->PlatformDevice, OpenGLContext);
-	});
-	RHITHREAD_GLTRACE_BLOCKING;
-	RHICmdList.ImmediateFlush(EImmediateFlushType::FlushRHIThread);
+	FrameSyncEvent = nullptr;
+	PlatformDestroyOpenGLContext(OpenGLRHI->PlatformDevice, OpenGLContext);
 
 	OpenGLContext = NULL;
 	OpenGLRHI->Viewports.Remove(this);
diff --git a/Engine/Source/Runtime/RenderCore/Private/RenderingThread.cpp b/Engine/Source/Runtime/RenderCore/Private/RenderingThread.cpp
--- a/Engine/Source/Runtime/RenderCore/Private/RenderingThread.cpp
+++ b/Engine/Source/Runtime/RenderCore/Private/RenderingThread.cpp
@@ -620,6 +620,8 @@ static void StartRenderingThread()
 	// will run on GT after RT is enabled
 	FlushRenderingCommands();
 
+	GDynamicRHI->RHIReleaseThreadOwnership();
+
 	switch (GRHISupportsRHIThread ? FRHIThread::TargetMode : ERHIThreadMode::None)
 	{
 	case ERHIThreadMode::DedicatedThread:
@@ -782,6 +784,8 @@ static void StopRenderingThread()
 	GRenderingThread = nullptr;
 	PRAGMA_ENABLE_DEPRECATION_WARNINGS
 
+	GDynamicRHI->RHIAcquireThreadOwnership();
+
 	GRHICommandList.LatchBypass();
 
 	delete GRenderingThreadRunnable;
