Remove OpenGL shared contexts
 - Only one GL context exists, and is owned by the thread at the bottom of the pipe (normally the RHI thread, but will be the game or render threads during startup / shutdown and if the RHI thread is disabled).
 - Removed GetContextStateForCurrentContext().
 - VERIFY_GL_SCOPE / CHECK_EXPECTED_GL_THREAD now explicitly checks that the calling thread owns the rendering context.

Fix issue with OpenGL queries when running without a render thread.
 - Both the TOPCounter and BOPCounter fields are incremented when the query is issued, which skips the call to CacheResult in GetResult().
 - We should compare against LastCachedBOPCounter instead to know if query results have been read back.

Reduce code duplication in RHIDraw/RHIDispatch functions
 - Added SetupDraw() / SetupDispatch() functions to handle the common operations

diff --git a/Engine/Source/Runtime/ApplicationCore/Public/Android/AndroidApplication.h b/Engine/Source/Runtime/ApplicationCore/Public/Android/AndroidApplication.h
--- a/Engine/Source/Runtime/ApplicationCore/Public/Android/AndroidApplication.h
+++ b/Engine/Source/Runtime/ApplicationCore/Public/Android/AndroidApplication.h
@@ -20,19 +20,6 @@ namespace FAndroidAppEntry
 	void OnPauseEvent();
 }
 
-struct FPlatformOpenGLContext;
-namespace FAndroidEGL
-{
-	// Back door into more intimate Android OpenGL variables (a.k.a. a hack)
-	FPlatformOpenGLContext*	GetRenderingContext();
-	FPlatformOpenGLContext*	CreateContext();
-	void					MakeCurrent(FPlatformOpenGLContext*);
-	void					ReleaseContext(FPlatformOpenGLContext*);
-	void					SwapBuffers(FPlatformOpenGLContext*);
-	void					SetFlipsEnabled(bool Enabled);
-	void					BindDisplayToContext(FPlatformOpenGLContext*);
-}
-
 //disable warnings from overriding the deprecated forcefeedback.  
 //calls to the deprecated function will still generate warnings.
 PRAGMA_DISABLE_DEPRECATION_WARNINGS
diff --git a/Engine/Source/Runtime/OpenGLDrv/Private/Android/AndroidEGL.cpp b/Engine/Source/Runtime/OpenGLDrv/Private/Android/AndroidEGL.cpp
--- a/Engine/Source/Runtime/OpenGLDrv/Private/Android/AndroidEGL.cpp
+++ b/Engine/Source/Runtime/OpenGLDrv/Private/Android/AndroidEGL.cpp
@@ -41,71 +41,77 @@ static TAutoConsoleVariable<int32> CVarAndroidGLESFlipYMethod(
 		 " 2: Force flip Y by BlitFrameBuffer."),
 	ECVF_RenderThreadSafe);
 
-
-const  int EGLMinRedBits		= 5;
-const  int EGLMinGreenBits		= 6;
-const  int EGLMinBlueBits		= 5;
-const  int EGLMinAlphaBits		= 0;
-const  int EGLMinDepthBits		= 16;
-const  int EGLMinStencilBits	= 8; // This is required for UMG clipping
-const  int EGLMinSampleBuffers	= 0;
-const  int EGLMinSampleSamples	= 0;
-
-
+const int EGLMinRedBits       = 5;
+const int EGLMinGreenBits     = 6;
+const int EGLMinBlueBits      = 5;
+const int EGLMinAlphaBits     = 0;
+const int EGLMinDepthBits     = 16;
+const int EGLMinStencilBits	  = 8; // This is required for UMG clipping
+const int EGLMinSampleBuffers = 0;
+const int EGLMinSampleSamples = 0;
 
 struct EGLConfigParms
 {
 	/** Whether this is a valid configuration or not */
-	int validConfig;
+	int validConfig = 0;
 	/** The number of bits requested for the red component */
-	int redSize ;
+	int redSize = 8;
 	/** The number of bits requested for the green component */
-	int greenSize  ;
+	int greenSize = 8;
 	/** The number of bits requested for the blue component */
-	int blueSize;
+	int blueSize = 8;
 	/** The number of bits requested for the alpha component */
-	int alphaSize ;
+	int alphaSize = 0;
 	/** The number of bits requested for the depth component */
-	int depthSize;
+	int depthSize = 24;
 	/** The number of bits requested for the stencil component */
-	int stencilSize ;
+	int stencilSize = 0;
 	/** The number of multisample buffers requested */
-	int sampleBuffers;
+	int sampleBuffers = 0;
 	/** The number of samples requested */
-	int sampleSamples;
+	int sampleSamples = 0;
+
+	EGLConfigParms()
+	{
+		// If not default, set the preference
+		int DepthBufferPreference = (int)FAndroidWindow::GetDepthBufferPreference();
+		if (DepthBufferPreference > 0)
+		{
+			depthSize = DepthBufferPreference;
+		}
 
-	EGLConfigParms();
-	EGLConfigParms(const EGLConfigParms& Parms);
+		if (FAndroidMisc::GetMobilePropagateAlphaSetting() > 0)
+		{
+			alphaSize = 8;
+		}
+	}
 };
 
 
 struct AndroidESPImpl
 {
-	FPlatformOpenGLContext	SharedContext;
 	FPlatformOpenGLContext	RenderingContext;
-	FPlatformOpenGLContext	SingleThreadedContext;
-
-	EGLDisplay eglDisplay;
-	EGLint eglNumConfigs;
-	EGLint eglFormat;
-	EGLConfig eglConfigParam;
-	EGLSurface eglSurface;
-	EGLSurface auxSurface;
-	EGLint eglWidth;
-	EGLint eglHeight;
-	EGLint NativeVisualID;
-	float eglRatio;
+
+	EGLDisplay     eglDisplay         = EGL_NO_DISPLAY;
+	EGLint         eglNumConfigs      = 0;
+	EGLint         eglFormat          = -1;
+	EGLConfig      eglConfigParam     = nullptr;
+	EGLSurface     eglSurface         = EGL_NO_SURFACE;
+	EGLSurface     auxSurface         = EGL_NO_SURFACE;
+	EGLint         eglWidth           = 8; // required for Gear VR apps with internal win surf mgmt
+	EGLint         eglHeight          = 8; // required for Gear VR apps with internal win surf mgmt
+	EGLint         NativeVisualID     = 0;
 	EGLConfigParms Parms;
-	int DepthSize;
-	ANativeWindow* Window;
-	bool Initalized ;
-	EOpenGLCurrentContext CurrentContextType;
-	GLuint ResolveFrameBuffer;
-	GLuint DummyFrameBuffer;
-	FPlatformRect CachedWindowRect;
-	bool bIsDebug = false;
-
-	AndroidESPImpl();
+	float          eglRatio           = 0;
+	int            DepthSize          = 0;
+	ANativeWindow* Window             = nullptr;
+	GLuint         ResolveFrameBuffer = 0;
+	GLuint         DummyFrameBuffer   = 0;
+	FPlatformRect  CachedWindowRect {};
+	bool           Initalized         = false;
+	bool           bIsDebug           = false;
+
+	AndroidESPImpl() = default;
 };
 
 const EGLint Attributes[] = {
@@ -123,46 +129,7 @@ const EGLint Attributes[] = {
 	EGL_NONE
 };
 
-
-EGLConfigParms::EGLConfigParms(const EGLConfigParms& Parms)
-{
-	validConfig = Parms.validConfig;
-	redSize = Parms.redSize;
-	greenSize = Parms.greenSize;
-	blueSize = Parms.blueSize;
-	alphaSize = Parms.alphaSize;
-	depthSize = Parms.depthSize;
-	stencilSize = Parms.stencilSize;
-	sampleBuffers = Parms.sampleBuffers;
-	sampleSamples = Parms.sampleSamples;
-}
-
-EGLConfigParms::EGLConfigParms()  : 
-validConfig (0)
-	,redSize(8)
-	,greenSize(8)
-	,blueSize(8)
-	,alphaSize(0)
-	,depthSize(24)
-	,stencilSize(0)
-	,sampleBuffers(0)
-	,sampleSamples(0)
-{
-	// If not default, set the preference
-	int DepthBufferPreference = (int)FAndroidWindow::GetDepthBufferPreference();
-	if (DepthBufferPreference > 0)
-		depthSize = DepthBufferPreference;
-	if (FAndroidMisc::GetMobilePropagateAlphaSetting() > 0)
-	{
-		alphaSize = 8;
-	}
-}
-
 AndroidEGL::AndroidEGL()
-:	bSupportsKHRCreateContext(false)
-,	bSupportsKHRSurfacelessContext(false)
-,	bSupportsKHRNoErrorContext(false)
-,	ContextAttributes(nullptr)
 {
 	PImplData = new AndroidESPImpl();
 
@@ -181,7 +148,6 @@ void AndroidEGL::ResetDisplay()
 	{
 		FPlatformMisc::LowLevelOutputDebugStringf( TEXT("AndroidEGL::ResetDisplay()" ));
 		eglMakeCurrent(PImplData->eglDisplay, EGL_NO_SURFACE, EGL_NO_SURFACE, EGL_NO_CONTEXT);
-		PImplData->CurrentContextType = CONTEXT_Invalid;
 	}
 }
 
@@ -196,20 +162,6 @@ void AndroidEGL::DestroyRenderSurface()
 	}
 
 	PImplData->RenderingContext.eglSurface = EGL_NO_SURFACE;
-	PImplData->SingleThreadedContext.eglSurface = EGL_NO_SURFACE;
-}
-
-void AndroidEGL::DestroySharedSurface()
-{
-	VERIFY_EGL_SCOPE();
-	FPlatformMisc::LowLevelOutputDebugStringf(TEXT("AndroidEGL::DestroySharedSurface()"));
-	if( PImplData->auxSurface != EGL_NO_SURFACE )
-	{
-		eglDestroySurface( PImplData->eglDisplay, PImplData->auxSurface);
-		PImplData->auxSurface = EGL_NO_SURFACE;
-	}
-
-	PImplData->SharedContext.eglSurface = EGL_NO_SURFACE;
 }
 
 void AndroidEGL::TerminateEGL()
@@ -230,13 +182,14 @@ EGLBoolean AndroidEGL::SetCurrentContext(EGLContext InContext, EGLSurface InSurf
 	EGLContext CurrentContext = GetCurrentContext();
 
 	// activate the context
-	if( CurrentContext != InContext)
+	if (CurrentContext != InContext)
 	{
-		if (CurrentContext !=EGL_NO_CONTEXT )
+		if (CurrentContext !=EGL_NO_CONTEXT)
 		{
 			glFlush();
 		}
-		if(InContext == EGL_NO_CONTEXT && InSurface == EGL_NO_SURFACE)
+
+		if (InContext == EGL_NO_CONTEXT && InSurface == EGL_NO_SURFACE)
 		{
 			ResetDisplay();
 		}
@@ -255,7 +208,9 @@ EGLBoolean AndroidEGL::SetCurrentContext(EGLContext InContext, EGLSurface InSurf
 					EGL_TEXTURE_FORMAT, EGL_NO_TEXTURE,
 					EGL_NONE
 				};
+
 				PImplData->auxSurface = eglCreatePbufferSurface(PImplData->eglDisplay, PImplData->eglConfigParam, PBufferAttribs);
+
 				if (PImplData->auxSurface == EGL_NO_SURFACE)
 				{
 					checkf(PImplData->auxSurface != EGL_NO_SURFACE, TEXT("eglCreatePbufferSurface error : 0x%x"), eglGetError());
@@ -305,7 +260,7 @@ void AndroidEGL::CreateEGLRenderSurface(ANativeWindow* InWindow, bool bCreateWnd
 
 		STANDALONE_DEBUG_LOGf(LogAndroid, TEXT("AndroidEGL::CreateEGLRenderSurface() %p" ), PImplData->eglSurface);
 
-		if(PImplData->eglSurface == EGL_NO_SURFACE )
+		if (PImplData->eglSurface == EGL_NO_SURFACE)
 		{
 			checkf(PImplData->eglSurface != EGL_NO_SURFACE, TEXT("eglCreateWindowSurface error : 0x%x"), eglGetError());
 			ResetInternal();
@@ -367,42 +322,6 @@ void AndroidEGL::CreateEGLRenderSurface(ANativeWindow* InWindow, bool bCreateWnd
 	}
 }
 
-void AndroidEGL::CreateEGLSharedSurface()
-{
-	VERIFY_EGL_SCOPE();
-
-	// due to possible early initialization, don't redo this
-	if (PImplData->auxSurface != EGL_NO_SURFACE)
-	{
-		FPlatformMisc::LowLevelOutputDebugStringf(TEXT("AndroidEGL::CreateEGLSharedSurface() Already initialized: %p"), PImplData->auxSurface);
-		return;
-	}
-
-	EGLint pbufferAttribs[] =
-	{
-		EGL_WIDTH, 1,
-		EGL_HEIGHT, 1,
-		EGL_TEXTURE_TARGET, EGL_NO_TEXTURE,
-		EGL_TEXTURE_FORMAT, EGL_NO_TEXTURE,
-		EGL_NONE
-	};
-
-	checkf(PImplData->eglWidth != 0, TEXT("eglWidth is ZERO; could be a problem!"));
-	checkf(PImplData->eglHeight != 0, TEXT("eglHeight is ZERO; could be a problem!"));
-	pbufferAttribs[1] = PImplData->eglWidth;
-	pbufferAttribs[3] = PImplData->eglHeight;
-
-	FPlatformMisc::LowLevelOutputDebugStringf( TEXT("AndroidEGL::CreateEGLSharedSurface(), auxSurface = eglCreatePbufferSurface(), %dx%d" ),
-		pbufferAttribs[1], pbufferAttribs[3]);
-	PImplData->auxSurface = eglCreatePbufferSurface(PImplData->eglDisplay, PImplData->eglConfigParam, pbufferAttribs);
-	if(PImplData->auxSurface== EGL_NO_SURFACE )
-	{
-		checkf(PImplData->auxSurface != EGL_NO_SURFACE, TEXT("eglCreatePbufferSurface error : 0x%x"), eglGetError());
-		ResetInternal();
-	}
-}
-
-
 void AndroidEGL::InitEGL(APIVariant API)
 {
 	VERIFY_EGL_SCOPE();
@@ -540,29 +459,6 @@ void AndroidEGL::InitEGL(APIVariant API)
 #endif
 }
 
-
-
-AndroidESPImpl::AndroidESPImpl():
-eglDisplay(EGL_NO_DISPLAY)
-	,eglNumConfigs(0)
-	,eglFormat(-1)
-	,eglConfigParam(NULL)
-	,eglSurface(EGL_NO_SURFACE)
-	,auxSurface(EGL_NO_SURFACE)
-	,eglWidth(8)  // required for Gear VR apps with internal win surf mgmt
-	,eglHeight(8) // required for Gear VR apps with internal win surf mgmt
-	,eglRatio(0)
-	,DepthSize(0)
-	,Window(NULL)
-	,Initalized(false)
-	,ResolveFrameBuffer(0)
-	,DummyFrameBuffer(0)
-	,NativeVisualID(0)
-	,CurrentContextType(CONTEXT_Invalid)
-	,CachedWindowRect(FPlatformRect(0,0,0,0))
-{
-}
-
 // call out to JNI to see if the application was packaged for Oculus Mobile
 extern bool AndroidThunkCpp_IsOculusMobileApplication();
 
@@ -571,57 +467,51 @@ void AndroidEGL::SetRenderContextWindowSurface()
 	VERIFY_GL_SCOPE();
 
 	FPlatformMisc::LowLevelOutputDebugStringf(TEXT("AndroidEGL::SetRenderContextWindowSurface  recreating context! tid: %d"),  FPlatformTLS::GetCurrentThreadId());
-	UnBind();
-	SetCurrentContext(EGL_NO_CONTEXT, EGL_NO_SURFACE);
-	bool bCreateSurface = !AndroidThunkCpp_IsOculusMobileApplication();
 
-	//SetRenderContextWindowSurface is called only when the window lock is successful.
-	PImplData->Window = (ANativeWindow*)FAndroidWindow::GetHardwareWindow_EventThread();
-	check(PImplData->Window);
-	InitSurface(false, bCreateSurface);
-	SetCurrentRenderingContext();
-	FPlatformMisc::LowLevelOutputDebugStringf(TEXT("AndroidEGL::SetRenderContextWindowSurface  DONE! tid: %d"),  FPlatformTLS::GetCurrentThreadId());
-}
+	UnBindRender();
 
-void AndroidEGL::ResizeRenderContextSurface()
-{
-	VERIFY_GL_SCOPE();
-
-	// Resize render originates from the gamethread, we cant use Window_Event here.
-	if (PImplData->Window && 
-		(PImplData->eglWidth != (PImplData->CachedWindowRect.Right - PImplData->CachedWindowRect.Left)
-		|| PImplData->eglHeight != (PImplData->CachedWindowRect.Bottom - PImplData->CachedWindowRect.Top))
-		)
+	SetCurrentContext(EGL_NO_CONTEXT, EGL_NO_SURFACE);
 	{
-		UE_LOG(LogAndroid, Log, TEXT("AndroidEGL::ResizeRenderContextSurface, PImplData->Window=%p, PImplData->eglWidth=%d, PImplData->eglHeight=%d!, CachedWidth=%d, CachedHeight=%d, tid: %d"), 
-			PImplData->Window, PImplData->eglWidth, PImplData->eglHeight, (PImplData->CachedWindowRect.Right - PImplData->CachedWindowRect.Left), (PImplData->CachedWindowRect.Bottom - PImplData->CachedWindowRect.Top), FPlatformTLS::GetCurrentThreadId());
-		UnBindRender();
-		SetCurrentContext(EGL_NO_CONTEXT, EGL_NO_SURFACE);
 		bool bCreateSurface = !AndroidThunkCpp_IsOculusMobileApplication();
+
+		// SetRenderContextWindowSurface is called only when the window lock is successful.
+		PImplData->Window = (ANativeWindow*)FAndroidWindow::GetHardwareWindow_EventThread();
+		check(PImplData->Window);
 		InitRenderSurface(false, bCreateSurface);
-		SetCurrentRenderingContext();
 	}
+	SetCurrentRenderingContext();
+
+	FPlatformMisc::LowLevelOutputDebugStringf(TEXT("AndroidEGL::SetRenderContextWindowSurface  DONE! tid: %d"),  FPlatformTLS::GetCurrentThreadId());
 }
 
-void AndroidEGL::ResizeSharedContextSurface()
+void AndroidEGL::ResizeRenderContextSurface()
 {
 	VERIFY_GL_SCOPE();
 
-	check(IsInGameThread());
-
-	// Resize shared is in gamethread, we cant use Window_Event here.
-	if (PImplData->Window &&
-		(PImplData->eglWidth != (PImplData->CachedWindowRect.Right - PImplData->CachedWindowRect.Left)
-			|| PImplData->eglHeight != (PImplData->CachedWindowRect.Bottom - PImplData->CachedWindowRect.Top))
-		)
+	// Resize render originates from the gamethread, we cant use Window_Event here.
+	if (PImplData->Window)
 	{
-		UE_LOG(LogAndroid, Log, TEXT("AndroidEGL::ResizeSharedContextSurface, PImplData->Window=%p, PImplData->eglWidth=%d, PImplData->eglHeight=%d!, CachedWidth=%d, CachedHeight=%d, tid: %d"),
-			PImplData->Window, PImplData->eglWidth, PImplData->eglHeight, (PImplData->CachedWindowRect.Right - PImplData->CachedWindowRect.Left), (PImplData->CachedWindowRect.Bottom - PImplData->CachedWindowRect.Top), FPlatformTLS::GetCurrentThreadId());
-		UnBindShared();
-		SetCurrentContext(EGL_NO_CONTEXT, EGL_NO_SURFACE);
-		bool bCreateSurface = !AndroidThunkCpp_IsOculusMobileApplication();
-		InitSharedSurface(false);
-		SetCurrentSharedContext();
+		if (PImplData->eglWidth != (PImplData->CachedWindowRect.Right - PImplData->CachedWindowRect.Left) || 
+			PImplData->eglHeight != (PImplData->CachedWindowRect.Bottom - PImplData->CachedWindowRect.Top))
+		{
+			UE_LOG(LogAndroid, Log, TEXT("AndroidEGL::ResizeRenderContextSurface, PImplData->Window=%p, PImplData->eglWidth=%d, PImplData->eglHeight=%d!, CachedWidth=%d, CachedHeight=%d, tid: %d"), 
+				PImplData->Window,
+				PImplData->eglWidth,
+				PImplData->eglHeight,
+				(PImplData->CachedWindowRect.Right - PImplData->CachedWindowRect.Left),
+				(PImplData->CachedWindowRect.Bottom - PImplData->CachedWindowRect.Top),
+				FPlatformTLS::GetCurrentThreadId()
+			);
+
+			UnBindRender();
+
+			SetCurrentContext(EGL_NO_CONTEXT, EGL_NO_SURFACE);
+			{
+				bool bCreateSurface = !AndroidThunkCpp_IsOculusMobileApplication();
+				InitRenderSurface(false, bCreateSurface);
+			}
+			SetCurrentRenderingContext();
+		}
 	}
 }
 
@@ -637,16 +527,16 @@ AndroidEGL* AndroidEGL::GetInstance()
 
 void AndroidEGL::DestroyBackBuffer()
 {
-	VERIFY_GL_SCOPE();
-
-	if(PImplData->ResolveFrameBuffer)
+	if (PImplData->ResolveFrameBuffer)
 	{
+		VERIFY_GL_SCOPE();
 		glDeleteFramebuffers(1, &PImplData->ResolveFrameBuffer);
 		PImplData->ResolveFrameBuffer = 0 ;
 	}
 
-	if(PImplData->DummyFrameBuffer)
+	if (PImplData->DummyFrameBuffer)
 	{
+		VERIFY_GL_SCOPE();
 		glDeleteFramebuffers(1, &PImplData->DummyFrameBuffer);
 		PImplData->DummyFrameBuffer = 0;
 	}
@@ -655,18 +545,10 @@ void AndroidEGL::DestroyBackBuffer()
 void AndroidEGL::InitBackBuffer()
 {
 	PImplData->RenderingContext.ViewportFramebuffer = GetResolveFrameBuffer();
-	PImplData->SharedContext.ViewportFramebuffer = GetResolveFrameBuffer();
-	PImplData->SingleThreadedContext.ViewportFramebuffer = GetResolveFrameBuffer();
 }
 
 extern void AndroidThunkCpp_SetDesiredViewSize(int32 Width, int32 Height);
 
-void AndroidEGL::InitSurface(bool bUseSmallSurface, bool bCreateWndSurface)
-{
-	InitRenderSurface(bUseSmallSurface, bCreateWndSurface);
-	InitSharedSurface(bUseSmallSurface);
-}
-
 void AndroidEGL::InitRenderSurface(bool bUseSmallSurface, bool bCreateWndSurface)
 {
 	FPlatformMisc::LowLevelOutputDebugStringf(TEXT("AndroidEGL::InitRenderSurface %d, %d"), int(bUseSmallSurface), int(bCreateWndSurface));
@@ -708,48 +590,6 @@ void AndroidEGL::InitRenderSurface(bool bUseSmallSurface, bool bCreateWndSurface
 	CreateEGLRenderSurface(PImplData->Window, bCreateWndSurface);
 
 	PImplData->RenderingContext.eglSurface = PImplData->eglSurface;
-	PImplData->SingleThreadedContext.eglSurface = PImplData->eglSurface;
-}
-
-void AndroidEGL::InitSharedSurface(bool bUseSmallSurface)
-{
-	FPlatformMisc::LowLevelOutputDebugStringf(TEXT("AndroidEGL::InitSharedSurface %d"), int(bUseSmallSurface));
-
-	check(PImplData->Window);
-
-	int32 Width = 8, Height = 8;
-	if (!bUseSmallSurface)
-	{
-		FPlatformRect WindowSize = FAndroidWindow::GetScreenRect();
-
-		if (PImplData->CachedWindowRect.Right > 0 && PImplData->CachedWindowRect.Bottom > 0)
-		{
-			// If we resumed from a lost window reuse the window size, the game thread will update the window dimensions.
-			FPlatformMisc::LowLevelOutputDebugStringf(TEXT("AndroidEGL::InitSharedSurface, Using CachedWindowRect, left: %d, top: %d, right: %d, bottom: %d "), PImplData->CachedWindowRect.Left, PImplData->CachedWindowRect.Top, PImplData->CachedWindowRect.Right, PImplData->CachedWindowRect.Bottom);
-			WindowSize = PImplData->CachedWindowRect;
-		}
-#if USE_ANDROID_STANDALONE
-		if (WindowSize.Left != 0 || WindowSize.Top != 0)
-		{
-			STANDALONE_DEBUG_LOGf(LogAndroid, TEXT("AndroidEGL::InitSharedSurface, WARNING!!! WindowSize is offset, left: %d, top: %d, right: %d, bottom: %d "), WindowSize.Left, WindowSize.Top, WindowSize.Right, WindowSize.Bottom);
-		}
-		Width = WindowSize.Right - WindowSize.Left;
-		Height = WindowSize.Bottom - WindowSize.Top;
-#else
-
-		Width = WindowSize.Right;
-		Height = WindowSize.Bottom;
-#endif
-
-		FPlatformMisc::LowLevelOutputDebugStringf(TEXT("AndroidEGL::InitSharedSurface, Using width: %d, height %d "), Width, Height);
-		AndroidThunkCpp_SetDesiredViewSize(Width, Height);
-	}
-
-	FPlatformMisc::LowLevelOutputDebugStringf(TEXT("AndroidEGL::InitSharedSurface, width: %d, height %d "), Width, Height);
-	CreateEGLSharedSurface();
-
-	PImplData->SharedContext.eglSurface = PImplData->auxSurface;
-
 }
 
 void AndroidEGL::ReInit()
@@ -762,11 +602,11 @@ void AndroidEGL::ReInit()
 
 	SetCurrentContext(EGL_NO_CONTEXT, EGL_NO_SURFACE);
 	bool bCreateSurface = !AndroidThunkCpp_IsOculusMobileApplication();
-	InitSurface(false, bCreateSurface);
-	SetCurrentSharedContext();
+	InitRenderSurface(false, bCreateSurface);
+	
+	SetCurrentContext(PImplData->RenderingContext.eglContext, PImplData->RenderingContext.eglSurface);
 }
 
-
 // EGL_CONTEXT_OPENGL_ROBUST_ACCESS_EXT is enabled if configrules asks for it or the command line specifies it.
 // If -OpenGLRobustContext=[0/1] is specified on the command line it takes precedence.
 void AndroidEGL::Init(APIVariant API, uint32 MajorVersion, uint32 MinorVersion)
@@ -887,10 +727,10 @@ void AndroidEGL::DestroyContext(EGLContext InContext)
 	}
 }
 
-EGLContext AndroidEGL::CreateContext(EGLContext InSharedContext)
+EGLContext AndroidEGL::CreateContext(EGLContext InParentContext)
 {
 	VERIFY_EGL_SCOPE();
-	return eglCreateContext(PImplData->eglDisplay, PImplData->eglConfigParam,  InSharedContext , ContextAttributes);
+	return eglCreateContext(PImplData->eglDisplay, PImplData->eglConfigParam, InParentContext, ContextAttributes);
 }
 
 int32 AndroidEGL::GetError()
@@ -908,13 +748,6 @@ GLuint AndroidEGL::GetResolveFrameBuffer()
 	return PImplData->ResolveFrameBuffer;
 }
 
-bool AndroidEGL::IsCurrentContextValid()
-{
-	VERIFY_EGL_SCOPE();
-	EGLContext eglContext =  eglGetCurrentContext();
-	return ( eglContext != EGL_NO_CONTEXT);
-}
-
 EGLContext AndroidEGL::GetCurrentContext()
 {
 	VERIFY_EGL_SCOPE();
@@ -950,28 +783,8 @@ ANativeWindow* AndroidEGL::GetNativeWindow() const
 
 bool AndroidEGL::InitContexts()
 {
-	PImplData->SharedContext.eglContext = CreateContext();
-	
-	PImplData->RenderingContext.eglContext = CreateContext(PImplData->SharedContext.eglContext);
-	
-	PImplData->SingleThreadedContext.eglContext = CreateContext();
-
-	return PImplData->SharedContext.eglContext != EGL_NO_CONTEXT;
-}
-
-void AndroidEGL::SetCurrentSharedContext()
-{
-	check(IsInGameThread());
-	PImplData->CurrentContextType = CONTEXT_Shared;
-
-	if(GUseThreadedRendering)
-	{
-		SetCurrentContext(PImplData->SharedContext.eglContext, PImplData->SharedContext.eglSurface);
-	}
-	else
-	{
-		SetCurrentContext(PImplData->SingleThreadedContext.eglContext, PImplData->SingleThreadedContext.eglSurface);
-	}
+	PImplData->RenderingContext.eglContext = CreateContext(EGL_NO_CONTEXT);
+	return PImplData->RenderingContext.eglContext != EGL_NO_CONTEXT;
 }
 
 void AndroidEGL::AcquireCurrentRenderingContext()
@@ -983,8 +796,6 @@ void AndroidEGL::AcquireCurrentRenderingContext()
 		// Dummy FBO we bind right after SwapBuffers to tell driver that backbuffer is no longer in use by the App
 		glGenFramebuffers(1, &PImplData->DummyFrameBuffer);
 		PImplData->RenderingContext.DummyFrameBuffer = PImplData->DummyFrameBuffer;
-		PImplData->SharedContext.DummyFrameBuffer = PImplData->DummyFrameBuffer;
-		PImplData->SingleThreadedContext.DummyFrameBuffer = PImplData->DummyFrameBuffer;
 	}
 
 	if (IsOfflineSurfaceRequired())
@@ -1004,20 +815,12 @@ void AndroidEGL::AcquireCurrentRenderingContext()
 
 void AndroidEGL::SetCurrentRenderingContext()
 {
-	PImplData->CurrentContextType = CONTEXT_Rendering;
-	if (GUseThreadedRendering)
-	{
-		SetCurrentContext(PImplData->RenderingContext.eglContext, PImplData->RenderingContext.eglSurface);
-	}
-	else
-	{
-		SetCurrentContext(PImplData->SingleThreadedContext.eglContext, PImplData->SingleThreadedContext.eglSurface);
-	}
+	SetCurrentContext(PImplData->RenderingContext.eglContext, PImplData->RenderingContext.eglSurface);
 }
 
 void AndroidEGL::ReleaseContextOwnership()
 {
-	if (GetCurrentContextType() == CONTEXT_Rendering)
+	if (PlatformOpenGLThreadHasRenderingContext())
 	{
 		FPlatformMisc::LowLevelOutputDebugStringf(TEXT("AndroidEGL::  ReleaseContextOwnership,  thread releasing rendering context tid: %d"), FPlatformTLS::GetCurrentThreadId());
 		SetCurrentContext(EGL_NO_CONTEXT, EGL_NO_SURFACE);
@@ -1033,65 +836,22 @@ void AndroidEGL::Terminate()
 	FPlatformMisc::LowLevelOutputDebugStringf(TEXT("AndroidEGL:: Terminate! tid: %d"), FPlatformTLS::GetCurrentThreadId());
 
 	ResetDisplay();
-	DestroyContext(PImplData->SharedContext.eglContext);
-	PImplData->SharedContext.Reset();
+
 	DestroyContext(PImplData->RenderingContext.eglContext);
 	PImplData->RenderingContext.Reset();
-	DestroyContext(PImplData->SingleThreadedContext.eglContext);
-	PImplData->SingleThreadedContext.Reset();
+
 	DestroyRenderSurface();
-	DestroySharedSurface();
 	TerminateEGL();
 }
 
-uint32_t AndroidEGL::GetCurrentContextType()
+bool AndroidEGL::ThreadHasRenderingContext()
 {
-	if(GUseThreadedRendering)
-	{
-		EGLContext CurrentContext = GetCurrentContext();
-		if (CurrentContext == PImplData->RenderingContext.eglContext)
-		{
-			return CONTEXT_Rendering;
-		}
-		else if (CurrentContext == PImplData->SharedContext.eglContext)
-		{
-			return CONTEXT_Shared;
-		}
-		else if (CurrentContext != EGL_NO_CONTEXT)
-		{
-			return CONTEXT_Other;
-		}
-	}
-	else
-	{
-		return CONTEXT_Shared;//make sure current context is valid one. //check(GetCurrentContext != NULL);
-	}
-
-	return CONTEXT_Invalid;
+	return GetCurrentContext() == PImplData->RenderingContext.eglContext;
 }
 
 FPlatformOpenGLContext* AndroidEGL::GetRenderingContext()
 {
-	if (GUseThreadedRendering)
-	{
-		return &PImplData->RenderingContext;
-	}
-	else
-	{
-		return &PImplData->SingleThreadedContext;
-	}
-}
-
-FPlatformOpenGLContext* AndroidEGL::GetSharedContext()
-{
-	if (GUseThreadedRendering)
-	{
-		return &PImplData->SharedContext;
-	}
-	else
-	{
-		return &PImplData->SingleThreadedContext;
-	}
+	return &PImplData->RenderingContext;
 }
 
 bool AndroidEGL::GetSupportsNoErrorContext()
@@ -1099,14 +859,6 @@ bool AndroidEGL::GetSupportsNoErrorContext()
 	return bSupportsKHRNoErrorContext;
 }
 
-void AndroidEGL::UnBind()
-{
-	FPlatformMisc::LowLevelOutputDebugString(TEXT("AndroidEGL::UnBind()"));
-	ResetDisplay();
-	DestroyRenderSurface();
-	DestroySharedSurface();
-}
-
 void AndroidEGL::UnBindRender()
 {
 	FPlatformMisc::LowLevelOutputDebugString(TEXT("AndroidEGL::UnBindRender()"));
@@ -1114,13 +866,6 @@ void AndroidEGL::UnBindRender()
 	DestroyRenderSurface();
 }
 
-void AndroidEGL::UnBindShared()
-{
-	FPlatformMisc::LowLevelOutputDebugString(TEXT("AndroidEGL::UnBindShared()"));
-	ResetDisplay();
-	DestroySharedSurface();
-}
-
 void FAndroidAppEntry::ReInitWindow(void* NewNativeWindowHandle)
 {
 	check(IsInGameThread());
@@ -1154,8 +899,6 @@ void AndroidEGL::RefreshWindowSize()
 	UE_LOG(LogAndroid, Log, TEXT("AndroidEGL::RefreshWindowSize updating window size = %d, %d, cached size : %d, %d tid : %d"), WindowRect.Right, WindowRect.Bottom, PImplData->CachedWindowRect.Right, PImplData->CachedWindowRect.Bottom, FPlatformTLS::GetCurrentThreadId());
 	PImplData->CachedWindowRect = WindowRect;
 
-	ResizeSharedContextSurface();
-
 	ENQUEUE_RENDER_COMMAND(EGLResizeRenderContextSurface)(
 		[](FRHICommandListImmediate& RHICmdList)
 	{
@@ -1325,7 +1068,7 @@ void SetSharedContextGameCommand(TSharedPtr<FEvent, ESPMode::ThreadSafe> GTBlock
 	check(IsInGameThread());
 	AndroidEGL* EGL = AndroidEGL::GetInstance();
 	EGL->SetCurrentContext(EGL_NO_CONTEXT, EGL_NO_SURFACE);
-	EGL->SetCurrentSharedContext();
+
 	GTBlockedTrigger->Trigger();
 }
 
@@ -1345,9 +1088,9 @@ void BlockRendering()
 #if !USE_ANDROID_ALTERNATIVE_SUSPEND
 	// Flush GT first in case it has any dependency on RT work to complete
 	FGraphEventRef GTBlockTask = FFunctionGraphTask::CreateAndDispatchWhenReady([BlockedTrigger]()
-		{
-			SetSharedContextGameCommand(BlockedTrigger);
-		}, TStatId(), NULL, ENamedThreads::GameThread);
+	{
+		SetSharedContextGameCommand(BlockedTrigger);
+	}, TStatId(), NULL, ENamedThreads::GameThread);
 
 	UE_LOG(LogAndroid, Log, TEXT("Waiting for game thread to release EGL context/surface."));
 	BlockedTrigger->Wait();
diff --git a/Engine/Source/Runtime/OpenGLDrv/Private/Android/AndroidEGL.h b/Engine/Source/Runtime/OpenGLDrv/Private/Android/AndroidEGL.h
--- a/Engine/Source/Runtime/OpenGLDrv/Private/Android/AndroidEGL.h
+++ b/Engine/Source/Runtime/OpenGLDrv/Private/Android/AndroidEGL.h
@@ -69,18 +69,18 @@ public:
 	~AndroidEGL();	
 
 	bool IsInitialized();
+
 	void InitBackBuffer();
 	void DestroyBackBuffer();
-	void Init( APIVariant API, uint32 MajorVersion, uint32 MinorVersion);
+
+	void Init(APIVariant API, uint32 MajorVersion, uint32 MinorVersion);
 	void ReInit();
-	void UnBind();
+
 	void UnBindRender();
-	void UnBindShared();
-	bool SwapBuffers();
+
 	void Terminate();
-	void InitSurface(bool bUseSmallSurface, bool bCreateWndSurface);
+
 	void InitRenderSurface(bool bUseSmallSurface, bool bCreateWndSurface);
-	void InitSharedSurface(bool bUseSmallSurface);
 	void UpdateBuffersTransform();
 	bool IsOfflineSurfaceRequired();
 
@@ -93,7 +93,7 @@ public:
 	ANativeWindow* GetNativeWindow() const;
 	void GetSwapIntervalRange(EGLint& OutMinSwapInterval, EGLint& OutMaxSwapInterval) const;
 
-	EGLContext CreateContext(EGLContext InSharedContext = EGL_NO_CONTEXT);
+	EGLContext CreateContext(EGLContext InParentContext);
 	int32 GetError();
 	EGLBoolean SetCurrentContext(EGLContext InContext, EGLSurface InSurface);
 
@@ -101,13 +101,12 @@ public:
 	void ReleaseContextOwnership();
 
 	GLuint GetResolveFrameBuffer();
-	bool IsCurrentContextValid();
-	EGLContext  GetCurrentContext(  );
-	void SetCurrentSharedContext();
+	EGLContext GetCurrentContext();
+
 	void SetCurrentRenderingContext();
-	uint32_t GetCurrentContextType();
+	bool ThreadHasRenderingContext();
 	FPlatformOpenGLContext* GetRenderingContext();
-	FPlatformOpenGLContext* GetSharedContext();
+
 	bool GetSupportsNoErrorContext();
 
 	// recreate the EGL surface for the current hardware window.
@@ -118,16 +117,14 @@ public:
 
 protected:
 	AndroidEGL();
-	static AndroidEGL* Singleton ;
+	static AndroidEGL* Singleton;
 
 private:
 	void InitEGL(APIVariant API);
 	void TerminateEGL();
 
 	void CreateEGLRenderSurface(ANativeWindow* InWindow, bool bCreateWndSurface);
-	void CreateEGLSharedSurface();
 	void DestroyRenderSurface();
-	void DestroySharedSurface();
 
 	bool InitContexts();
 	void DestroyContext(EGLContext InContext);
@@ -141,15 +138,14 @@ private:
 
 	// Actual Update to the egl surface to match the GT's requested size.
 	void ResizeRenderContextSurface();
-	void ResizeSharedContextSurface();
 
-	bool bSupportsKHRCreateContext;
-	bool bSupportsKHRSurfacelessContext;
-	bool bSupportsKHRNoErrorContext;
-	bool bSupportsEXTRobustContext;
-	bool bIsEXTRobustContextActive = false;
+	bool bSupportsKHRCreateContext      = false;
+	bool bSupportsKHRSurfacelessContext = false;
+	bool bSupportsKHRNoErrorContext     = false;
+	bool bSupportsEXTRobustContext      = false;
+	bool bIsEXTRobustContextActive      = false;
 
-	int *ContextAttributes;
+	int* ContextAttributes = nullptr;
 };
 
 #define ENABLE_CONFIG_FILTER 1
diff --git a/Engine/Source/Runtime/OpenGLDrv/Private/Android/AndroidOpenGL.cpp b/Engine/Source/Runtime/OpenGLDrv/Private/Android/AndroidOpenGL.cpp
--- a/Engine/Source/Runtime/OpenGLDrv/Private/Android/AndroidOpenGL.cpp
+++ b/Engine/Source/Runtime/OpenGLDrv/Private/Android/AndroidOpenGL.cpp
@@ -48,17 +48,12 @@ struct FPlatformOpenGLDevice
 {
 	bool TargetDirty;
 
-	void SetCurrentSharedContext();
 	void SetCurrentRenderingContext();
-	void SetupCurrentContext();
-	void SetCurrentNULLContext();
 
 	FPlatformOpenGLDevice();
 	~FPlatformOpenGLDevice();
 	void Init();
 	void LoadEXT();
-	void Terminate();
-	void ReInit();
 };
 
 
@@ -178,17 +173,23 @@ void FPlatformOpenGLDevice::Init()
 
 	FPlatformMisc::LowLevelOutputDebugString(TEXT("FPlatformOpenGLDevice:Init"));
 	bool bCreateSurface = !AndroidThunkCpp_IsOculusMobileApplication();
-	AndroidEGL::GetInstance()->InitSurface(false, bCreateSurface);
+	AndroidEGL::GetInstance()->InitRenderSurface(false, bCreateSurface);
 
 	LoadEXT();
 	PlatformRenderingContextSetup(this);
-
-	InitDefaultGLContextState();
 	InitDebugContext();
+	{
+		VERIFY_GL_SCOPE();
 
-	PlatformSharedContextSetup(this);
+		GLuint* DefaultVao = &AndroidEGL::GetInstance()->GetRenderingContext()->DefaultVertexArrayObject;
+
+		if (*DefaultVao == 0)
+		{
+			glGenVertexArrays(1, DefaultVao);
+			glBindVertexArray(*DefaultVao);
+		}
+	}
 	InitDefaultGLContextState();
-	InitDebugContext();
 
 	AndroidEGL::GetInstance()->InitBackBuffer(); //can be done only after context is made current.
 
@@ -299,27 +300,20 @@ bool PlatformBlitToViewport(IRHICommandContext& RHICmdContext, FPlatformOpenGLDe
 void PlatformRenderingContextSetup(FPlatformOpenGLDevice* Device)
 {
 	Device->SetCurrentRenderingContext();
-	Device->SetupCurrentContext();
 }
 
 void PlatformFlushIfNeeded()
 {
 }
 
-void PlatformSharedContextSetup(FPlatformOpenGLDevice* Device)
-{
-	Device->SetCurrentSharedContext();
-	Device->SetupCurrentContext();
-}
-
 void PlatformNULLContextSetup()
 {
 	AndroidEGL::GetInstance()->ReleaseContextOwnership();
 }
 
-EOpenGLCurrentContext PlatformOpenGLCurrentContext(FPlatformOpenGLDevice* Device)
+bool PlatformOpenGLThreadHasRenderingContext()
 {
-	return (EOpenGLCurrentContext)AndroidEGL::GetInstance()->GetCurrentContextType();
+	return AndroidEGL::GetInstance()->ThreadHasRenderingContext();
 }
 
 void PlatformRestoreDesktopDisplayMode()
@@ -392,22 +386,8 @@ bool PlatformInitOpenGL()
 	return true;
 }
 
-bool PlatformOpenGLContextValid()
-{
-	return AndroidEGL::GetInstance()->IsCurrentContextValid();
-}
-
 // =============================================================
 
-void PlatformGetNewOcclusionQuery( GLuint* OutQuery, uint64* OutQueryContext )
-{
-}
-
-bool PlatformContextIsCurrent( uint64 QueryContext )
-{
-	return true;
-}
-
 void FPlatformOpenGLDevice::LoadEXT()
 {
 	eglGetSystemTimeNV_p = (PFNEGLGETSYSTEMTIMENVPROC)((void*)eglGetProcAddress("eglGetSystemTimeNV"));
@@ -520,15 +500,6 @@ int32 PlatformGlGetError()
 
 // =============================================================
 
-void PlatformReleaseOcclusionQuery( GLuint Query, uint64 QueryContext )
-{
-}
-
-void FPlatformOpenGLDevice::SetCurrentSharedContext()
-{
-	AndroidEGL::GetInstance()->SetCurrentSharedContext();
-}
-
 void PlatformDestroyOpenGLDevice(FPlatformOpenGLDevice* Device)
 {
 	delete Device;
@@ -539,32 +510,6 @@ void FPlatformOpenGLDevice::SetCurrentRenderingContext()
 	AndroidEGL::GetInstance()->AcquireCurrentRenderingContext();
 }
 
-void FPlatformOpenGLDevice::SetupCurrentContext()
-{
-	GLuint* DefaultVao = nullptr;
-	EOpenGLCurrentContext ContextType = (EOpenGLCurrentContext)AndroidEGL::GetInstance()->GetCurrentContextType();
-
-	if (ContextType == CONTEXT_Rendering)
-	{
-		DefaultVao = &AndroidEGL::GetInstance()->GetRenderingContext()->DefaultVertexArrayObject;
-	}
-	else if (ContextType == CONTEXT_Shared)
-	{
-		DefaultVao = &AndroidEGL::GetInstance()->GetSharedContext()->DefaultVertexArrayObject;
-	}
-	else
-	{
-		//Invalid or Other return
-		return;
-	}
-	
-	if (*DefaultVao == 0)
-	{
-		glGenVertexArrays(1, DefaultVao);
-		glBindVertexArray(*DefaultVao);
-	}
-}
-
 //--------------------------------
 
 bool FAndroidOpenGL::SupportsFramebufferSRGBEnable()
diff --git a/Engine/Source/Runtime/OpenGLDrv/Private/Android/AndroidOpenGLPrivate.cpp b/Engine/Source/Runtime/OpenGLDrv/Private/Android/AndroidOpenGLPrivate.cpp
--- a/Engine/Source/Runtime/OpenGLDrv/Private/Android/AndroidOpenGLPrivate.cpp
+++ b/Engine/Source/Runtime/OpenGLDrv/Private/Android/AndroidOpenGLPrivate.cpp
@@ -4,6 +4,12 @@
 
 bool GAndroidGPUInfoReady = false;
 
+FAndroidGPUInfo& FAndroidGPUInfo::Get()
+{
+	static FAndroidGPUInfo This;
+	return This;
+}
+
 FAndroidGPUInfo::FAndroidGPUInfo()
 {
 	// this is only valid in the game thread, make sure we are initialized there before being called on other threads!
@@ -20,8 +26,8 @@ FAndroidGPUInfo::FAndroidGPUInfo()
 	// Do not create a window surface if the app is for Oculus Mobile (use small buffer)
 	bool bCreateSurface = !AndroidThunkCpp_IsOculusMobileApplication();
 	FPlatformMisc::LowLevelOutputDebugString(TEXT("FAndroidGPUInfo"));
-	EGL->InitSurface(false, bCreateSurface);
-	EGL->SetCurrentSharedContext();
+	EGL->InitRenderSurface(false, bCreateSurface);
+	EGL->SetCurrentRenderingContext();
 
 	// get extensions
 	// Process the extension caps directly here, as FOpenGL might not yet be setup
diff --git a/Engine/Source/Runtime/OpenGLDrv/Private/Android/AndroidOpenGLPrivate.h b/Engine/Source/Runtime/OpenGLDrv/Private/Android/AndroidOpenGLPrivate.h
--- a/Engine/Source/Runtime/OpenGLDrv/Private/Android/AndroidOpenGLPrivate.h
+++ b/Engine/Source/Runtime/OpenGLDrv/Private/Android/AndroidOpenGLPrivate.h
@@ -20,11 +20,7 @@ extern bool ShouldUseGPUFencesToLimitLatency();
 class FAndroidGPUInfo
 {
 public:
-	static FAndroidGPUInfo& Get()
-	{
-		static FAndroidGPUInfo This;
-		return This;
-	}
+	static FAndroidGPUInfo& Get();
 
 	FString GLVersion;
 	FString VendorName;
diff --git a/Engine/Source/Runtime/OpenGLDrv/Private/Linux/LinuxOpenGLPlatform.cpp b/Engine/Source/Runtime/OpenGLDrv/Private/Linux/LinuxOpenGLPlatform.cpp
--- a/Engine/Source/Runtime/OpenGLDrv/Private/Linux/LinuxOpenGLPlatform.cpp
+++ b/Engine/Source/Runtime/OpenGLDrv/Private/Linux/LinuxOpenGLPlatform.cpp
@@ -173,25 +173,28 @@ void PlatformReleaseOpenGLContext(FPlatformOpenGLDevice* Device, FPlatformOpenGL
 /** Platform specific OpenGL device. */
 struct FPlatformOpenGLDevice
 {
-	FPlatformOpenGLContext	SharedContext;
+	static FPlatformOpenGLDevice* Singleton;
+
 	FPlatformOpenGLContext	RenderingContext;
-	int32					NumUsedContexts;
 
 	/** Guards against operating on viewport contexts from more than one thread at the same time. */
 	FCriticalSection* ContextUsageGuard;
 
-	FPlatformOpenGLDevice() : NumUsedContexts(0)
+	FPlatformOpenGLDevice()
 	{
+		check(!Singleton);
+		Singleton = this;
+
 		extern void InitDebugContext();
 
 		ContextUsageGuard = new FCriticalSection;
 
 		verifyf(SDL_GL_SetAttribute(SDL_GL_SHARE_WITH_CURRENT_CONTEXT, 0) == 0, TEXT("SDL_GL_SetAttribute(SDL_GL_SHARE_WITH_CURRENT_CONTEXT, 0) failed: %s\n."), UTF8_TO_TCHAR(SDL_GetError()));
 
-		Linux_PlatformCreateDummyGLWindow(&SharedContext);
-		Linux_PlatformCreateOpenGLContextCore(&SharedContext);
+		Linux_PlatformCreateDummyGLWindow(&RenderingContext);
+		Linux_PlatformCreateOpenGLContextCore(&RenderingContext);
 
-		if (SharedContext.hGLContext == nullptr)
+		if (RenderingContext.hGLContext == nullptr)
 		{
 			FPlatformMisc::MessageBoxExt(EAppMsgType::Ok,
 				*(NSLOCTEXT("Renderer", "LinuxInsufficientDriversText", "Cannot create OpenGL context. Check that the drivers and hardware support at least OpenGL 4.3 (or re-run with -opengl3)").ToString()),
@@ -201,45 +204,31 @@ struct FPlatformOpenGLDevice
 			return;
 		}
 
-		{
-			FScopeContext ScopeContext(&SharedContext);
-			InitDebugContext();
-			glGenVertexArrays(1, &SharedContext.VertexArrayObject);
-			glBindVertexArray(SharedContext.VertexArrayObject);
-			InitDefaultGLContextState();
-		}
-
-		verifyf(SDL_GL_SetAttribute(SDL_GL_SHARE_WITH_CURRENT_CONTEXT, 1) == 0, TEXT("SDL_GL_SetAttribute(SDL_GL_SHARE_WITH_CURRENT_CONTEXT, 1) failed: %s\n."), UTF8_TO_TCHAR(SDL_GetError()));
-
-		Linux_ContextMakeCurrent(SharedContext.hWnd, SharedContext.hGLContext);
-
-		Linux_PlatformCreateDummyGLWindow(&RenderingContext);
-		Linux_PlatformCreateOpenGLContextCore(&RenderingContext);
+		Linux_ContextMakeCurrent(RenderingContext.hWnd, RenderingContext.hGLContext);
 
-		check(RenderingContext.hGLContext);
+		InitDebugContext();
+		glGenVertexArrays(1, &RenderingContext.VertexArrayObject);
+		glBindVertexArray(RenderingContext.VertexArrayObject);
+		InitDefaultGLContextState();
 
-		{
-			FScopeContext ScopeContext(&RenderingContext);
-			InitDebugContext();
-			glGenVertexArrays(1, &RenderingContext.VertexArrayObject);
-			glBindVertexArray(RenderingContext.VertexArrayObject);
-			InitDefaultGLContextState();
-		}
+		verifyf(SDL_GL_SetAttribute(SDL_GL_SHARE_WITH_CURRENT_CONTEXT, 1) == 0, TEXT("SDL_GL_SetAttribute(SDL_GL_SHARE_WITH_CURRENT_CONTEXT, 1) failed: %s\n."), UTF8_TO_TCHAR(SDL_GetError()));
 	}
 
 	~FPlatformOpenGLDevice()
 	{
-		check(NumUsedContexts == 0);
-
 		Linux_ContextMakeCurrent(NULL, NULL);
 
 		PlatformReleaseOpenGLContext(this, &RenderingContext);
-		PlatformReleaseOpenGLContext(this, &SharedContext);
 
 		delete ContextUsageGuard;
+
+		check(Singleton == this);
+		Singleton = nullptr;
 	}
 };
 
+FPlatformOpenGLDevice* FPlatformOpenGLDevice::Singleton = nullptr;
+
 FPlatformOpenGLDevice* PlatformCreateOpenGLDevice()
 {
 	return new FPlatformOpenGLDevice;
@@ -266,9 +255,10 @@ FPlatformOpenGLContext* PlatformCreateOpenGLContext(FPlatformOpenGLDevice* Devic
 	Context->hWnd = (SDL_HWindow)InWindowHandle;
 	Context->bReleaseWindowOnDestroy = false;
 
-	check(Device->SharedContext.hGLContext)
+	check(Device->RenderingContext.hGLContext);
+
 	{
-		FScopeContext Scope(&(Device->SharedContext));
+		FScopeContext Scope(&Device->RenderingContext);
 		verifyf(SDL_GL_SetAttribute(SDL_GL_SHARE_WITH_CURRENT_CONTEXT, 1) == 0, TEXT("SDL_GL_SetAttribute(SDL_GL_SHARE_WITH_CURRENT_CONTEXT, 1) failed: %s\n."), UTF8_TO_TCHAR(SDL_GetError()));
 		Linux_PlatformCreateOpenGLContextCore(Context);
 	}
@@ -531,20 +521,6 @@ void PlatformRenderingContextSetup(FPlatformOpenGLDevice* Device)
 	Linux_ContextMakeCurrent(Device->RenderingContext.hWnd, Device->RenderingContext.hGLContext);
 }
 
-void PlatformSharedContextSetup(FPlatformOpenGLDevice* Device)
-{
-	check(Device && Device->SharedContext.hWnd && Device->SharedContext.hGLContext);
-
-	// no need to glFlush() on Windows, it does flush by itself before switching contexts
-	if (Linux_GetCurrentContext())
-	{
-		glFlush();
-	}
-
-	Linux_ContextMakeCurrent(Device->SharedContext.hWnd, Device->SharedContext.hGLContext);
-}
-
-
 void PlatformNULLContextSetup()
 {
 	if (Linux_GetCurrentContext())
@@ -555,8 +531,6 @@ void PlatformNULLContextSetup()
 	Linux_ContextMakeCurrent(NULL, NULL);
 }
 
-
-
 /**
  * Resize the GL context.
  */
@@ -820,42 +794,18 @@ bool PlatformInitOpenGL()
 	return bOpenGLSupported;
 }
 
-bool PlatformOpenGLContextValid()
-{
-	return (Linux_GetCurrentContext() != NULL);
-}
-
 int32 PlatformGlGetError()
 {
 	return glGetError();
 }
 
-EOpenGLCurrentContext PlatformOpenGLCurrentContext(FPlatformOpenGLDevice* Device)
+bool PlatformOpenGLThreadHasRenderingContext()
 {
-	SDL_HGLContext hGLContext = Linux_GetCurrentContext();
-
-	if (LIKELY(hGLContext == Device->RenderingContext.hGLContext))	// most common case
-	{
-		return CONTEXT_Rendering;
-	}
-	else if (hGLContext == Device->SharedContext.hGLContext)
-	{
-		return CONTEXT_Shared;
-	}
-	else if (hGLContext)
-	{
-		return CONTEXT_Other;
-	}
-	return CONTEXT_Invalid;
+	return Linux_GetCurrentContext() == FPlatformOpenGLDevice::Singleton->RenderingContext.hGLContext;
 }
 
 // =============================================================
 
-bool PlatformContextIsCurrent(uint64 QueryContext)
-{
-	return (uint64)Linux_GetCurrentContext() == QueryContext;
-}
-
 FOpenGLTexture* PlatformCreateBuiltinBackBuffer(FOpenGLDynamicRHI* OpenGLRHI, uint32 SizeX, uint32 SizeY)
 {
 	return nullptr;
diff --git a/Engine/Source/Runtime/OpenGLDrv/Private/OpenGLBuffer.cpp b/Engine/Source/Runtime/OpenGLDrv/Private/OpenGLBuffer.cpp
--- a/Engine/Source/Runtime/OpenGLDrv/Private/OpenGLBuffer.cpp
+++ b/Engine/Source/Runtime/OpenGLDrv/Private/OpenGLBuffer.cpp
@@ -176,6 +176,17 @@ void BeginFrame_VertexBufferCleanup()
 	FrameBytes = 0;
 }
 
+void FOpenGLBufferBase::Bind()
+{
+	VERIFY_GL_SCOPE();
+	FOpenGLDynamicRHI::Get().CachedBindBuffer(Type, Resource);
+}
+
+void FOpenGLBufferBase::OnBufferDeletion()
+{
+	FOpenGLDynamicRHI::Get().OnBufferDeletion(Type, Resource);
+}
+
 FBufferRHIRef FOpenGLDynamicRHI::RHICreateBuffer(FRHICommandListBase& RHICmdList, FRHIBufferDesc const& Desc, ERHIAccess ResourceState, FRHIResourceCreateInfo& CreateInfo)
 {
 	const void* Data = nullptr;
diff --git a/Engine/Source/Runtime/OpenGLDrv/Private/OpenGLCommands.cpp b/Engine/Source/Runtime/OpenGLDrv/Private/OpenGLCommands.cpp
--- a/Engine/Source/Runtime/OpenGLDrv/Private/OpenGLCommands.cpp
+++ b/Engine/Source/Runtime/OpenGLDrv/Private/OpenGLCommands.cpp
@@ -206,7 +206,7 @@ void FOpenGLDynamicRHI::RHISetRasterizerState(FRHIRasterizerState* NewStateRHI)
 	PendingState.RasterizerState = NewState->Data;
 }
 
-void FOpenGLDynamicRHI::UpdateRasterizerStateInOpenGLContext( FOpenGLContextState& ContextState )
+void FOpenGLDynamicRHI::UpdateRasterizerStateInOpenGLContext()
 {
 	VERIFY_GL_SCOPE();
 	if (FOpenGL::SupportsPolygonMode() && ContextState.RasterizerState.FillMode != PendingState.RasterizerState.FillMode)
@@ -281,7 +281,7 @@ void FOpenGLDynamicRHI::UpdateRasterizerStateInOpenGLContext( FOpenGLContextStat
 	}
 }
 
-void FOpenGLDynamicRHI::UpdateViewportInOpenGLContext( FOpenGLContextState& ContextState )
+void FOpenGLDynamicRHI::UpdateViewportInOpenGLContext()
 {
 	VERIFY_GL_SCOPE();
 	if (ContextState.Viewport != PendingState.Viewport)
@@ -328,7 +328,7 @@ void FOpenGLDynamicRHI::RHISetScissorRect(bool bEnable,uint32 MinX,uint32 MinY,u
 	PendingState.Scissor.Max.Y = MaxY;
 }
 
-inline void FOpenGLDynamicRHI::UpdateScissorRectInOpenGLContext( FOpenGLContextState& ContextState )
+void FOpenGLDynamicRHI::UpdateScissorRectInOpenGLContext()
 {
 	VERIFY_GL_SCOPE();
 	if (ContextState.bScissorEnabled != PendingState.bScissorEnabled)
@@ -344,8 +344,7 @@ inline void FOpenGLDynamicRHI::UpdateScissorRectInOpenGLContext( FOpenGLContextS
 		ContextState.bScissorEnabled = PendingState.bScissorEnabled;
 	}
 
-	if( PendingState.bScissorEnabled &&
-		ContextState.Scissor != PendingState.Scissor )
+	if (PendingState.bScissorEnabled && ContextState.Scissor != PendingState.Scissor)
 	{
 		check(PendingState.Scissor.Min.X <= PendingState.Scissor.Max.X);
 		check(PendingState.Scissor.Min.Y <= PendingState.Scissor.Max.Y);
@@ -386,14 +385,14 @@ void FOpenGLDynamicRHI::InternalSetSamplerStates(GLint TextureIndex, FOpenGLSamp
 	PendingState.SamplerStates[TextureIndex] = SamplerState;
 }
 
-void FOpenGLDynamicRHI::CachedSetupTextureStageInner(FOpenGLContextState& ContextState, GLint TextureIndex, GLenum Target, GLuint Resource, GLint LimitMip, GLint NumMips)
+void FOpenGLDynamicRHI::CachedSetupTextureStageInner(GLint TextureIndex, GLenum Target, GLuint Resource, GLint LimitMip, GLint NumMips)
 {
 	DETAILED_QUICK_SCOPE_CYCLE_COUNTER(STAT_CachedSetupTextureStage);
 	VERIFY_GL_SCOPE();
 	FTextureStage& TextureState = ContextState.Textures[TextureIndex];
 
 	// Something will have to be changed. Switch to the stage in question.
-	if( ContextState.ActiveTexture != TextureIndex )
+	if (ContextState.ActiveTexture != TextureIndex)
 	{
 		glActiveTexture( GL_TEXTURE0 + TextureIndex );
 		ContextState.ActiveTexture = TextureIndex;
@@ -491,7 +490,7 @@ void FOpenGLDynamicRHI::CachedSetupTextureStageInner(FOpenGLContextState& Contex
 	TextureState.Resource = Resource;
 }
 
-inline void FOpenGLDynamicRHI::ApplyTextureStage(FOpenGLContextState& ContextState, GLint TextureIndex, const FTextureStage& TextureStage, FOpenGLSamplerState* SamplerState)
+inline void FOpenGLDynamicRHI::ApplyTextureStage(GLint TextureIndex, const FTextureStage& TextureStage, FOpenGLSamplerState* SamplerState)
 {
 	GLenum Target = TextureStage.Target;
 	VERIFY_GL_SCOPE();
@@ -499,7 +498,7 @@ inline void FOpenGLDynamicRHI::ApplyTextureStage(FOpenGLContextState& ContextSta
 	if (!bHasTexture || TextureStage.Texture->SamplerState != SamplerState)
 	{
 		// Texture must be bound first
-		if( ContextState.ActiveTexture != TextureIndex )
+		if (ContextState.ActiveTexture != TextureIndex)
 		{
 			glActiveTexture(GL_TEXTURE0 + TextureIndex);
 			ContextState.ActiveTexture = TextureIndex;
@@ -550,7 +549,7 @@ inline void FOpenGLDynamicRHI::ApplyTextureStage(FOpenGLContextState& ContextSta
 }
 
 template <typename StateType>
-void FOpenGLDynamicRHI::SetupTexturesForDraw( FOpenGLContextState& ContextState, const StateType& ShaderState, int32 MaxTexturesNeeded )
+void FOpenGLDynamicRHI::SetupTexturesForDraw(const StateType& ShaderState, int32 MaxTexturesNeeded)
 {
 	VERIFY_GL_SCOPE();
 	SCOPE_CYCLE_COUNTER_DETAILED(STAT_OpenGLTextureBindTime);
@@ -566,7 +565,7 @@ void FOpenGLDynamicRHI::SetupTexturesForDraw( FOpenGLContextState& ContextState,
 			// it won't be useful for this draw, so telling OpenGL we don't really need it to give the driver
 			// more leeway in memory management, and avoid false alarms about same texture being set on
 			// texture stage and in framebuffer.
-			CachedSetupTextureStage( ContextState, TextureStageIndex, GL_NONE, 0, -1, 1 );
+			CachedSetupTextureStage(TextureStageIndex, GL_NONE, 0, -1, 1 );
 		}
 		else
 		{
@@ -603,7 +602,7 @@ void FOpenGLDynamicRHI::SetupTexturesForDraw( FOpenGLContextState& ContextState,
 				}
 			}
 #endif
-			CachedSetupTextureStage( ContextState, TextureStageIndex, TextureStage.Target, TextureStage.Resource, TextureStage.LimitMip, TextureStage.NumMips );
+			CachedSetupTextureStage(TextureStageIndex, TextureStage.Target, TextureStage.Resource, TextureStage.LimitMip, TextureStage.NumMips );
 			
 			bool bExternalTexture = (TextureStage.Target == GL_TEXTURE_EXTERNAL_OES);
 			if (!bExternalTexture)
@@ -620,21 +619,21 @@ void FOpenGLDynamicRHI::SetupTexturesForDraw( FOpenGLContextState& ContextState,
 			{
 				FOpenGL::BindSampler(TextureStageIndex, 0);
 				ContextState.SamplerStates[TextureStageIndex] = nullptr;
-				ApplyTextureStage( ContextState, TextureStageIndex, TextureStage, PendingState.SamplerStates[TextureStageIndex] );
+				ApplyTextureStage(TextureStageIndex, TextureStage, PendingState.SamplerStates[TextureStageIndex] );
 			}
 		}
 	}
 
 	// For now, continue to clear unused stages
-	for( int32 TextureStageIndex = MaxProgramTexture + 1; TextureStageIndex < MaxTexturesNeeded; ++TextureStageIndex )
+	for (int32 TextureStageIndex = MaxProgramTexture + 1; TextureStageIndex < MaxTexturesNeeded; ++TextureStageIndex)
 	{
-		CachedSetupTextureStage( ContextState, TextureStageIndex, GL_NONE, 0, -1, 1 );
+		CachedSetupTextureStage(TextureStageIndex, GL_NONE, 0, -1, 1 );
 	}
 }
 
-void FOpenGLDynamicRHI::SetupTexturesForDraw( FOpenGLContextState& ContextState )
+void FOpenGLDynamicRHI::SetupTexturesForDraw()
 {
-	SetupTexturesForDraw(ContextState, PendingState.BoundShaderState, FOpenGL::GetMaxCombinedTextureImageUnits());
+	SetupTexturesForDraw(PendingState.BoundShaderState, FOpenGL::GetMaxCombinedTextureImageUnits());
 }
 
 void FOpenGLDynamicRHI::InternalSetShaderImageUAV(GLint UAVIndex, GLenum Format, GLuint Resource, bool bLayered, GLint Layer, GLenum Access, GLint Level)
@@ -658,21 +657,21 @@ void FOpenGLDynamicRHI::InternalSetShaderBufferUAV(GLint UAVIndex, GLuint Resour
 	PendingState.UAVs[UAVIndex].bLayered = false;
 }
 
-void FOpenGLDynamicRHI::SetupUAVsForDraw(FOpenGLContextState& ContextState)
+void FOpenGLDynamicRHI::SetupUAVsForDraw()
 {
 	int32 MaxUAVUnitUsed = 0;
 	const TBitArray<>& NeededBits = PendingState.BoundShaderState->GetUAVNeeds(MaxUAVUnitUsed);
-	SetupUAVsForProgram(ContextState, NeededBits, MaxUAVUnitUsed);
+	SetupUAVsForProgram(NeededBits, MaxUAVUnitUsed);
 }
 
-void FOpenGLDynamicRHI::SetupUAVsForCompute(FOpenGLContextState& ContextState, const FOpenGLComputeShader* ComputeShader)
+void FOpenGLDynamicRHI::SetupUAVsForCompute(const FOpenGLComputeShader* ComputeShader)
 {
 	int32 MaxUAVUnitUsed = 0;
 	const TBitArray<>& NeededBits = ComputeShader->GetUAVNeeds(MaxUAVUnitUsed);
-	SetupUAVsForProgram(ContextState, NeededBits, MaxUAVUnitUsed);
+	SetupUAVsForProgram(NeededBits, MaxUAVUnitUsed);
 }
 
-void FOpenGLDynamicRHI::SetupUAVsForProgram(FOpenGLContextState& ContextState, const TBitArray<>& NeededBits, int32 MaxUAVUnitUsed)
+void FOpenGLDynamicRHI::SetupUAVsForProgram(const TBitArray<>& NeededBits, int32 MaxUAVUnitUsed)
 {
 	if (MaxUAVUnitUsed < 0 && ContextState.ActiveUAVMask == 0)
 	{
@@ -684,12 +683,12 @@ void FOpenGLDynamicRHI::SetupUAVsForProgram(FOpenGLContextState& ContextState, c
 	{
 		if (!NeededBits[UAVStageIndex])
 		{
-			CachedSetupUAVStage(ContextState, UAVStageIndex, 0, 0, false, 0, GL_READ_WRITE, 0);
+			CachedSetupUAVStage(UAVStageIndex, 0, 0, false, 0, GL_READ_WRITE, 0);
 		}
 		else
 		{
 			const FUAVStage& UAVStage = PendingState.UAVs[UAVStageIndex];
-			CachedSetupUAVStage(ContextState, UAVStageIndex, UAVStage.Format, UAVStage.Resource, UAVStage.bLayered, UAVStage.Layer, UAVStage.Access, UAVStage.Level);
+			CachedSetupUAVStage(UAVStageIndex, UAVStage.Format, UAVStage.Resource, UAVStage.bLayered, UAVStage.Layer, UAVStage.Access, UAVStage.Level);
 		}
 	}
 
@@ -700,7 +699,7 @@ void FOpenGLDynamicRHI::SetupUAVsForProgram(FOpenGLContextState& ContextState, c
 		const int32 NumUAVs = ContextState.UAVs.Num();
 		for (; UAVStageIndex < NumUAVs; ++UAVStageIndex)
 		{
-			CachedSetupUAVStage(ContextState, UAVStageIndex, 0, 0, false, 0, GL_READ_WRITE, 0);
+			CachedSetupUAVStage(UAVStageIndex, 0, 0, false, 0, GL_READ_WRITE, 0);
 		}
 	}
 }
@@ -735,11 +734,11 @@ static bool IsImageTextureFormatSupported(GLenum Format)
 	return true;
 }
 
-void FOpenGLDynamicRHI::CachedSetupUAVStage( FOpenGLContextState& ContextState, GLint UAVIndex, GLenum Format, GLuint Resource, bool bLayered, GLint Layer, GLenum Access, GLint Level)
+void FOpenGLDynamicRHI::CachedSetupUAVStage(GLint UAVIndex, GLenum Format, GLuint Resource, bool bLayered, GLint Layer, GLenum Access, GLint Level)
 {
 	VERIFY_GL_SCOPE();
 	
-	FUAVStage& UAVStage = ContextState.UAVs[UAVIndex];	
+	FUAVStage& UAVStage = ContextState.UAVs[UAVIndex];
 
 	if (UAVStage.Format == Format && 
 		UAVStage.Resource == Resource &&
@@ -822,11 +821,11 @@ void FOpenGLDynamicRHI::CachedSetupUAVStage( FOpenGLContextState& ContextState,
 	uint32 UAVBit = 1 << UAVIndex;
 	if (Resource != 0)
 	{
-		ContextState.ActiveUAVMask|= UAVBit;
+		ContextState.ActiveUAVMask |= UAVBit;
 	}
 	else
 	{
-		ContextState.ActiveUAVMask&= ~UAVBit;
+		ContextState.ActiveUAVMask &= ~UAVBit;
 	}
 }
 
@@ -859,13 +858,11 @@ void FOpenGLDynamicRHI::UpdateSRV(FOpenGLShaderResourceView* SRV)
 		glPixelStorei(GL_PACK_ALIGNMENT, 4);
 		glBindBuffer( GL_PIXEL_PACK_BUFFER, 0 );
 		
-		FOpenGLContextState& ContextState = GetContextStateForCurrentContext();
-		
 		GLenum Target = SRV->Target;
 		
-		CachedSetupTextureStage(ContextState, FOpenGL::GetMaxCombinedTextureImageUnits() - 1, Target, SRV->Resource, -1, 1);
+		CachedSetupTextureStage(FOpenGL::GetMaxCombinedTextureImageUnits() - 1, Target, SRV->Resource, -1, 1);
 		
-		CachedBindPixelUnpackBuffer(ContextState, PixelBuffer->Resource);
+		CachedBindBuffer(GL_PIXEL_UNPACK_BUFFER, PixelBuffer->Resource);
 		
 		glPixelStorei(GL_UNPACK_ROW_LENGTH, SizeX);
 		
@@ -875,7 +872,7 @@ void FOpenGLDynamicRHI::UpdateSRV(FOpenGLShaderResourceView* SRV)
 		
 		glPixelStorei(GL_UNPACK_ROW_LENGTH, 0);
 		
-		CachedBindPixelUnpackBuffer(ContextState, 0);
+		CachedBindBuffer(GL_PIXEL_UNPACK_BUFFER, 0);
 		
 		glBindFramebuffer(GL_FRAMEBUFFER, ContextState.Framebuffer);
 		ContextState.Framebuffer = -1;
@@ -1240,7 +1237,7 @@ void FOpenGLDynamicRHI::RHISetStencilRef(uint32 StencilRef)
 	PendingState.StencilRef = StencilRef;
 }
 
-void FOpenGLDynamicRHI::UpdateDepthStencilStateInOpenGLContext( FOpenGLContextState& ContextState )
+void FOpenGLDynamicRHI::UpdateDepthStencilStateInOpenGLContext()
 {
 	VERIFY_GL_SCOPE();
 	if (ContextState.DepthStencilState.bZEnable != PendingState.DepthStencilState.bZEnable)
@@ -1377,7 +1374,7 @@ void FOpenGLDynamicRHI::UpdateDepthStencilStateInOpenGLContext( FOpenGLContextSt
 	}
 }
 
-void FOpenGLDynamicRHI::SetPendingBlendStateForActiveRenderTargets( FOpenGLContextState& ContextState )
+void FOpenGLDynamicRHI::SetPendingBlendStateForActiveRenderTargets()
 {
 	VERIFY_GL_SCOPE();
 
@@ -1807,7 +1804,7 @@ void FOpenGLDynamicRHI::SetRenderTargetsAndClear(const FRHISetRenderTargetsInfo&
 
 // Primitive drawing.
 
-void FOpenGLDynamicRHI::SetupVertexArrays(FOpenGLContextState& ContextState, uint32 BaseVertexIndex, FOpenGLStream* Streams, uint32 NumStreams, uint32 MaxVertices)
+void FOpenGLDynamicRHI::SetupVertexArrays(uint32 BaseVertexIndex, FOpenGLStream* Streams, uint32 NumStreams, uint32 MaxVertices)
 {
 	VERIFY_GL_SCOPE();
 	bool KnowsDivisor[NUM_OPENGL_VERTEX_STREAMS] = { 0 };
@@ -1833,7 +1830,7 @@ void FOpenGLDynamicRHI::SetupVertexArrays(FOpenGLContextState& ContextState, uin
 			const uint32 StreamIndex = VertexElement.StreamIndex;
 
 			//only setup/track attributes actually in use
-			FOpenGLCachedAttr &Attr = ContextState.VertexAttrs[AttributeIndex];
+			FOpenGLCachedAttr& Attr = ContextState.VertexAttrs[AttributeIndex];
 			if (AttributeMask.IsFieldEnabled(AttributeIndex))
 			{
 				if (VertexElement.StreamIndex < NumStreams)
@@ -1983,139 +1980,68 @@ void FOpenGLDynamicRHI::SetupVertexArrays(FOpenGLContextState& ContextState, uin
 	check(NotUsedButActiveStreamMask == 0);
 }
 
-void FOpenGLDynamicRHI::OnProgramDeletion( GLint ProgramResource )
+void FOpenGLDynamicRHI::OnProgramDeletion(GLint ProgramResource)
 {
 	VERIFY_GL_SCOPE();
-	if( SharedContextState.Program == ProgramResource )
+	if (ContextState.Program == ProgramResource)
 	{
-		SharedContextState.Program = -1;
-	}
-
-	if( RenderingContextState.Program == ProgramResource )
-	{
-		RenderingContextState.Program = -1;
+		ContextState.Program = -1;
 	}
 }
 
-void FOpenGLDynamicRHI::OnBufferDeletion( GLuint BufferResource )
+void FOpenGLDynamicRHI::OnBufferDeletion(GLenum Type, GLuint BufferResource)
 {
 	VERIFY_GL_SCOPE();
-
-	if (SharedContextState.ArrayBufferBound == BufferResource)
-	{
-		SharedContextState.ArrayBufferBound = -1;	// will force refresh
-	}
-
-	if (RenderingContextState.ArrayBufferBound == BufferResource)
-	{
-		RenderingContextState.ArrayBufferBound = -1;	// will force refresh
-	}
-
-	if (SharedContextState.StorageBufferBound == BufferResource)
-	{
-		SharedContextState.StorageBufferBound = -1;	// will force refresh
-	}
-
-	if (RenderingContextState.StorageBufferBound == BufferResource)
+	if (Type == GL_PIXEL_UNPACK_BUFFER)
 	{
-		RenderingContextState.StorageBufferBound = -1;	// will force refresh
-	}
-
-	// loop through active streams
-	uint32 ActiveStreamMask = SharedContextState.ActiveStreamMask;
-	for (GLuint StreamIndex = 0; StreamIndex < NUM_OPENGL_VERTEX_STREAMS && ActiveStreamMask; StreamIndex++)
-	{
-		FOpenGLStream& CachedStream = SharedContextState.VertexStreams[StreamIndex];
-		if ((ActiveStreamMask & 0x1) && 
-			CachedStream.VertexBufferResource == BufferResource)
+		if (ContextState.PixelUnpackBufferBound == BufferResource)
 		{
-			FOpenGL::BindVertexBuffer(StreamIndex, 0, 0, 0); // brianh@nvidia: work around driver bug 1809000
-			CachedStream.VertexBufferResource = 0;
-			CachedStream.Offset = 0;
-			CachedStream.Stride = 0;
+			ContextState.PixelUnpackBufferBound = -1;	// will force refresh
 		}
-		ActiveStreamMask >>= 1;
 	}
-	
-	// loop through active streams
-	ActiveStreamMask = RenderingContextState.ActiveStreamMask;
-	for (GLuint StreamIndex = 0; StreamIndex < NUM_OPENGL_VERTEX_STREAMS && ActiveStreamMask; StreamIndex++)
+	else
 	{
-		FOpenGLStream& CachedStream = RenderingContextState.VertexStreams[StreamIndex];
-		if ((ActiveStreamMask & 0x1) && 
-			CachedStream.VertexBufferResource == BufferResource)
+		if (ContextState.ArrayBufferBound == BufferResource)
 		{
-			FOpenGL::BindVertexBuffer(StreamIndex, 0, 0, 0); // brianh@nvidia: work around driver bug 1809000
-			CachedStream.VertexBufferResource = 0;
-			CachedStream.Offset = 0;
-			CachedStream.Stride = 0;
+			ContextState.ArrayBufferBound = -1; // will force refresh
+		}
+
+		if (ContextState.StorageBufferBound == BufferResource)
+		{
+			ContextState.StorageBufferBound = -1; // will force refresh
 		}
-		ActiveStreamMask >>= 1;
-	}	
 	
-	// Storage buffer
-	{
-		int32 NumStorageBuffers = PendingState.UAVs.Num();
-		for (int32 StorageBufferIndex = 0; StorageBufferIndex < NumStorageBuffers; StorageBufferIndex++)
+		// loop through active streams
+		uint32 ActiveStreamMask = ContextState.ActiveStreamMask;
+		for (GLuint StreamIndex = 0; StreamIndex < NUM_OPENGL_VERTEX_STREAMS && ActiveStreamMask; StreamIndex++)
 		{
-			if (PendingState.UAVs[StorageBufferIndex].Format == 0 && PendingState.UAVs[StorageBufferIndex].Resource == BufferResource)
+			FOpenGLStream& CachedStream = ContextState.VertexStreams[StreamIndex];
+			if ((ActiveStreamMask & 0x1) && 
+				CachedStream.VertexBufferResource == BufferResource)
 			{
-				PendingState.UAVs[StorageBufferIndex].Resource = 0;
+				FOpenGL::BindVertexBuffer(StreamIndex, 0, 0, 0); // brianh@nvidia: work around driver bug 1809000
+				CachedStream.VertexBufferResource = 0;
+				CachedStream.Offset = 0;
+				CachedStream.Stride = 0;
 			}
-		}
-	}
-
-	if (SharedContextState.ElementArrayBufferBound == BufferResource)
-	{
-		SharedContextState.ElementArrayBufferBound = -1;	// will force refresh
-	}
-
-	if (RenderingContextState.ElementArrayBufferBound == BufferResource)
-	{
-		RenderingContextState.ElementArrayBufferBound = -1;	// will force refresh
-	}
-}
-
-void FOpenGLDynamicRHI::OnPixelBufferDeletion( GLuint PixelBufferResource )
-{
-	VERIFY_GL_SCOPE();
-	if (SharedContextState.PixelUnpackBufferBound == PixelBufferResource)
-	{
-		SharedContextState.PixelUnpackBufferBound = -1;	// will force refresh
-	}
-
-	if (RenderingContextState.PixelUnpackBufferBound == PixelBufferResource)
-	{
-		RenderingContextState.PixelUnpackBufferBound = -1;	// will force refresh
-	}
-}
-
-void FOpenGLDynamicRHI::OnUniformBufferDeletion( GLuint UniformBufferResource, uint32 AllocatedSize, bool bStreamDraw )
-{
-	VERIFY_GL_SCOPE();
-	if (SharedContextState.UniformBufferBound == UniformBufferResource)
-	{
-		SharedContextState.UniformBufferBound = -1;	// will force refresh
-	}
-
-	if (RenderingContextState.UniformBufferBound == UniformBufferResource)
-	{
-		RenderingContextState.UniformBufferBound = -1;	// will force refresh
-	}
-
-	if (!GUseEmulatedUniformBuffers)
-	{
-		for (GLuint UniformBufferIndex = 0; UniformBufferIndex < CrossCompiler::NUM_SHADER_STAGES*OGL_MAX_UNIFORM_BUFFER_BINDINGS; UniformBufferIndex++)
+			ActiveStreamMask >>= 1;
+		}	
+	
+		// Storage buffer
 		{
-			if (SharedContextState.UniformBuffers[UniformBufferIndex] == UniformBufferResource)
+			int32 NumStorageBuffers = PendingState.UAVs.Num();
+			for (int32 StorageBufferIndex = 0; StorageBufferIndex < NumStorageBuffers; StorageBufferIndex++)
 			{
-				SharedContextState.UniformBuffers[UniformBufferIndex] = FOpenGLCachedUniformBuffer_Invalid;	// that'll enforce state update on next cache test
+				if (PendingState.UAVs[StorageBufferIndex].Format == 0 && PendingState.UAVs[StorageBufferIndex].Resource == BufferResource)
+				{
+					PendingState.UAVs[StorageBufferIndex].Resource = 0;
+				}
 			}
+		}
 
-			if (RenderingContextState.UniformBuffers[UniformBufferIndex] == UniformBufferResource)
-			{
-				RenderingContextState.UniformBuffers[UniformBufferIndex] = FOpenGLCachedUniformBuffer_Invalid;	// that'll enforce state update on next cache test
-			}
+		if (ContextState.ElementArrayBufferBound == BufferResource)
+		{
+			ContextState.ElementArrayBufferBound = -1;	// will force refresh
 		}
 	}
 }
@@ -2225,84 +2151,57 @@ void FOpenGLDynamicRHI::CommitComputeResourceTables(FOpenGLComputeShader* Comput
 	PendingState.DirtyUniformBuffers[SF_Compute] = 0;
 }
 
-class FOpengGLFenceKick
+FOpenGLDynamicRHI::FOpenGLFenceKick::~FOpenGLFenceKick()
 {
-public:
-
-	~FOpengGLFenceKick() { Reset(); }
+	Reset();
+}
 
-	void OnDrawCall(const FOpenGLContextState& ContextState)
+void FOpenGLDynamicRHI::FOpenGLFenceKick::OnDrawCall()
+{
+	if (OpenGLConsoleVariables::GOpenGLFenceKickPerDrawCount)
 	{
-		if(OpenGLConsoleVariables::GOpenGLFenceKickPerDrawCount)
+		// we kick every GOpenGLFenceKickPerDrawCount draws within a render pass.
+		// counter is reset when the FBO changes.
+		if (RHI.ContextState.Framebuffer != LastSeenFramebuffer)
 		{
-			// we kick every GOpenGLFenceKickPerDrawCount draws within a render pass.
-			// counter is reset when the FBO changes.
-			if (ContextState.Framebuffer != LastSeenFramebuffer)
-			{
-				LastSeenFramebuffer = ContextState.Framebuffer;
-				DrawCounter = 0;
-			}
-			else if (++DrawCounter >= OpenGLConsoleVariables::GOpenGLFenceKickPerDrawCount)
-			{
-				InsertKick();
-				DrawCounter = 0;
-			}
+			LastSeenFramebuffer = RHI.ContextState.Framebuffer;
+			DrawCounter = 0;
 		}
-	}
-
-	void Reset()
-	{
-		for (UGLsync Sync : Syncs)
+		else if (++DrawCounter >= OpenGLConsoleVariables::GOpenGLFenceKickPerDrawCount)
 		{
-			FOpenGL::DeleteSync(Sync);
+			InsertKick();
+			DrawCounter = 0;
 		}
-		Syncs.Reset();
-		DrawCounter = 0;
-		LastSeenFramebuffer = 0;
 	}
+}
 
-private:
-
-	void InsertKick()
+void FOpenGLDynamicRHI::FOpenGLFenceKick::Reset()
+{
+	for (UGLsync Sync : Syncs)
 	{
-		// record syncs incase a create/destroy pair is optimised by the driver.
-		UGLsync Sync = FOpenGL::FenceSync(GL_SYNC_GPU_COMMANDS_COMPLETE, 0);
-		Syncs.Add(Sync);
-		check(Syncs.Num() < 10000); // arbitrary check to ensure that we do eventually reset the array.
+		FOpenGL::DeleteSync(Sync);
 	}
+	Syncs.Reset();
+	DrawCounter = 0;
+	LastSeenFramebuffer = 0;
+}
 
-	TArray<UGLsync> Syncs;
-	int32 DrawCounter = 0;
-	GLuint LastSeenFramebuffer = 0;
-};
-static FOpengGLFenceKick GOpenGLKickHint;
-
-void OpenGLCommands_OnEndFrame()
+void FOpenGLDynamicRHI::FOpenGLFenceKick::InsertKick()
 {
-	GOpenGLKickHint.Reset();
+	// record syncs incase a create/destroy pair is optimised by the driver.
+	UGLsync Sync = FOpenGL::FenceSync(GL_SYNC_GPU_COMMANDS_COMPLETE, 0);
+	Syncs.Add(Sync);
+	check(Syncs.Num() < 10000); // arbitrary check to ensure that we do eventually reset the array.
 }
 
 void FOpenGLDynamicRHI::RHIDrawPrimitive(uint32 BaseVertexIndex,uint32 NumPrimitives,uint32 NumInstances)
 {
-	SCOPE_CYCLE_COUNTER_DETAILED(STAT_OpenGLDrawPrimitiveTime);
 	VERIFY_GL_SCOPE();
+	SCOPE_CYCLE_COUNTER_DETAILED(STAT_OpenGLDrawPrimitiveTime);
 	RHI_DRAW_CALL_STATS(PrimitiveType,NumPrimitives*NumInstances);
 
-	FOpenGLContextState& ContextState = GetContextStateForCurrentContext();
-	BindPendingFramebuffer(ContextState);
-	SetPendingBlendStateForActiveRenderTargets(ContextState);
-	UpdateViewportInOpenGLContext(ContextState);
-	UpdateScissorRectInOpenGLContext(ContextState);
-	UpdateRasterizerStateInOpenGLContext(ContextState);
-	UpdateDepthStencilStateInOpenGLContext(ContextState);
-	BindPendingShaderState(ContextState);
-	CommitGraphicsResourceTables();
-	SetupTexturesForDraw(ContextState);
-	SetupUAVsForDraw(ContextState);
-	CommitNonComputeShaderConstants();
-	CachedBindElementArrayBuffer(ContextState,0);
-	uint32 VertexCount = GetVertexCountForPrimitiveCount(NumPrimitives,PrimitiveType);
-	SetupVertexArrays(ContextState, BaseVertexIndex, PendingState.Streams, NUM_OPENGL_VERTEX_STREAMS, VertexCount);
+	uint32 VertexCount = GetVertexCountForPrimitiveCount(NumPrimitives, PrimitiveType);
+	SetupDraw(nullptr, BaseVertexIndex, VertexCount);
 
 	GLenum DrawMode = GL_TRIANGLES;
 	GLsizei NumElements = 0;
@@ -2313,116 +2212,87 @@ void FOpenGLDynamicRHI::RHIDrawPrimitive(uint32 BaseVertexIndex,uint32 NumPrimit
 	{
 		SCOPE_CYCLE_COUNTER_DETAILED(STAT_OpenGLDrawPrimitiveDriverTime);
 		CONDITIONAL_SCOPE_CYCLE_COUNTER(STAT_OpenGLShaderFirstDrawTime, PendingState.BoundShaderState->RequiresDriverInstantiation());
+
 		glDrawArrays(DrawMode, 0, NumElements);
 	}
 	else
 	{
 		SCOPE_CYCLE_COUNTER_DETAILED(STAT_OpenGLDrawPrimitiveDriverTime);
 		CONDITIONAL_SCOPE_CYCLE_COUNTER(STAT_OpenGLShaderFirstDrawTime, PendingState.BoundShaderState->RequiresDriverInstantiation());
+
 		FOpenGL::DrawArraysInstanced(DrawMode, 0, NumElements, NumInstances);
 	}
-	GOpenGLKickHint.OnDrawCall(ContextState);
+
+	KickHint.OnDrawCall();
 }
 
 void FOpenGLDynamicRHI::RHIDrawPrimitiveIndirect(FRHIBuffer* ArgumentBufferRHI, uint32 ArgumentOffset)
 {
-	if (FOpenGL::SupportsDrawIndirect())
+	if (!FOpenGL::SupportsDrawIndirect())
 	{
-		VERIFY_GL_SCOPE();
+		UE_LOG(LogRHI, Fatal, TEXT("OpenGL RHI does not yet support indirect draw calls."));
+		return;
+	}
 
-		check(ArgumentBufferRHI);
-		RegisterGPUWork(0);
-
-		FOpenGLContextState& ContextState = GetContextStateForCurrentContext();
-		BindPendingFramebuffer(ContextState);
-		SetPendingBlendStateForActiveRenderTargets(ContextState);
-		UpdateViewportInOpenGLContext(ContextState);
-		UpdateScissorRectInOpenGLContext(ContextState);
-		UpdateRasterizerStateInOpenGLContext(ContextState);
-		UpdateDepthStencilStateInOpenGLContext(ContextState);
-		BindPendingShaderState(ContextState);
-		CommitGraphicsResourceTables();
-		SetupTexturesForDraw(ContextState);
-		SetupUAVsForDraw(ContextState);
-		CommitNonComputeShaderConstants();
-		CachedBindElementArrayBuffer(ContextState,0);
+	VERIFY_GL_SCOPE();
 
-		// Zero-stride buffer emulation won't work here, need to use VAB with proper zero strides
-		SetupVertexArrays(ContextState, 0, PendingState.Streams, NUM_OPENGL_VERTEX_STREAMS, 1);
+	check(ArgumentBufferRHI);
+	RegisterGPUWork(0);
 
-		GLenum DrawMode = GL_TRIANGLES;
-		GLsizei NumElements = 0;
-		FindPrimitiveType(PrimitiveType, 0, DrawMode, NumElements);
+	// Zero-stride buffer emulation won't work here, need to use VAB with proper zero strides
+	SetupDraw(nullptr, 0, 1);
 
-		FOpenGLBuffer* ArgumentBuffer = ResourceCast(ArgumentBufferRHI);
+	GLenum DrawMode = GL_TRIANGLES;
+	GLsizei NumElements = 0;
+	FindPrimitiveType(PrimitiveType, 0, DrawMode, NumElements);
 
+	FOpenGLBuffer* ArgumentBuffer = ResourceCast(ArgumentBufferRHI);
+	glBindBuffer( GL_DRAW_INDIRECT_BUFFER, ArgumentBuffer->Resource);
 
-		glBindBuffer( GL_DRAW_INDIRECT_BUFFER, ArgumentBuffer->Resource);
-		{
-			CONDITIONAL_SCOPE_CYCLE_COUNTER(STAT_OpenGLShaderFirstDrawTime, PendingState.BoundShaderState->RequiresDriverInstantiation());
-			FOpenGL::DrawArraysIndirect( DrawMode, INDEX_TO_VOID(ArgumentOffset));
-		}
-		glBindBuffer( GL_DRAW_INDIRECT_BUFFER, 0);
-		GOpenGLKickHint.OnDrawCall(ContextState);
-	}
-	else
 	{
-		UE_LOG(LogRHI, Fatal,TEXT("OpenGL RHI does not yet support indirect draw calls."));
+		CONDITIONAL_SCOPE_CYCLE_COUNTER(STAT_OpenGLShaderFirstDrawTime, PendingState.BoundShaderState->RequiresDriverInstantiation());
+		FOpenGL::DrawArraysIndirect( DrawMode, INDEX_TO_VOID(ArgumentOffset));
 	}
+
+	glBindBuffer(GL_DRAW_INDIRECT_BUFFER, 0);
+
+	KickHint.OnDrawCall();
 }
 
 void FOpenGLDynamicRHI::RHIDrawIndexedIndirect(FRHIBuffer* IndexBufferRHI, FRHIBuffer* ArgumentsBufferRHI, int32 DrawArgumentsIndex, uint32 /*NumInstances*/)
 {
-	if (FOpenGL::SupportsDrawIndirect())
+	if (!FOpenGL::SupportsDrawIndirect())
 	{
-		VERIFY_GL_SCOPE();
-
-		FOpenGLBuffer* IndexBuffer = ResourceCast(IndexBufferRHI);
-		RegisterGPUWork(1);
-
-		check(ArgumentsBufferRHI);
-
-		FOpenGLContextState& ContextState = GetContextStateForCurrentContext();
-		BindPendingFramebuffer(ContextState);
-		SetPendingBlendStateForActiveRenderTargets(ContextState);
-		UpdateViewportInOpenGLContext(ContextState);
-		UpdateScissorRectInOpenGLContext(ContextState);
-		UpdateRasterizerStateInOpenGLContext(ContextState);
-		UpdateDepthStencilStateInOpenGLContext(ContextState);
-		BindPendingShaderState(ContextState);
-		CommitGraphicsResourceTables();
-		SetupTexturesForDraw(ContextState);
-		SetupUAVsForDraw(ContextState);
-		CommitNonComputeShaderConstants();
-		CachedBindElementArrayBuffer(ContextState,IndexBuffer->Resource);
-
-		// Zero-stride buffer emulation won't work here, need to use VAB with proper zero strides
-		SetupVertexArrays(ContextState, 0, PendingState.Streams, NUM_OPENGL_VERTEX_STREAMS, 1);
+		UE_LOG(LogRHI, Fatal, TEXT("OpenGL RHI does not yet support indirect draw calls."));
+		return;
+	}
 
-		GLenum DrawMode = GL_TRIANGLES;
-		GLsizei NumElements = 0;
-		FindPrimitiveType(PrimitiveType, 0, DrawMode, NumElements);
+	VERIFY_GL_SCOPE();
+	RegisterGPUWork(1);
 
-		GLenum IndexType = IndexBuffer->GetStride() == sizeof(uint32) ? GL_UNSIGNED_INT : GL_UNSIGNED_SHORT;
+	FOpenGLBuffer* IndexBuffer = ResourceCast(IndexBufferRHI);
+	SetupDraw(IndexBuffer, 0, 1);
 
-		FOpenGLBuffer* ArgumentsBuffer = ResourceCast(ArgumentsBufferRHI);
+	GLenum DrawMode = GL_TRIANGLES;
+	GLsizei NumElements = 0;
+	FindPrimitiveType(PrimitiveType, 0, DrawMode, NumElements);
 
+	GLenum IndexType = IndexBuffer->GetStride() == sizeof(uint32) ? GL_UNSIGNED_INT : GL_UNSIGNED_SHORT;
 
-		glBindBuffer( GL_DRAW_INDIRECT_BUFFER, ArgumentsBuffer->Resource);
-		{
-			CONDITIONAL_SCOPE_CYCLE_COUNTER(STAT_OpenGLShaderFirstDrawTime, PendingState.BoundShaderState->RequiresDriverInstantiation());
-	
-			// Offset is based on an index into the list of structures
-			FOpenGL::DrawElementsIndirect( DrawMode, IndexType, INDEX_TO_VOID(DrawArgumentsIndex * 5 *sizeof(uint32)));
-		}
-		glBindBuffer( GL_DRAW_INDIRECT_BUFFER, 0);
+	check(ArgumentsBufferRHI);
+	FOpenGLBuffer* ArgumentsBuffer = ResourceCast(ArgumentsBufferRHI);
+	glBindBuffer(GL_DRAW_INDIRECT_BUFFER, ArgumentsBuffer->Resource);
 
-		GOpenGLKickHint.OnDrawCall(ContextState);
-	}
-	else
 	{
-		UE_LOG(LogRHI, Fatal,TEXT("OpenGL RHI does not yet support indirect draw calls."));
+		CONDITIONAL_SCOPE_CYCLE_COUNTER(STAT_OpenGLShaderFirstDrawTime, PendingState.BoundShaderState->RequiresDriverInstantiation());
+	
+		// Offset is based on an index into the list of structures
+		FOpenGL::DrawElementsIndirect(DrawMode, IndexType, INDEX_TO_VOID(DrawArgumentsIndex * 5 * sizeof(uint32)));
 	}
+
+	glBindBuffer(GL_DRAW_INDIRECT_BUFFER, 0);
+
+	KickHint.OnDrawCall();
 }
 
 void FOpenGLDynamicRHI::RHIDrawIndexedPrimitive(FRHIBuffer* IndexBufferRHI, int32 BaseVertexIndex, uint32 FirstInstance, uint32 NumVertices, uint32 StartIndex, uint32 NumPrimitives, uint32 NumInstances)
@@ -2434,59 +2304,7 @@ void FOpenGLDynamicRHI::RHIDrawIndexedPrimitive(FRHIBuffer* IndexBufferRHI, int3
 
 	RHI_DRAW_CALL_STATS(PrimitiveType,NumPrimitives*NumInstances);
 
-	FOpenGLContextState& ContextState = GetContextStateForCurrentContext();
-	{
-		DETAILED_QUICK_SCOPE_CYCLE_COUNTER(STAT_BindPendingFramebuffer);
-		BindPendingFramebuffer(ContextState);
-	}
-	{
-		DETAILED_QUICK_SCOPE_CYCLE_COUNTER(STAT_SetPendingBlendStateForActiveRenderTargets);
-		SetPendingBlendStateForActiveRenderTargets(ContextState);
-	}
-	{
-		DETAILED_QUICK_SCOPE_CYCLE_COUNTER(STAT_UpdateViewportInOpenGLContext);
-		UpdateViewportInOpenGLContext(ContextState);
-	}
-	{
-		DETAILED_QUICK_SCOPE_CYCLE_COUNTER(STAT_UpdateScissorRectInOpenGLContext);
-		UpdateScissorRectInOpenGLContext(ContextState);
-	}
-	{
-		DETAILED_QUICK_SCOPE_CYCLE_COUNTER(STAT_UpdateRasterizerStateInOpenGLContext);
-		UpdateRasterizerStateInOpenGLContext(ContextState);
-	}
-	{
-		DETAILED_QUICK_SCOPE_CYCLE_COUNTER(STAT_UpdateDepthStencilStateInOpenGLContext);
-		UpdateDepthStencilStateInOpenGLContext(ContextState);
-	}
-	{
-		DETAILED_QUICK_SCOPE_CYCLE_COUNTER(STAT_BindPendingShaderState);
-		BindPendingShaderState(ContextState);
-	}
-	{
-		DETAILED_QUICK_SCOPE_CYCLE_COUNTER(STAT_CommitGraphicsResourceTables);
-		CommitGraphicsResourceTables();
-	}
-	{
-		DETAILED_QUICK_SCOPE_CYCLE_COUNTER(STAT_SetupTexturesForDraw);
-		SetupTexturesForDraw(ContextState);
-	}
-	{
-		DETAILED_QUICK_SCOPE_CYCLE_COUNTER(STAT_SetupUAVsForDraw);
-		SetupUAVsForDraw(ContextState);
-	}
-	{
-		DETAILED_QUICK_SCOPE_CYCLE_COUNTER(STAT_CommitNonComputeShaderConstants);
-		CommitNonComputeShaderConstants();
-	}
-	{
-		DETAILED_QUICK_SCOPE_CYCLE_COUNTER(STAT_CachedBindElementArrayBuffer);
-		CachedBindElementArrayBuffer(ContextState, IndexBuffer->Resource);
-	}
-	{
-		DETAILED_QUICK_SCOPE_CYCLE_COUNTER(STAT_SetupVertexArrays);
-		SetupVertexArrays(ContextState, BaseVertexIndex, PendingState.Streams, NUM_OPENGL_VERTEX_STREAMS, NumVertices + StartIndex);
-	}
+	SetupDraw(IndexBuffer, BaseVertexIndex, NumVertices + StartIndex);
 
 	GLenum DrawMode = GL_TRIANGLES;
 	GLsizei NumElements = 0;
@@ -2507,7 +2325,7 @@ void FOpenGLDynamicRHI::RHIDrawIndexedPrimitive(FRHIBuffer* IndexBufferRHI, int3
 	{
 		SCOPE_CYCLE_COUNTER_DETAILED(STAT_OpenGLDrawPrimitiveDriverTime);
 		CONDITIONAL_SCOPE_CYCLE_COUNTER(STAT_OpenGLShaderFirstDrawTime, PendingState.BoundShaderState->RequiresDriverInstantiation());
-		if ( FOpenGL::SupportsDrawIndexOffset() )
+		if (FOpenGL::SupportsDrawIndexOffset())
 		{
 			FOpenGL::DrawRangeElements(DrawMode, 0, NumVertices, NumElements, IndexType, INDEX_TO_VOID(StartIndex));
 		}
@@ -2516,60 +2334,103 @@ void FOpenGLDynamicRHI::RHIDrawIndexedPrimitive(FRHIBuffer* IndexBufferRHI, int3
 			glDrawElements(DrawMode, NumElements, IndexType, INDEX_TO_VOID(StartIndex));
 		}
 	}
-	GOpenGLKickHint.OnDrawCall(ContextState);
+	KickHint.OnDrawCall();
 }
 
 void FOpenGLDynamicRHI::RHIDrawIndexedPrimitiveIndirect(FRHIBuffer* IndexBufferRHI, FRHIBuffer* ArgumentBufferRHI, uint32 ArgumentOffset)
 {
-	if (FOpenGL::SupportsDrawIndirect())
+	if (!FOpenGL::SupportsDrawIndirect())
 	{
-		VERIFY_GL_SCOPE();
-
-		FOpenGLBuffer* IndexBuffer = ResourceCast(IndexBufferRHI);
-		RegisterGPUWork(1);
+		UE_LOG(LogRHI, Fatal, TEXT("OpenGL RHI does not yet support indirect draw calls."));
+		return;
+	}
 
-		check(ArgumentBufferRHI);
+	VERIFY_GL_SCOPE();
 
-		FOpenGLContextState& ContextState = GetContextStateForCurrentContext();
-		BindPendingFramebuffer(ContextState);
-		SetPendingBlendStateForActiveRenderTargets(ContextState);
-		UpdateViewportInOpenGLContext(ContextState);
-		UpdateScissorRectInOpenGLContext(ContextState);
-		UpdateRasterizerStateInOpenGLContext(ContextState);
-		UpdateDepthStencilStateInOpenGLContext(ContextState);
-		BindPendingShaderState(ContextState);
-		CommitGraphicsResourceTables();
-		SetupTexturesForDraw(ContextState);
-		SetupUAVsForDraw(ContextState);
-		CommitNonComputeShaderConstants();
-		CachedBindElementArrayBuffer(ContextState,IndexBuffer->Resource);
+	FOpenGLBuffer* IndexBuffer = ResourceCast(IndexBufferRHI);
+	RegisterGPUWork(1);
 
-		// @ToDo Zero-stride buffer emulation won't work here, need to use VAB with proper zero strides
-		SetupVertexArrays(ContextState, 0, PendingState.Streams, NUM_OPENGL_VERTEX_STREAMS, 1);
+	check(ArgumentBufferRHI);
 
-		GLenum DrawMode = GL_TRIANGLES;
-		GLsizei NumElements = 0;
-		FindPrimitiveType(PrimitiveType, 0, DrawMode, NumElements);
+	// @ToDo Zero-stride buffer emulation won't work here, need to use VAB with proper zero strides
+	SetupDraw(IndexBuffer, 0, 1);
 
-		GLenum IndexType = IndexBuffer->GetStride() == sizeof(uint32) ? GL_UNSIGNED_INT : GL_UNSIGNED_SHORT;
+	GLenum DrawMode = GL_TRIANGLES;
+	GLsizei NumElements = 0;
+	FindPrimitiveType(PrimitiveType, 0, DrawMode, NumElements);
 
-		FOpenGLBuffer* ArgumentBuffer = ResourceCast(ArgumentBufferRHI);
+	GLenum IndexType = IndexBuffer->GetStride() == sizeof(uint32)	
+		? GL_UNSIGNED_INT
+		: GL_UNSIGNED_SHORT;
 
+	FOpenGLBuffer* ArgumentBuffer = ResourceCast(ArgumentBufferRHI);
+	glBindBuffer(GL_DRAW_INDIRECT_BUFFER, ArgumentBuffer->Resource);
 
-		glBindBuffer( GL_DRAW_INDIRECT_BUFFER, ArgumentBuffer->Resource);
-		{
-			CONDITIONAL_SCOPE_CYCLE_COUNTER(STAT_OpenGLShaderFirstDrawTime, PendingState.BoundShaderState->RequiresDriverInstantiation());
+	{
+		CONDITIONAL_SCOPE_CYCLE_COUNTER(STAT_OpenGLShaderFirstDrawTime, PendingState.BoundShaderState->RequiresDriverInstantiation());
 		
-			// Offset is based on an index into the list of structures
-			FOpenGL::DrawElementsIndirect( DrawMode, IndexType, INDEX_TO_VOID(ArgumentOffset));
-		}
-		glBindBuffer( GL_DRAW_INDIRECT_BUFFER, 0);
+		// Offset is based on an index into the list of structures
+		FOpenGL::DrawElementsIndirect( DrawMode, IndexType, INDEX_TO_VOID(ArgumentOffset));
+	}
+
+	glBindBuffer(GL_DRAW_INDIRECT_BUFFER, 0);
 
-		GOpenGLKickHint.OnDrawCall(ContextState);
+	KickHint.OnDrawCall();
+}
+
+void FOpenGLDynamicRHI::SetupDraw(FOpenGLBuffer* IndexBuffer, uint32 BaseVertexIndex, uint32 MaxVertices)
+{
+	{
+		DETAILED_QUICK_SCOPE_CYCLE_COUNTER(STAT_BindPendingFramebuffer);
+		BindPendingFramebuffer();
+	}
+	{
+		DETAILED_QUICK_SCOPE_CYCLE_COUNTER(STAT_SetPendingBlendStateForActiveRenderTargets);
+		SetPendingBlendStateForActiveRenderTargets();
+	}
+	{
+		DETAILED_QUICK_SCOPE_CYCLE_COUNTER(STAT_UpdateViewportInOpenGLContext);
+		UpdateViewportInOpenGLContext();
 	}
-	else
 	{
-		UE_LOG(LogRHI, Fatal,TEXT("OpenGL RHI does not yet support indirect draw calls."));
+		DETAILED_QUICK_SCOPE_CYCLE_COUNTER(STAT_UpdateScissorRectInOpenGLContext);
+		UpdateScissorRectInOpenGLContext();
+	}
+	{
+		DETAILED_QUICK_SCOPE_CYCLE_COUNTER(STAT_UpdateRasterizerStateInOpenGLContext);
+		UpdateRasterizerStateInOpenGLContext();
+	}
+	{
+		DETAILED_QUICK_SCOPE_CYCLE_COUNTER(STAT_UpdateDepthStencilStateInOpenGLContext);
+		UpdateDepthStencilStateInOpenGLContext();
+	}
+	{
+		DETAILED_QUICK_SCOPE_CYCLE_COUNTER(STAT_BindPendingShaderState);
+		BindPendingShaderState();
+	}
+	{
+		DETAILED_QUICK_SCOPE_CYCLE_COUNTER(STAT_CommitGraphicsResourceTables);
+		CommitGraphicsResourceTables();
+	}
+	{
+		DETAILED_QUICK_SCOPE_CYCLE_COUNTER(STAT_SetupTexturesForDraw);
+		SetupTexturesForDraw();
+	}
+	{
+		DETAILED_QUICK_SCOPE_CYCLE_COUNTER(STAT_SetupUAVsForDraw);
+		SetupUAVsForDraw();
+	}
+	{
+		DETAILED_QUICK_SCOPE_CYCLE_COUNTER(STAT_CommitNonComputeShaderConstants);
+		CommitNonComputeShaderConstants();
+	}
+	{
+		DETAILED_QUICK_SCOPE_CYCLE_COUNTER(STAT_CachedBindBuffer);
+		CachedBindBuffer(GL_ELEMENT_ARRAY_BUFFER, IndexBuffer ? IndexBuffer->Resource : 0);
+	}
+	{
+		DETAILED_QUICK_SCOPE_CYCLE_COUNTER(STAT_SetupVertexArrays);
+		SetupVertexArrays(BaseVertexIndex, PendingState.Streams, NUM_OPENGL_VERTEX_STREAMS, MaxVertices);
 	}
 }
 
@@ -2595,7 +2456,7 @@ static inline void ClearCurrentDepthStencilWithCurrentScissor( int8 ClearType, f
 	}
 }
 
-void FOpenGLDynamicRHI::ClearCurrentFramebufferWithCurrentScissor(FOpenGLContextState& ContextState, int8 ClearType, int32 NumClearColors, const bool* bClearColorArray, const FLinearColor* ClearColorArray, float Depth, uint32 Stencil)
+void FOpenGLDynamicRHI::ClearCurrentFramebufferWithCurrentScissor(int8 ClearType, int32 NumClearColors, const bool* bClearColorArray, const FLinearColor* ClearColorArray, float Depth, uint32 Stencil)
 {
 	VERIFY_GL_SCOPE();
 		
@@ -2650,8 +2511,8 @@ void FOpenGLDynamicRHI::RHIClearMRT(const bool* bClearColorArray,int32 NumClearC
 
 	bool bScissorChanged = false;
 	RegisterGPUWork(0);
-	FOpenGLContextState& ContextState = GetContextStateForCurrentContext();
-	BindPendingFramebuffer(ContextState);
+
+	BindPendingFramebuffer();
 
 	if (bPrevScissorEnabled || PendingState.Viewport.Min.X != 0 || PendingState.Viewport.Min.Y != 0 || PendingState.Viewport.Max.X != PendingState.RenderTargetWidth || PendingState.Viewport.Max.Y != PendingState.RenderTargetHeight)
 	{
@@ -2660,7 +2521,7 @@ void FOpenGLDynamicRHI::RHIClearMRT(const bool* bClearColorArray,int32 NumClearC
 	}
 
 	// Always update in case there are uncommitted changes to disable scissor
-	UpdateScissorRectInOpenGLContext(ContextState);
+	UpdateScissorRectInOpenGLContext();
 
 	int8 ClearType = CT_None;
 
@@ -2713,7 +2574,7 @@ void FOpenGLDynamicRHI::RHIClearMRT(const bool* bClearColorArray,int32 NumClearC
 	}
 
 	// Just one clear
-	ClearCurrentFramebufferWithCurrentScissor(ContextState, ClearType, NumClearColors, bClearColorArray, ClearColorArray, Depth, Stencil);
+	ClearCurrentFramebufferWithCurrentScissor(ClearType, NumClearColors, bClearColorArray, ClearColorArray, Depth, Stencil);
 
 	if (bScissorChanged)
 	{
@@ -2815,6 +2676,19 @@ void FOpenGLDynamicRHI::RHISetComputeShader(FRHIComputeShader* ComputeShaderRHI)
 	ApplyStaticUniformBuffers(ComputeShaderRHI);
 }
 
+void FOpenGLDynamicRHI::SetupDispatch()
+{
+	FOpenGLComputeShader* ComputeShader = PendingState.CurrentComputeShader;
+	check(ComputeShader);
+	LinkComputeProgram(ComputeShader);
+
+	BindPendingComputeShaderState(ComputeShader);
+	CommitComputeResourceTables(ComputeShader);
+	SetupTexturesForDraw(ComputeShader, FOpenGL::GetMaxComputeTextureImageUnits());
+	SetupUAVsForCompute(ComputeShader);
+	CommitComputeShaderConstants(ComputeShader);
+}
+
 void FOpenGLDynamicRHI::RHIDispatchComputeShader(uint32 ThreadGroupCountX, uint32 ThreadGroupCountY, uint32 ThreadGroupCountZ)
 { 
 	if (OpenGLConsoleVariables::bSkipCompute)
@@ -2823,21 +2697,10 @@ void FOpenGLDynamicRHI::RHIDispatchComputeShader(uint32 ThreadGroupCountX, uint3
 	}
 
 	VERIFY_GL_SCOPE();
-		
-	FOpenGLComputeShader* ComputeShader = PendingState.CurrentComputeShader;
-	check(ComputeShader);
-	LinkComputeProgram(ComputeShader);
-
-	FOpenGLContextState& ContextState = GetContextStateForCurrentContext();
-
-	RegisterGPUDispatch(FIntVector(ThreadGroupCountX, ThreadGroupCountY, ThreadGroupCountZ));	
-
-	BindPendingComputeShaderState(ContextState, ComputeShader);
-	CommitComputeResourceTables(ComputeShader);
-	SetupTexturesForDraw(ContextState, ComputeShader, FOpenGL::GetMaxComputeTextureImageUnits());
-	SetupUAVsForCompute(ContextState, ComputeShader);
-	CommitComputeShaderConstants(ComputeShader);
+	RegisterGPUDispatch(FIntVector(ThreadGroupCountX, ThreadGroupCountY, ThreadGroupCountZ));
 	
+	SetupDispatch();
+
 	FOpenGL::MemoryBarrier(GL_ALL_BARRIER_BITS);
 	FOpenGL::DispatchCompute(ThreadGroupCountX, ThreadGroupCountY, ThreadGroupCountZ);
 	FOpenGL::MemoryBarrier(GL_ALL_BARRIER_BITS);
@@ -2846,27 +2709,11 @@ void FOpenGLDynamicRHI::RHIDispatchComputeShader(uint32 ThreadGroupCountX, uint3
 void FOpenGLDynamicRHI::RHIDispatchIndirectComputeShader(FRHIBuffer* ArgumentBufferRHI, uint32 ArgumentOffset)
 {
 	VERIFY_GL_SCOPE();
+	RegisterGPUDispatch(FIntVector(1, 1, 1));
 
-	FOpenGLComputeShader* ComputeShader = PendingState.CurrentComputeShader;
-	check(ComputeShader);
-	LinkComputeProgram(ComputeShader);
-
-	FOpenGLBuffer* ArgumentBuffer = ResourceCast(ArgumentBufferRHI);
-
-	FOpenGLContextState& ContextState = GetContextStateForCurrentContext();
-
-	RegisterGPUDispatch(FIntVector(1, 1, 1));	
-
-	BindPendingComputeShaderState(ContextState, ComputeShader);
-
-	CommitComputeResourceTables(ComputeShader);
-
-	SetupTexturesForDraw(ContextState, ComputeShader, FOpenGL::GetMaxComputeTextureImageUnits());
-
-	SetupUAVsForCompute(ContextState, ComputeShader);
-	
-	CommitComputeShaderConstants(ComputeShader);
+	SetupDispatch();
 	
+	FOpenGLBuffer* ArgumentBuffer = ResourceCast(ArgumentBufferRHI);
 	glBindBuffer( GL_DISPATCH_INDIRECT_BUFFER, ArgumentBuffer->Resource);
 	
 	FOpenGL::MemoryBarrier(GL_ALL_BARRIER_BITS);
@@ -2878,7 +2725,7 @@ void FOpenGLDynamicRHI::RHIDispatchIndirectComputeShader(FRHIBuffer* ArgumentBuf
 
 void FOpenGLDynamicRHI::RHISetMultipleViewports(uint32 Count, const FViewportBounds* Data)
 {
-	UE_LOG(LogRHI, Fatal,TEXT("OpenGL Render path does not support multiple Viewports!"));
+	UE_LOG(LogRHI, Fatal, TEXT("OpenGL Render path does not support multiple Viewports!"));
 }
 
 void FOpenGLDynamicRHI::RHIEnableDepthBoundsTest(bool bEnable)
@@ -2974,34 +2821,71 @@ void FOpenGLDynamicRHI::RHICopyToStagingBuffer(FRHIBuffer* SourceBufferRHI, FRHI
 
 void FOpenGLDynamicRHI::RHIPostExternalCommandsReset()
 {
-	auto &RCS = RenderingContextState;
-	auto &SCS = SharedContextState;
-	if((int)RCS.Program >= 0) FOpenGL::BindProgramPipeline(RCS.Program);
-	glViewport(RCS.Viewport.Min.X, RCS.Viewport.Min.Y, RCS.Viewport.Max.X - RCS.Viewport.Min.X, RCS.Viewport.Max.Y - RCS.Viewport.Min.Y);
-	FOpenGL::DepthRange(RCS.DepthMinZ, RCS.DepthMaxZ);
-	RCS.bScissorEnabled ? glEnable(GL_SCISSOR_TEST) : glDisable(GL_SCISSOR_TEST);
-	glScissor(RCS.Scissor.Min.X, RCS.Scissor.Min.Y, RCS.Scissor.Max.X - RCS.Scissor.Min.X, RCS.Scissor.Max.Y - RCS.Scissor.Min.Y);
-	if((int)RCS.Framebuffer >= 0) glBindFramebuffer(GL_FRAMEBUFFER, RCS.Framebuffer);
-	if((int)RCS.ArrayBufferBound >= 0) glBindBuffer(GL_ARRAY_BUFFER, RCS.ArrayBufferBound);
-	if((int)RCS.ElementArrayBufferBound >= 0) glBindBuffer(GL_ELEMENT_ARRAY_BUFFER, RCS.ElementArrayBufferBound);
-	glBindBuffer(GL_PIXEL_UNPACK_BUFFER, RCS.PixelUnpackBufferBound);
-	glBindBuffer(GL_UNIFORM_BUFFER, RCS.UniformBufferBound);
-	SCS.BlendState = RCS.BlendState = {};
+	if ((int)ContextState.Program >= 0)
+	{
+		FOpenGL::BindProgramPipeline(ContextState.Program);
+	}
+
+	glViewport(
+		ContextState.Viewport.Min.X,
+		ContextState.Viewport.Min.Y,
+		ContextState.Viewport.Max.X - ContextState.Viewport.Min.X,
+		ContextState.Viewport.Max.Y - ContextState.Viewport.Min.Y
+	);
+
+	FOpenGL::DepthRange(ContextState.DepthMinZ, ContextState.DepthMaxZ);
+
+	if (ContextState.bScissorEnabled)
+	{
+		glEnable(GL_SCISSOR_TEST);
+	}
+	else
+	{
+		glDisable(GL_SCISSOR_TEST);
+	}
+
+	glScissor(
+		ContextState.Scissor.Min.X,
+		ContextState.Scissor.Min.Y,
+		ContextState.Scissor.Max.X - ContextState.Scissor.Min.X,
+		ContextState.Scissor.Max.Y - ContextState.Scissor.Min.Y
+	);
+
+	if ((int)ContextState.Framebuffer >= 0)
+	{
+		glBindFramebuffer(GL_FRAMEBUFFER, ContextState.Framebuffer);
+	}
+	if ((int)ContextState.ArrayBufferBound >= 0)
+	{
+		glBindBuffer(GL_ARRAY_BUFFER, ContextState.ArrayBufferBound);
+	}
+	if ((int)ContextState.ElementArrayBufferBound >= 0)
+	{
+		glBindBuffer(GL_ELEMENT_ARRAY_BUFFER, ContextState.ElementArrayBufferBound);
+	}
+
+	glBindBuffer(GL_PIXEL_UNPACK_BUFFER, ContextState.PixelUnpackBufferBound);
+	glBindBuffer(GL_UNIFORM_BUFFER, ContextState.UniformBufferBound);
+
+	ContextState.BlendState = {};
+
 	glActiveTexture(GL_TEXTURE0);
-	SCS.ActiveTexture = RCS.ActiveTexture = 0;
-	SCS.RasterizerState = RCS.RasterizerState = {};
-	UpdateRasterizerStateInOpenGLContext(RCS);
-	SCS.ActiveStreamMask = RCS.ActiveStreamMask = 0;
+	ContextState.ActiveTexture = 0;
+	ContextState.RasterizerState = {};
+
+	UpdateRasterizerStateInOpenGLContext();
+	ContextState.ActiveStreamMask = 0;
+
 	for (GLuint UniformBufferIndex = 0; UniformBufferIndex < CrossCompiler::NUM_SHADER_STAGES * OGL_MAX_UNIFORM_BUFFER_BINDINGS; UniformBufferIndex++)
 	{
-		SCS.UniformBuffers[UniformBufferIndex] = FOpenGLCachedUniformBuffer_Invalid;	// that'll enforce state update on next cache test
-		RCS.UniformBuffers[UniformBufferIndex] = FOpenGLCachedUniformBuffer_Invalid;	// that'll enforce state update on next cache test
+		ContextState.UniformBuffers[UniformBufferIndex] = FOpenGLCachedUniformBuffer_Invalid;	// that'll enforce state update on next cache test
 	}
+
 	for (uint32 Index = 0; Index < NUM_OPENGL_VERTEX_STREAMS; Index++)
 	{
-		if (RCS.VertexStreams[Index].VertexBufferResource > 0)
+		if (ContextState.VertexStreams[Index].VertexBufferResource > 0)
 		{
-			FOpenGL::BindVertexBuffer(Index, RCS.VertexStreams[Index].VertexBufferResource, RCS.VertexStreams[Index].Offset, RCS.VertexStreams[Index].Stride);
+			FOpenGL::BindVertexBuffer(Index, ContextState.VertexStreams[Index].VertexBufferResource, ContextState.VertexStreams[Index].Offset, ContextState.VertexStreams[Index].Stride);
 		}
 		else
 		{
@@ -3011,10 +2895,10 @@ void FOpenGLDynamicRHI::RHIPostExternalCommandsReset()
 
 	for (uint32 Index = 0; Index < NUM_OPENGL_VERTEX_STREAMS; Index++)
 	{
-		if(RCS.GetVertexAttrEnabled(Index)) 
+		if (ContextState.GetVertexAttrEnabled(Index))
 		{
-			FOpenGLCachedAttr& Attr = RCS.VertexAttrs[Index];
-			if (Attr.StreamIndex < NUM_OPENGL_VERTEX_STREAMS && RCS.VertexStreams[Attr.StreamIndex].VertexBufferResource > 0)
+			FOpenGLCachedAttr& Attr = ContextState.VertexAttrs[Index];
+			if (Attr.StreamIndex < NUM_OPENGL_VERTEX_STREAMS && ContextState.VertexStreams[Attr.StreamIndex].VertexBufferResource > 0)
 			{
 				if (!Attr.bShouldConvertToFloat)
 				{
diff --git a/Engine/Source/Runtime/OpenGLDrv/Private/OpenGLDevice.cpp b/Engine/Source/Runtime/OpenGLDrv/Private/OpenGLDevice.cpp
--- a/Engine/Source/Runtime/OpenGLDrv/Private/OpenGLDevice.cpp
+++ b/Engine/Source/Runtime/OpenGLDrv/Private/OpenGLDevice.cpp
@@ -87,124 +87,16 @@ static TAutoConsoleVariable<bool> CVarAllowPSOPrecaching(
 	TEXT("false: GL RHI will disable precaching (even if r.PSOPrecaching=1). "),
 	ECVF_RenderThreadSafe | ECVF_ReadOnly);
 
-void OnProgramDeletion( GLint ProgramResource )
-{
-	check(PrivateOpenGLDevicePtr);
-	PrivateOpenGLDevicePtr->OnProgramDeletion( ProgramResource );
-}
-
-void OnBufferDeletion( GLuint BufferResource )
-{
-	check(PrivateOpenGLDevicePtr);
-	PrivateOpenGLDevicePtr->OnBufferDeletion( BufferResource );
-}
-
-void OnPixelBufferDeletion( GLuint PixelBufferResource )
-{
-	check(PrivateOpenGLDevicePtr);
-	PrivateOpenGLDevicePtr->OnPixelBufferDeletion( PixelBufferResource );
-}
-
-void OnUniformBufferDeletion( GLuint UniformBufferResource, uint32 AllocatedSize, bool bStreamDraw, uint32 , uint8*  )
-{
-	check(PrivateOpenGLDevicePtr);
-	PrivateOpenGLDevicePtr->OnUniformBufferDeletion( UniformBufferResource, AllocatedSize, bStreamDraw );
-}
-
-void CachedBindBuffer( GLenum Type, GLuint Buffer )
-{
-	check(PrivateOpenGLDevicePtr);
-	if (Type == GL_ARRAY_BUFFER)
-	{
-		PrivateOpenGLDevicePtr->CachedBindArrayBuffer(PrivateOpenGLDevicePtr->GetContextStateForCurrentContext(), Buffer);
-	}
-	else if (Type == GL_ELEMENT_ARRAY_BUFFER)
-	{
-		PrivateOpenGLDevicePtr->CachedBindElementArrayBuffer(PrivateOpenGLDevicePtr->GetContextStateForCurrentContext(), Buffer);
-	}
-	else if (Type == GL_SHADER_STORAGE_BUFFER)
-	{
-		PrivateOpenGLDevicePtr->CachedBindStorageBuffer(PrivateOpenGLDevicePtr->GetContextStateForCurrentContext(), Buffer);
-	}
-	else
-	{
-		checkNoEntry();
-	}
-}
-
-void CachedBindPixelUnpackBuffer( GLenum Type, GLuint Buffer )
-{
-	check(PrivateOpenGLDevicePtr);
-	PrivateOpenGLDevicePtr->CachedBindPixelUnpackBuffer(PrivateOpenGLDevicePtr->GetContextStateForCurrentContext(),Buffer);
-}
-
-void CachedBindUniformBuffer( GLuint Buffer )
-{
-	check(PrivateOpenGLDevicePtr);
-	if (FOpenGL::SupportsUniformBuffers())
-	{
-		PrivateOpenGLDevicePtr->CachedBindUniformBuffer(PrivateOpenGLDevicePtr->GetContextStateForCurrentContext(),Buffer);
-	}
-}
-
-bool IsUniformBufferBound( GLuint Buffer )
-{
-	check(PrivateOpenGLDevicePtr);
-	return PrivateOpenGLDevicePtr->IsUniformBufferBound(PrivateOpenGLDevicePtr->GetContextStateForCurrentContext(),Buffer);
-}
-
 extern void BeginFrame_UniformBufferPoolCleanup();
 extern void BeginFrame_VertexBufferCleanup();
 
-EOpenGLCurrentContext FOpenGLDynamicRHI::GetCurrentContext()
-{
-	return PlatformOpenGLCurrentContext(FOpenGLDynamicRHI::Get().PlatformDevice);
-}
-
-FOpenGLContextState& FOpenGLDynamicRHI::GetContextStateForCurrentContext()
-{
-	if (!CachedContextState)
-	{
-		int32 ContextType = (int32)PlatformOpenGLCurrentContext(PlatformDevice);
-		checkf(ContextType >= 0, TEXT("Invalid GL context on current thread."));
-
-		if (ContextType == CONTEXT_Rendering)
-		{
-			CachedContextState = &RenderingContextState;
-		}
-		else
-		{
-			CachedContextState = &SharedContextState;
-		}
-	}
-	else
-	{
-	#if DO_CHECK
-		switch ((int32)PlatformOpenGLCurrentContext(PlatformDevice))
-		{
-		default:
-		case CONTEXT_Other:
-		case CONTEXT_Invalid:
-			checkf(false, TEXT("Invalid GL context on current thread."));
-			break;
-
-		case CONTEXT_Rendering: check(CachedContextState == &RenderingContextState); break;
-		case CONTEXT_Shared   : check(CachedContextState == &SharedContextState   ); break;
-		}
-	#endif // DO_CHECK
-	}
-
-	return *CachedContextState;
-}
-
 void FOpenGLDynamicRHI::RHIEndFrame(const FRHIEndFrameArgs& Args)
 {
 #if (RHI_NEW_GPU_PROFILER == 0)
 	GPUProfilingData->EndFrame();
 #endif
 
-	extern void OpenGLCommands_OnEndFrame();
-	OpenGLCommands_OnEndFrame();
+	KickHint.Reset();
 
 #if PLATFORM_ANDROID //adding #if since not sure if this is required for any other platform.
 	PendingState.DepthStencil = 0;
@@ -222,8 +114,13 @@ void FOpenGLDynamicRHI::RHIEndFrame(const FRHIEndFrameArgs& Args)
 
 JNI_METHOD void Java_com_epicgames_unreal_MediaPlayer14_nativeClearCachedAttributeState(JNIEnv* jenv, jobject thiz, jint PositionAttrib, jint TexCoordsAttrib)
 {
-	FOpenGLContextState& ContextState = PrivateOpenGLDevicePtr->GetContextStateForCurrentContext();
+	FOpenGLDynamicRHI::Get().ClearCachedAttributeState(PositionAttrib, TexCoordsAttrib);
+}
 
+#endif
+
+void FOpenGLDynamicRHI::ClearCachedAttributeState(int32 PositionAttrib, int32 TexCoordsAttrib)
+{
 	// update vertex attributes state
 	ContextState.SetVertexAttrEnabled(PositionAttrib, false);
 	ContextState.VertexAttrs[PositionAttrib].Size = -1; // will force attribute update
@@ -240,8 +137,6 @@ JNI_METHOD void Java_com_epicgames_unreal_MediaPlayer14_nativeClearCachedAttribu
 	FOpenGL::BindProgramPipeline(ContextState.Program);
 }
 
-#endif
-
 bool GDisableOpenGLDebugOutput = false;
 
 #if defined(GL_ARB_debug_output) || defined(GL_KHR_debug)
@@ -1274,8 +1169,6 @@ FOpenGLDynamicRHI::FOpenGLDynamicRHI()
 	VERIFY_GL_SCOPE();
 	InitRHICapabilitiesForGL();
 
-	check(PlatformOpenGLCurrentContext(PlatformDevice) == CONTEXT_Shared);
-
 	if (PlatformCanEnableGPUCapture())
 	{
 		EnableIdealGPUCaptureOptions(true);
@@ -1500,8 +1393,7 @@ void FOpenGLDynamicRHI::Cleanup()
 
 	VERIFY_GL_SCOPE();
 	PendingState.CleanupResources();
-	SharedContextState.CleanupResources();
-	RenderingContextState.CleanupResources();
+	ContextState.CleanupResources();
 }
 
 void FOpenGLDynamicRHI::RHIFlushResources()
@@ -1511,21 +1403,14 @@ void FOpenGLDynamicRHI::RHIFlushResources()
 
 void FOpenGLDynamicRHI::RHIAcquireThreadOwnership()
 {
-	check(!bRevertToSharedContextAfterDrawingViewport);	// if this is true, then main thread is rendering using our context right now.
 	PlatformRenderingContextSetup(PlatformDevice);
-	CachedContextState = nullptr;
-
-	bIsRenderingContextAcquired = true;
 	VERIFY_GL(RHIAcquireThreadOwnership);
 }
 
 void FOpenGLDynamicRHI::RHIReleaseThreadOwnership()
 {
 	VERIFY_GL(RHIReleaseThreadOwnership);
-	bIsRenderingContextAcquired = false;
-
 	PlatformNULLContextSetup();
-	CachedContextState = nullptr;
 }
 
 void* FOpenGLDynamicRHI::RHIGetNativeDevice()
diff --git a/Engine/Source/Runtime/OpenGLDrv/Private/OpenGLDrv.cpp b/Engine/Source/Runtime/OpenGLDrv/Private/OpenGLDrv.cpp
--- a/Engine/Source/Runtime/OpenGLDrv/Private/OpenGLDrv.cpp
+++ b/Engine/Source/Runtime/OpenGLDrv/Private/OpenGLDrv.cpp
@@ -100,7 +100,7 @@ bool FOpenGLDynamicRHI::RHIEGLSupportsNoErrorContext() const
 void FOpenGLDynamicRHI::RHIInitEGLInstanceGLES2()
 {
 	AndroidEGL::GetInstance()->Init(AndroidEGL::AV_OpenGLES, 2, 0);
-	AndroidEGL::GetInstance()->InitSurface(false, false);
+	AndroidEGL::GetInstance()->InitRenderSurface(false, false);
 }
 
 void FOpenGLDynamicRHI::RHIInitEGLBackBuffer()
@@ -529,8 +529,7 @@ float FOpenGLEventNode::GetTiming()
 void FOpenGLDynamicRHI::InitializeStateResources()
 {
 	VERIFY_GL_SCOPE();
-	SharedContextState.InitializeResources(FOpenGL::GetMaxCombinedTextureImageUnits(), FOpenGL::GetMaxCombinedUAVUnits());
-	RenderingContextState.InitializeResources(FOpenGL::GetMaxCombinedTextureImageUnits(), FOpenGL::GetMaxCombinedUAVUnits());
+	ContextState.InitializeResources(FOpenGL::GetMaxCombinedTextureImageUnits(), FOpenGL::GetMaxCombinedUAVUnits());
 	PendingState.InitializeResources(FOpenGL::GetMaxCombinedTextureImageUnits(), FOpenGL::GetMaxCombinedUAVUnits());
 }
 
diff --git a/Engine/Source/Runtime/OpenGLDrv/Private/OpenGLDrv.h b/Engine/Source/Runtime/OpenGLDrv/Private/OpenGLDrv.h
--- a/Engine/Source/Runtime/OpenGLDrv/Private/OpenGLDrv.h
+++ b/Engine/Source/Runtime/OpenGLDrv/Private/OpenGLDrv.h
@@ -30,12 +30,6 @@
 
 #define FOpenGLCachedUniformBuffer_Invalid 0xFFFFFFFF
 
-class FOpenGLDynamicRHI;
-class FResourceBulkDataInterface;
-struct FOpenGLResourceBinder;
-struct Rect;
-enum EOpenGLCurrentContext : int8;
-
 template<class T> struct TOpenGLResourceTraits;
 
 #if (RHI_NEW_GPU_PROFILER == 0)
@@ -263,8 +257,6 @@ public:
 
 	friend class FOpenGLViewport;
 
-	static EOpenGLCurrentContext GetCurrentContext();
-
 	/** Initialization constructor. */
 	FOpenGLDynamicRHI();
 
@@ -482,14 +474,13 @@ public:
 	void Cleanup();
 
 	void PurgeFramebufferFromCaches(GLuint Framebuffer);
-	void OnBufferDeletion(GLuint VertexBufferResource);
-	void OnPixelBufferDeletion(GLuint PixelBufferResource);
-	void OnUniformBufferDeletion(GLuint UniformBufferResource,uint32 AllocatedSize,bool bStreamDraw);
+	void OnBufferDeletion(GLenum Type, GLuint BufferResource);
 	void OnProgramDeletion(GLint ProgramResource);
 	void InvalidateTextureResourceInCache(GLuint Resource);
 	void InvalidateUAVResourceInCache(GLuint Resource);
+
 	/** Set a resource on texture target of a specific real OpenGL stage. Goes through cache to eliminate redundant calls. */
-	FORCEINLINE void CachedSetupTextureStage(FOpenGLContextState& ContextState, GLint TextureIndex, GLenum Target, GLuint Resource, GLint BaseMip, GLint NumMips)
+	FORCEINLINE void CachedSetupTextureStage(GLint TextureIndex, GLenum Target, GLuint Resource, GLint BaseMip, GLint NumMips)
 	{
 		FTextureStage& TextureState = ContextState.Textures[TextureIndex];
 		const bool bSameTarget = (TextureState.Target == Target);
@@ -500,69 +491,67 @@ public:
 			// Nothing changed, no need to update
 			return;
 		}
-		CachedSetupTextureStageInner(ContextState, TextureIndex, Target, Resource, BaseMip, NumMips);
+
+		CachedSetupTextureStageInner(TextureIndex, Target, Resource, BaseMip, NumMips);
 	}
 
-	void CachedSetupTextureStageInner(FOpenGLContextState& ContextState, GLint TextureIndex, GLenum Target, GLuint Resource, GLint BaseMip, GLint NumMips);
-	void CachedSetupUAVStage(FOpenGLContextState& ContextState, GLint UAVIndex, GLenum Format, GLuint Resource, bool bLayered, GLint Layer, GLenum Access, GLint Level);
+	void CachedSetupTextureStageInner(GLint TextureIndex, GLenum Target, GLuint Resource, GLint BaseMip, GLint NumMips);
+	void CachedSetupUAVStage(GLint UAVIndex, GLenum Format, GLuint Resource, bool bLayered, GLint Layer, GLenum Access, GLint Level);
 	void UpdateSRV(FOpenGLShaderResourceView* SRV);
-	FOpenGLContextState& GetContextStateForCurrentContext();
 
-	FORCEINLINE void CachedBindArrayBuffer( FOpenGLContextState& ContextState, GLuint Buffer )
+	void CachedBindUniformBuffer(GLuint Buffer)
 	{
-		VERIFY_GL_SCOPE();
-		if( ContextState.ArrayBufferBound != Buffer )
-		{
-			glBindBuffer( GL_ARRAY_BUFFER, Buffer );
-			ContextState.ArrayBufferBound = Buffer;
-		}
-	}
+		check(FOpenGL::SupportsUniformBuffers());
 
-	void CachedBindElementArrayBuffer( FOpenGLContextState& ContextState, GLuint Buffer )
-	{
 		VERIFY_GL_SCOPE();
-		if( ContextState.ElementArrayBufferBound != Buffer )
+		if (ContextState.UniformBufferBound != Buffer)
 		{
-			glBindBuffer( GL_ELEMENT_ARRAY_BUFFER, Buffer );
-			ContextState.ElementArrayBufferBound = Buffer;
+			glBindBuffer(GL_UNIFORM_BUFFER, Buffer);
+			ContextState.UniformBufferBound = Buffer;
 		}
 	}
-	
-	FORCEINLINE void CachedBindStorageBuffer( FOpenGLContextState& ContextState, GLuint Buffer )
+
+	void CachedBindBuffer(GLenum Type, GLuint Buffer)
 	{
 		VERIFY_GL_SCOPE();
-		if( ContextState.StorageBufferBound != Buffer )
+		switch (Type)
 		{
-			glBindBuffer( GL_SHADER_STORAGE_BUFFER, Buffer );
-			ContextState.StorageBufferBound = Buffer;
-		}
-	}
+		case GL_ARRAY_BUFFER:
+			if (ContextState.ArrayBufferBound != Buffer)
+			{
+				glBindBuffer(GL_ARRAY_BUFFER, Buffer);
+				ContextState.ArrayBufferBound = Buffer;
+			}
+			break;
 
-	void CachedBindPixelUnpackBuffer( FOpenGLContextState& ContextState, GLuint Buffer )
-	{
-		VERIFY_GL_SCOPE();
+		case GL_ELEMENT_ARRAY_BUFFER:
+			if (ContextState.ElementArrayBufferBound != Buffer)
+			{
+				glBindBuffer(GL_ELEMENT_ARRAY_BUFFER, Buffer);
+				ContextState.ElementArrayBufferBound = Buffer;
+			}
+			break;
 
-		if( ContextState.PixelUnpackBufferBound != Buffer )
-		{
-			glBindBuffer( GL_PIXEL_UNPACK_BUFFER, Buffer );
-			ContextState.PixelUnpackBufferBound = Buffer;
-		}
-	}
+		case GL_SHADER_STORAGE_BUFFER:
+			if (ContextState.StorageBufferBound != Buffer)
+			{
+				glBindBuffer(GL_SHADER_STORAGE_BUFFER, Buffer);
+				ContextState.StorageBufferBound = Buffer;
+			}
+			break;
 
-	void CachedBindUniformBuffer( FOpenGLContextState& ContextState, GLuint Buffer )
-	{
-		VERIFY_GL_SCOPE();
-		check(IsInRenderingThread()||IsInRHIThread());
-		if( ContextState.UniformBufferBound != Buffer )
-		{
-			glBindBuffer( GL_UNIFORM_BUFFER, Buffer );
-			ContextState.UniformBufferBound = Buffer;
-		}
-	}
+		case GL_PIXEL_UNPACK_BUFFER:
+			if (ContextState.PixelUnpackBufferBound != Buffer)
+			{
+				glBindBuffer(GL_PIXEL_UNPACK_BUFFER, Buffer);
+				ContextState.PixelUnpackBufferBound = Buffer;
+			}
+			break;
 
-	bool IsUniformBufferBound( FOpenGLContextState& ContextState, GLuint Buffer ) const
-	{
-		return ( ContextState.UniformBufferBound == Buffer );
+		default:
+			checkNoEntry();
+			break;
+		}
 	}
 
 	FOpenGLSamplerState* GetPointSamplerState() const { return (FOpenGLSamplerState*)PointSamplerState.GetReference(); }
@@ -680,6 +669,8 @@ public:
 	GLuint GetOpenGLFramebuffer(uint32 NumSimultaneousRenderTargets, FOpenGLTexture** RenderTargets, const uint32* ArrayIndices, const uint32* MipmapLevels, FOpenGLTexture* DepthStencilTarget);
 	GLuint GetOpenGLFramebuffer(uint32 NumSimultaneousRenderTargets, FOpenGLTexture** RenderTargets, const uint32* ArrayIndices, const uint32* MipmapLevels, FOpenGLTexture* DepthStencilTarget, int32 NumRenderingSamples);
 	
+	void ResolveTexture(FOpenGLTexture* Texture, uint32 MipIndex, uint32 ArrayIndex);
+
 private:
 	FBoundShaderStateRHIRef RHICreateBoundShaderState_Internal(FRHIVertexDeclaration* VertexDeclaration, FRHIVertexShader* VertexShader, FRHIPixelShader* PixelShader, FRHIGeometryShader* GeometryShader, bool FromPSOFileCache);
 
@@ -697,17 +688,12 @@ private:
 	TRefCountPtr<FOpenGLViewport>		DrawingViewport;
 	bool								bRevertToSharedContextAfterDrawingViewport = false;
 
-	bool								bIsRenderingContextAcquired = false;
-
 	EPrimitiveType						PrimitiveType = PT_Num;
 
 	/** A history of the most recently used bound shader states, used to keep transient bound shader states from being recreated for each use. */
 	TGlobalResource< TBoundShaderStateHistory<10000> > BoundShaderStateHistory;
 
-	/** Per-context state caching */
-	FOpenGLContextState	SharedContextState;
-	FOpenGLContextState	RenderingContextState;
-	FOpenGLContextState* CachedContextState = nullptr;
+	FOpenGLContextState	ContextState;
 
 	template <typename TRHIShader>
 	void ApplyStaticUniformBuffers(TRHIShader* Shader);
@@ -753,16 +739,19 @@ private:
 
 	void InitializeStateResources();
 
-	void SetupVertexArrays(FOpenGLContextState& ContextCache, uint32 BaseVertexIndex, FOpenGLStream* Streams, uint32 NumStreams, uint32 MaxVertices);
+	void SetupVertexArrays(uint32 BaseVertexIndex, FOpenGLStream* Streams, uint32 NumStreams, uint32 MaxVertices);
+
+	void SetupDraw(FOpenGLBuffer* IndexBuffer, uint32 BaseVertexIndex, uint32 MaxVertices);
+	void SetupDispatch();
 
 	/** needs to be called before each draw call */
-	void BindPendingFramebuffer( FOpenGLContextState& ContextState );
-	void BindPendingShaderState( FOpenGLContextState& ContextState );
-	void BindPendingComputeShaderState( FOpenGLContextState& ContextState, FOpenGLComputeShader* ComputeShader );
-	void UpdateRasterizerStateInOpenGLContext( FOpenGLContextState& ContextState );
-	void UpdateDepthStencilStateInOpenGLContext( FOpenGLContextState& ContextState );
-	void UpdateScissorRectInOpenGLContext( FOpenGLContextState& ContextState );
-	void UpdateViewportInOpenGLContext( FOpenGLContextState& ContextState );
+	void BindPendingFramebuffer();
+	void BindPendingShaderState();
+	void BindPendingComputeShaderState(FOpenGLComputeShader* ComputeShader);
+	void UpdateRasterizerStateInOpenGLContext();
+	void UpdateDepthStencilStateInOpenGLContext();
+	void UpdateScissorRectInOpenGLContext();
+	void UpdateViewportInOpenGLContext();
 	
 	template <class ShaderType> void SetResourcesFromTables(ShaderType* Shader);
 	FORCEINLINE void CommitGraphicsResourceTables()
@@ -776,15 +765,15 @@ private:
 	void CommitComputeResourceTables(FOpenGLComputeShader* ComputeShader);
 	void CommitNonComputeShaderConstants();
 	void CommitComputeShaderConstants(FOpenGLComputeShader* ComputeShader);
-	void SetPendingBlendStateForActiveRenderTargets( FOpenGLContextState& ContextState );
+	void SetPendingBlendStateForActiveRenderTargets();
 	
-	void SetupTexturesForDraw( FOpenGLContextState& ContextState);
 	template <typename StateType>
-	void SetupTexturesForDraw( FOpenGLContextState& ContextState, const StateType& ShaderState, int32 MaxTexturesNeeded);
+	void SetupTexturesForDraw(const StateType& ShaderState, int32 MaxTexturesNeeded);
+	void SetupTexturesForDraw();
 
-	void SetupUAVsForDraw(FOpenGLContextState& ContextState);
-	void SetupUAVsForCompute(FOpenGLContextState& ContextState, const FOpenGLComputeShader* ComputeShader);
-	void SetupUAVsForProgram(FOpenGLContextState& ContextState, const TBitArray<>& NeededBits, int32 MaxUAVUnitUsed);
+	void SetupUAVsForDraw();
+	void SetupUAVsForCompute(const FOpenGLComputeShader* ComputeShader);
+	void SetupUAVsForProgram(const TBitArray<>& NeededBits, int32 MaxUAVUnitUsed);
 
 	void RHIClearMRT(const bool* bClearColorArray, int32 NumClearColors, const FLinearColor* ColorArray, bool bClearDepth, float Depth, bool bClearStencil, uint32 Stencil);
 
@@ -795,18 +784,41 @@ public:
 	void InternalSetShaderBufferUAV(GLint UAVIndex, GLuint Resource);
 	void InternalSetSamplerStates(GLint TextureIndex, FOpenGLSamplerState* SamplerState);
 	void InitializeGLTextureInternal(FOpenGLTexture* Texture, void const* BulkDataPtr, uint64 BulkDataSize);
+
+	void ClearCachedAttributeState(int32 PositionAttrib, int32 TexCoordsAttrib);
+
 private:
 
-	void ApplyTextureStage(FOpenGLContextState& ContextState, GLint TextureIndex, const FTextureStage& TextureStage, FOpenGLSamplerState* SamplerState);
+	void ApplyTextureStage(GLint TextureIndex, const FTextureStage& TextureStage, FOpenGLSamplerState* SamplerState);
 
-	void ReadSurfaceDataRaw(FOpenGLContextState& ContextState, FRHITexture* TextureRHI,FIntRect Rect,TArray<uint8>& OutData, FReadSurfaceDataFlags InFlags);
+	void ReadSurfaceDataRaw(FRHITexture* TextureRHI, FIntRect Rect, TArray<uint8>& OutData, FReadSurfaceDataFlags InFlags);
 
-	void BindUniformBufferBase(FOpenGLContextState& ContextState, int32 NumUniformBuffers, FRHIUniformBuffer** BoundUniformBuffers, uint32* DynamicOffsets, uint32 FirstUniformBuffer, bool ForceUpdate);
+	void BindUniformBufferBase(int32 NumUniformBuffers, FRHIUniformBuffer** BoundUniformBuffers, uint32* DynamicOffsets, uint32 FirstUniformBuffer, bool ForceUpdate);
 
-	void ClearCurrentFramebufferWithCurrentScissor(FOpenGLContextState& ContextState, int8 ClearType, int32 NumClearColors, const bool* bClearColorArray, const FLinearColor* ClearColorArray, float Depth, uint32 Stencil);
+	void ClearCurrentFramebufferWithCurrentScissor(int8 ClearType, int32 NumClearColors, const bool* bClearColorArray, const FLinearColor* ClearColorArray, float Depth, uint32 Stencil);
 
 	FTextureLockTracker GLLockTracker;
 
+	class FOpenGLFenceKick
+	{
+	public:
+		FOpenGLFenceKick(FOpenGLDynamicRHI* RHI)
+			: RHI(*RHI)
+		{}
+		~FOpenGLFenceKick();
+
+		void OnDrawCall();
+		void Reset();
+
+	private:
+		void InsertKick();
+
+		TArray<UGLsync> Syncs;
+		int32 DrawCounter = 0;
+		GLuint LastSeenFramebuffer = 0;
+
+		FOpenGLDynamicRHI& RHI;
+	} KickHint { this };
 };
 
 PRAGMA_ENABLE_DEPRECATION_WARNINGS
diff --git a/Engine/Source/Runtime/OpenGLDrv/Private/OpenGLDrvPrivate.h b/Engine/Source/Runtime/OpenGLDrv/Private/OpenGLDrvPrivate.h
--- a/Engine/Source/Runtime/OpenGLDrv/Private/OpenGLDrvPrivate.h
+++ b/Engine/Source/Runtime/OpenGLDrv/Private/OpenGLDrvPrivate.h
@@ -26,8 +26,6 @@
 // UE_LOG(LogRHI, Log, TEXT("  " #IntEnum ": %d"), Value_##IntEnum)
 // #endif
 
-
-
 #if !defined(__GNUC__) && !defined(__clang__) && !defined(__INTEL_COMPILER) && !(defined(_MSC_VER) && _MSC_VER >= 1900)
 #define LOG_AND_GET_GL_INT(IntEnum, Default, Dest) \
 		do \
@@ -66,9 +64,6 @@
 		} \
 	while (0)
 
-struct FPlatformOpenGLContext;
-struct FPlatformOpenGLDevice;
-
 /**
  * The OpenGL RHI stats.
  */
@@ -117,15 +112,6 @@ DECLARE_CYCLE_STAT_EXTERN(TEXT("Uniform bind time"),STAT_OpenGLUniformBindTime,S
 DECLARE_CYCLE_STAT_EXTERN(TEXT("VBO setup time"),STAT_OpenGLVBOSetupTime,STATGROUP_OpenGLRHI, );
 #endif
 
-enum EOpenGLCurrentContext : int8
-{
-	CONTEXT_Other = -2,
-	CONTEXT_Invalid = -1,
-	CONTEXT_Shared = 0,
-	CONTEXT_Rendering = 1,
-//	CONTEXT_TextureStreaming = 2,
-};
-
 /*------------------------------------------------------------------------------
 	All platforms using OpenGL must implement the following API.
 ------------------------------------------------------------------------------*/
@@ -143,12 +129,6 @@ struct FPlatformOpenGLDevice;
  */
 bool PlatformInitOpenGL();
 
-/**
- * Returns true, if current thread has a valid OpenGL context selected, false otherwise.
- * Useful for debugging.
- */
-bool PlatformOpenGLContextValid();
-
 /**
  * Just a glGetError() call. Added to make it possible to compile GL_VERIFY macros in
  * other projects than OpenGLDrv.
@@ -156,17 +136,9 @@ bool PlatformOpenGLContextValid();
 int32 PlatformGlGetError();
 
 /**
- * Say which of the standard UE OpenGL contexts is currently in use. This is used by OpenGLDrv to
- * decide which of the context cache objects to use for mirroring OpenGL context state in
- * order to optimize out redundant OpenGL calls.
- */
-EOpenGLCurrentContext PlatformOpenGLCurrentContext(FPlatformOpenGLDevice* Device);
-
-/**
- * Check if the query's OpenGL context is a current one. If it's not, there's no sense in issuing
- * OpenGL command for this query and it must be released and recreated.
+ * Returns true when the calling thread owns the OpenGL rendering context.
  */
-bool PlatformContextIsCurrent( uint64 QueryContext );
+bool PlatformOpenGLThreadHasRenderingContext();
 
 /**
  * Create the OpenGL device, encompassing all data needed to handle set of OpenGL contexts
@@ -208,12 +180,6 @@ void PlatformRenderingContextSetup(FPlatformOpenGLDevice* Device);
 // Some platforms require flushing and rebinding renderbuffers after context/thread changes
 void PlatformFlushIfNeeded();
 
-/**
- * Set up shared context on current thread. This thread should be always set on main game thread (expect at times
- * when main game thread is doing the rendering for the rendering thread that doesn't exist).
- */
-void PlatformSharedContextSetup(FPlatformOpenGLDevice* Device);
-
 /**
  * This is used to detach an OpenGL context from current thread.
  */
diff --git a/Engine/Source/Runtime/OpenGLDrv/Private/OpenGLQuery.cpp b/Engine/Source/Runtime/OpenGLDrv/Private/OpenGLQuery.cpp
--- a/Engine/Source/Runtime/OpenGLDrv/Private/OpenGLQuery.cpp
+++ b/Engine/Source/Runtime/OpenGLDrv/Private/OpenGLQuery.cpp
@@ -78,60 +78,39 @@ void FOpenGLRenderQuery::Unlink()
 	Prev = nullptr;
 }
 
-void FOpenGLRenderQuery::CheckContext()
-{
-	check(Resource);
-	if (bSharedContext)
-	{
-		check(FOpenGLDynamicRHI::GetCurrentContext() == CONTEXT_Shared);
-	}
-	else
-	{
-		check(FOpenGLDynamicRHI::GetCurrentContext() == CONTEXT_Rendering);
-	}
-}
-
 void FOpenGLRenderQuery::AcquireGlQuery()
 {
+	VERIFY_GL_SCOPE();
+
 	if (Resource != 0)
 	{
 		// Already acquired
 		return;
 	}
 
-	EOpenGLCurrentContext Context = FOpenGLDynamicRHI::GetCurrentContext();
-	if (Context == CONTEXT_Shared)
+	while (ActiveQueries.First && ActiveQueries.Count >= GRHIMaximumInFlightQueries)
 	{
-		// Don't do any pooling on the shared context.
-		bSharedContext = true;
-		FOpenGL::GenQueries(1, &Resource);
+		// We can't start another query until more become available, due to the query count limit.
+		// Block for results on the oldest in-flight queries.
+		ActiveQueries.First->CacheResult(true);
 	}
-	else
-	{
-		check(Context == CONTEXT_Rendering);
-
-		while (ActiveQueries.First && ActiveQueries.Count >= GRHIMaximumInFlightQueries)
-		{
-			// We can't start another query until more become available, due to the query count limit.
-			// Block for results on the oldest in-flight queries.
-			ActiveQueries.First->CacheResult(true);
-		}
 
-		ActiveQueries.Count++;
+	ActiveQueries.Count++;
 
-		if (PooledQueries[Type].Num())
-		{
-			Resource = PooledQueries[Type].Pop();
-		}
-		else
-		{
-			FOpenGL::GenQueries(1, &Resource);
-		}
+	if (PooledQueries[Type].Num())
+	{
+		Resource = PooledQueries[Type].Pop();
+	}
+	else
+	{
+		FOpenGL::GenQueries(1, &Resource);
 	}
 }
 
 void FOpenGLRenderQuery::ReleaseGlQuery()
 {
+	VERIFY_GL_SCOPE();
+
 	if (Resource == 0)
 	{
 		// Already released
@@ -139,22 +118,11 @@ void FOpenGLRenderQuery::ReleaseGlQuery()
 		return;
 	}
 
-	CheckContext();
+	check(ActiveQueries.Count > 0);
+	ActiveQueries.Count--;
 
-	if (bSharedContext)
-	{
-		// Don't do any pooling on the shared context.
-		FOpenGL::DeleteQueries(1, &Resource);
-		bSharedContext = false;
-	}
-	else
-	{
-		check(ActiveQueries.Count > 0);
-		ActiveQueries.Count--;
+	PooledQueries[Type].Add(Resource);
 
-		PooledQueries[Type].Add(Resource);
-	}
-	
 	Resource = 0;
 
 	Unlink();
@@ -164,10 +132,10 @@ void FOpenGLRenderQuery::Begin()
 {
 	VERIFY_GL_SCOPE();
 
-	check(Resource == 0);
+	check(!IsLinked());
 	AcquireGlQuery();
 
-	CheckContext();
+	check(Resource);
 
 	switch(Type)
 	{
@@ -196,7 +164,7 @@ void FOpenGLRenderQuery::End()
 	VERIFY_GL_SCOPE();
 	AcquireGlQuery();
 	
-	CheckContext();
+	check(Resource);
 
 	switch (Type)
 	{
@@ -224,6 +192,8 @@ void FOpenGLRenderQuery::End()
 
 bool FOpenGLRenderQuery::CacheResult(bool bWait)
 {
+	VERIFY_GL_SCOPE();
+
 	if (BOPCounter == LastCachedBOPCounter.load(std::memory_order_relaxed))
 	{
 		// Value has been cached and no newer query operation has started.
@@ -231,7 +201,7 @@ bool FOpenGLRenderQuery::CacheResult(bool bWait)
 		return true;
 	}
 
-	CheckContext();
+	check(Resource);
 
 	if (!bWait)
 	{
@@ -335,19 +305,16 @@ bool FOpenGLRenderQuery_RHI::GetResult(bool bWait, uint64& OutResult)
 	// The query has not yet completed, and we want to wait for results.
 	// Append an RHI thread command that will force a readback of the GL query, then flush the RHI thread.
 	//	
-	RHICmdList.EnqueueLambda([Counter = TOPCounter, this](FRHICommandListImmediate&)
+	RHICmdList.EnqueueLambda([this](FRHICommandListImmediate&)
 	{
-		if (BOPCounter != Counter)
-		{
-			// Query result is not ready yet. Fetch it...
-			CacheResult(true);
-		}
+		// Fetch the query result if it hasn't arrived yet...
+		CacheResult(true);
 	});
 
 	// Wait for the above lambda to execute
 	RHICmdList.ImmediateFlush(EImmediateFlushType::FlushRHIThread);
 
-	checkf(TOPCounter == BOPCounter, TEXT("Attempting to get data from an RHI render query which was never issued."));
+	checkf(TOPCounter == LastCachedBOPCounter, TEXT("Attempting to get data from an RHI render query which was never issued."));
 	check(!IsLinked());
 
 	OutResult = FOpenGLRenderQuery::GetResult();
@@ -356,6 +323,12 @@ bool FOpenGLRenderQuery_RHI::GetResult(bool bWait, uint64& OutResult)
 
 void FOpenGLRenderQuery::PollQueryResults()
 {
+	if (!PlatformOpenGLThreadHasRenderingContext())
+	{
+		// Don't poll queries if this thread doesn't own the GL context.
+		return;
+	}
+
 	if (ActiveQueries.First)
 	{
 		TRACE_CPUPROFILER_EVENT_SCOPE(PollQueryResults);
@@ -802,7 +775,7 @@ void FOpenGLGPUFence::Wait(FRHICommandListImmediate& RHICmdList, FRHIGPUMask GPU
 
 void FOpenGLGPUFence::PollFencesUntil(FGraphEvent* Target)
 {
-	if (FOpenGLDynamicRHI::GetCurrentContext() != CONTEXT_Rendering && !Target)
+	if (!PlatformOpenGLThreadHasRenderingContext() && !Target)
 	{
 		// Don't poll fences if this thread doesn't own the GL context.
 		return;
diff --git a/Engine/Source/Runtime/OpenGLDrv/Private/OpenGLRenderTarget.cpp b/Engine/Source/Runtime/OpenGLDrv/Private/OpenGLRenderTarget.cpp
--- a/Engine/Source/Runtime/OpenGLDrv/Private/OpenGLRenderTarget.cpp
+++ b/Engine/Source/Runtime/OpenGLDrv/Private/OpenGLRenderTarget.cpp
@@ -43,12 +43,10 @@ public:
 		const uint32* InRenderTargetArrayIndices,
 		const uint32* InRenderTargetMipmapLevels,
 		FOpenGLTexture* InDepthStencilTarget,
-		int32 InNumRenderingSamples,
-		EOpenGLCurrentContext InContext
+		int32 InNumRenderingSamples
 		)
 		:	DepthStencilTarget(InDepthStencilTarget)
 		,	NumRenderingSamples(InNumRenderingSamples)
-		,	Context(InContext)
 	{
 		uint32 RenderTargetIndex;
 		for( RenderTargetIndex = 0; RenderTargetIndex < InNumRenderTargets; ++RenderTargetIndex )
@@ -76,8 +74,7 @@ public:
 		return	
 			!FMemory::Memcmp(A.RenderTargets, B.RenderTargets, sizeof(A.RenderTargets) ) && 
 			A.DepthStencilTarget == B.DepthStencilTarget &&
-			A.NumRenderingSamples == B.NumRenderingSamples &&
-			A.Context == B.Context;
+			A.NumRenderingSamples == B.NumRenderingSamples;
 	}
 
 	/**
@@ -87,7 +84,7 @@ public:
 	*/
 	friend uint32 GetTypeHash(const FOpenGLFramebufferKey &Key)
 	{
-		return FCrc::MemCrc_DEPRECATED(Key.RenderTargets, sizeof(Key.RenderTargets)) ^ GetTypeHash(Key.DepthStencilTarget) ^ GetTypeHash(Key.NumRenderingSamples) ^ GetTypeHash(Key.Context);
+		return FCrc::MemCrc_DEPRECATED(Key.RenderTargets, sizeof(Key.RenderTargets)) ^ GetTypeHash(Key.DepthStencilTarget) ^ GetTypeHash(Key.NumRenderingSamples);
 	}
 
 	const FOpenGLTexture* GetRenderTarget( int32 Index ) const { return RenderTargets[Index].Texture; }
@@ -99,7 +96,6 @@ private:
 	RenderTargetInfo RenderTargets[MaxSimultaneousRenderTargets];
 	FOpenGLTexture* DepthStencilTarget;
 	int32 NumRenderingSamples; // MSAA on tile
-	EOpenGLCurrentContext Context;
 };
 
 static void ConditionallyAllocateRenderbufferStorage(FOpenGLTexture& RenderTarget)
@@ -144,7 +140,7 @@ GLuint FOpenGLDynamicRHI::GetOpenGLFramebuffer(uint32 NumSimultaneousRenderTarge
 
 	check( NumSimultaneousRenderTargets <= MaxSimultaneousRenderTargets );
 
-	uint32 FramebufferRet = GetOpenGLFramebufferCache().FindRef(FOpenGLFramebufferKey(NumSimultaneousRenderTargets, RenderTargets, ArrayIndices, MipmapLevels, DepthStencilTarget, NumRenderingSamples, PlatformOpenGLCurrentContext(PlatformDevice)));
+	uint32 FramebufferRet = GetOpenGLFramebufferCache().FindRef(FOpenGLFramebufferKey(NumSimultaneousRenderTargets, RenderTargets, ArrayIndices, MipmapLevels, DepthStencilTarget, NumRenderingSamples));
 	if( FramebufferRet > 0 )
 	{
 		// Found and is valid. We never store zero as a result, increasing all results by 1 to avoid range overlap.
@@ -210,7 +206,7 @@ GLuint FOpenGLDynamicRHI::GetOpenGLFramebuffer(uint32 NumSimultaneousRenderTarge
 		FOpenGL::ReadBuffer(GL_NONE);
 		FOpenGL::DrawBuffer(GL_COLOR_ATTACHMENT0);
 
-		GetOpenGLFramebufferCache().Add(FOpenGLFramebufferKey(NumSimultaneousRenderTargets, RenderTargets, ArrayIndices, MipmapLevels, DepthStencilTarget, NumRenderingSamples, PlatformOpenGLCurrentContext(PlatformDevice)), Framebuffer + 1);
+		GetOpenGLFramebufferCache().Add(FOpenGLFramebufferKey(NumSimultaneousRenderTargets, RenderTargets, ArrayIndices, MipmapLevels, DepthStencilTarget, NumRenderingSamples), Framebuffer + 1);
 		
 		return Framebuffer;
 	}
@@ -367,7 +363,7 @@ GLuint FOpenGLDynamicRHI::GetOpenGLFramebuffer(uint32 NumSimultaneousRenderTarge
 	
 	FOpenGL::CheckFrameBuffer();
 
-	GetOpenGLFramebufferCache().Add(FOpenGLFramebufferKey(NumSimultaneousRenderTargets, RenderTargets, ArrayIndices, MipmapLevels, DepthStencilTarget, NumRenderingSamples, PlatformOpenGLCurrentContext(PlatformDevice)), Framebuffer+1);
+	GetOpenGLFramebufferCache().Add(FOpenGLFramebufferKey(NumSimultaneousRenderTargets, RenderTargets, ArrayIndices, MipmapLevels, DepthStencilTarget, NumRenderingSamples), Framebuffer+1);
 
 	return Framebuffer;
 }
@@ -421,7 +417,7 @@ void FOpenGLDynamicRHI::PurgeFramebufferFromCaches( GLuint Framebuffer )
 {
 	VERIFY_GL_SCOPE();
 
-	if( Framebuffer == PendingState.Framebuffer )
+	if (Framebuffer == PendingState.Framebuffer)
 	{
 		PendingState.Framebuffer = 0;
 		FMemory::Memset(PendingState.RenderTargets, 0, sizeof(PendingState.RenderTargets));
@@ -431,18 +427,13 @@ void FOpenGLDynamicRHI::PurgeFramebufferFromCaches( GLuint Framebuffer )
 		PendingState.bFramebufferSetupInvalid = true;
 	}
 
-	if( Framebuffer == SharedContextState.Framebuffer )
+	if (Framebuffer == ContextState.Framebuffer)
 	{
-		SharedContextState.Framebuffer = -1;
-	}
-
-	if( Framebuffer == RenderingContextState.Framebuffer )
-	{
-		RenderingContextState.Framebuffer = -1;
+		ContextState.Framebuffer = -1;
 	}
 }
 
-void FOpenGLDynamicRHI::ReadSurfaceDataRaw(FOpenGLContextState& ContextState, FRHITexture* TextureRHI,FIntRect Rect,TArray<uint8>& OutData, FReadSurfaceDataFlags InFlags)
+void FOpenGLDynamicRHI::ReadSurfaceDataRaw(FRHITexture* TextureRHI, FIntRect Rect, TArray<uint8>& OutData, FReadSurfaceDataFlags InFlags)
 {
 	VERIFY_GL_SCOPE();
 
@@ -753,14 +744,13 @@ void FOpenGLDynamicRHI::RHIReadSurfaceData(FRHITexture* TextureRHI,FIntRect Rect
 	FRHICommandListImmediate& RHICmdList = FRHICommandListImmediate::Get();
 	
 	RHICmdList.EnqueueLambda([&](FRHICommandListImmediate&)
-		{
-			TArray<uint8> Temp;
+	{
+		TArray<uint8> Temp;
 
-			FOpenGLContextState& ContextState = GetContextStateForCurrentContext();
+		ReadSurfaceDataRaw(TextureRHI, Rect, Temp, InFlags);
+		FMemory::Memcpy(OutData.GetData(), Temp.GetData(), Size * sizeof(FColor));
+	});
 
-			ReadSurfaceDataRaw(ContextState, TextureRHI, Rect, Temp, InFlags);
-			FMemory::Memcpy(OutData.GetData(), Temp.GetData(), Size * sizeof(FColor));
-		});
 	RHITHREAD_GLTRACE_BLOCKING;
 	RHICmdList.ImmediateFlush(EImmediateFlushType::FlushRHIThread);
 }
@@ -807,7 +797,7 @@ void FOpenGLDynamicRHI::RHIReadSurfaceData(FRHITexture* TextureRHI, FIntRect Rec
 		glReadPixels(Rect.Min.X, Rect.Min.Y, SizeX, SizeY, GL_RGBA, GL_FLOAT, OutData.GetData());
 		glPixelStorei(GL_PACK_ALIGNMENT, 4);
 
-		GetContextStateForCurrentContext().Framebuffer = (GLuint)-1;	
+		ContextState.Framebuffer = (GLuint)-1;
 	});
 	RHITHREAD_GLTRACE_BLOCKING;
 	RHICmdList.ImmediateFlush(EImmediateFlushType::FlushRHIThread);
@@ -930,7 +920,7 @@ void FOpenGLDynamicRHI::RHIReadSurfaceFloatData(FRHITexture* TextureRHI,FIntRect
 			glDeleteFramebuffers( 1, &SourceFramebuffer);
 		}
 
-		GetContextStateForCurrentContext().Framebuffer = (GLuint)-1;
+		ContextState.Framebuffer = (GLuint)-1;
 	});
 	RHITHREAD_GLTRACE_BLOCKING;
 	RHICmdList.ImmediateFlush(EImmediateFlushType::FlushRHIThread);
@@ -985,7 +975,6 @@ void FOpenGLDynamicRHI::RHIRead3DSurfaceFloatData(FRHITexture* TextureRHI,FIntRe
 		glPixelStorei( GL_PACK_ALIGNMENT, 4 );
 
 		// Clean up
-		FOpenGLContextState& ContextState = GetContextStateForCurrentContext();
 		auto& TextureState = ContextState.Textures[0];
 		glBindTexture(GL_TEXTURE_3D, (TextureState.Target == GL_TEXTURE_3D) ? TextureState.Resource : 0);
 		glActiveTexture( GL_TEXTURE0 + ContextState.ActiveTexture );
@@ -999,7 +988,7 @@ void FOpenGLDynamicRHI::RHIRead3DSurfaceFloatData(FRHITexture* TextureRHI,FIntRe
 
 
 
-void FOpenGLDynamicRHI::BindPendingFramebuffer( FOpenGLContextState& ContextState )
+void FOpenGLDynamicRHI::BindPendingFramebuffer()
 {
 	VERIFY_GL_SCOPE();
 
@@ -1078,7 +1067,6 @@ void FOpenGLDynamicRHI::RHIBeginRenderPass(const FRHIRenderPassInfo& InInfo, con
 	}
 
 	// Reenable non-coherent framebuffer fetch if needed
-	FOpenGLContextState& ContextState = GetContextStateForCurrentContext();
 	if (!ContextState.bNonCoherentFramebufferFetchEnabled)
 	{
 		ContextState.bNonCoherentFramebufferFetchEnabled = FAndroidOpenGL::ResetNonCoherentFramebufferFetch();
@@ -1140,7 +1128,7 @@ void FOpenGLDynamicRHI::RHINextSubpass()
 	if (RenderPassInfo.SubpassHint == ESubpassHint::DepthReadSubpass ||
 		RenderPassInfo.SubpassHint == ESubpassHint::DeferredShadingSubpass)
 	{
-		if (GetContextStateForCurrentContext().bNonCoherentFramebufferFetchEnabled)
+		if (ContextState.bNonCoherentFramebufferFetchEnabled)
 		{
 			FOpenGL::FrameBufferFetchBarrier();
 		}
diff --git a/Engine/Source/Runtime/OpenGLDrv/Private/OpenGLResources.h b/Engine/Source/Runtime/OpenGLDrv/Private/OpenGLResources.h
--- a/Engine/Source/Runtime/OpenGLDrv/Private/OpenGLResources.h
+++ b/Engine/Source/Runtime/OpenGLDrv/Private/OpenGLResources.h
@@ -36,17 +36,6 @@ class FOpenGLLinkedProgram;
 class FOpenGLTexture;
 typedef TArray<ANSICHAR> FAnsiCharArray;
 
-
-extern void OnBufferDeletion( GLuint BufferResource );
-extern void OnPixelBufferDeletion( GLuint PixelBufferResource );
-extern void OnUniformBufferDeletion( GLuint UniformBufferResource, uint32 AllocatedSize, bool bStreamDraw, uint32 Offset, uint8* Pointer );
-extern void OnProgramDeletion( GLint ProgramResource );
-
-extern void CachedBindBuffer( GLenum Type, GLuint Buffer );
-extern void CachedBindPixelUnpackBuffer( GLenum Type, GLuint Buffer );
-extern void CachedBindUniformBuffer( GLuint Buffer );
-extern bool IsUniformBufferBound( GLuint Buffer );
-
 namespace OpenGLConsoleVariables
 {
 	extern int32 bUseMapBuffer;
@@ -105,10 +94,23 @@ private:
 	class FOpenGLView* LinkedViews = nullptr;
 };
 
-typedef void (*BufferBindFunction)( GLenum Type, GLuint Buffer );
+class FOpenGLBufferBase
+{
+protected:
+	FOpenGLBufferBase(GLenum Type)
+		: Type(Type)
+	{}
 
-template <typename BaseType, BufferBindFunction BufBind>
-class TOpenGLBuffer : public BaseType
+public:
+	GLenum Type;
+	GLuint Resource = 0;
+
+	void Bind();
+	void OnBufferDeletion();
+};
+
+template <typename BaseType>
+class TOpenGLBuffer : public BaseType, public FOpenGLBufferBase
 {
 	void LoadData(uint32 InOffset, uint32 InSize, const void* InData)
 	{
@@ -145,14 +147,7 @@ class TOpenGLBuffer : public BaseType
 public:
 	TOpenGLBuffer(FRHICommandListBase* RHICmdList, GLenum InType, FRHIBufferDesc const& BufferDesc, const void *InData)
 		: BaseType(BufferDesc)
-		, Resource(0)
-		, Type(InType)
-		, bIsLocked(false)
-		, bIsLockReadOnly(false)
-		, bLockBufferWasAllocated(false)
-		, LockSize(0)
-		, LockOffset(0)
-		, LockBuffer(nullptr)
+		, FOpenGLBufferBase(InType)
 	{
 		if (RHICmdList && RHICmdList->IsTopOfPipe() && InData)
 		{
@@ -166,7 +161,6 @@ public:
 			VERIFY_GL_SCOPE();
 
 			FOpenGL::GenBuffers(1, &Resource);
-			check(Type != GL_UNIFORM_BUFFER || !IsUniformBufferBound(Resource));
 
 			Bind();
 
@@ -200,13 +194,7 @@ public:
 		ReleaseOwnership();
 	}
 
-	void Bind()
-	{
-		VERIFY_GL_SCOPE();
-		BufBind(Type, Resource);
-	}
-
-	uint8 *Lock(uint32 InOffset, uint32 InSize, bool bReadOnly, bool bDiscard)
+	uint8* Lock(uint32 InOffset, uint32 InSize, bool bReadOnly, bool bDiscard)
 	{
 		SCOPE_CYCLE_COUNTER_DETAILED(STAT_OpenGLMapBufferTime);
 		check(InOffset + InSize <= BaseType::GetSize());
@@ -304,7 +292,7 @@ public:
 		return Data;
 	}
 
-	uint8 *LockWriteOnlyUnsynchronized(uint32 InOffset, uint32 InSize, bool bDiscard)
+	uint8* LockWriteOnlyUnsynchronized(uint32 InOffset, uint32 InSize, bool bDiscard)
 	{
 		//SCOPE_CYCLE_COUNTER_DETAILED(STAT_OpenGLMapBufferTime);
 		check(InOffset + InSize <= BaseType::GetSize());
@@ -505,11 +493,10 @@ public:
 				}
 			}
 
-			if (BaseType::OnDelete(Resource, BaseType::GetSize(), false, 0))
-			{
-				FOpenGL::DeleteBuffers(1, &Resource);
-				Resource = 0;
-			}
+			OnBufferDeletion();
+
+			FOpenGL::DeleteBuffers(1, &Resource);
+			Resource = 0;
 
 			LockBuffer = nullptr;
 			OpenGLBufferStats::UpdateBufferStats(BaseType::GetDesc(), false);
@@ -520,17 +507,14 @@ public:
 		BaseType::ReleaseOwnership();
 	}
 
-	GLuint Resource;
-	GLenum Type;
-
 private:
-	uint32 bIsLocked : 1;
-	uint32 bIsLockReadOnly : 1;
-	uint32 bLockBufferWasAllocated : 1;
+	uint32 bIsLocked               : 1 = false;
+	uint32 bIsLockReadOnly         : 1 = false;
+	uint32 bLockBufferWasAllocated : 1 = false;
 
-	GLuint LockSize;
-	GLuint LockOffset;
-	void* LockBuffer;
+	GLuint LockSize = 0;
+	GLuint LockOffset = 0;
+	void* LockBuffer = nullptr;
 
 	// A cached allocation that can be reused. The same allocation can never be in CachedBuffer and LockBuffer at the same time.
 	void* CachedBuffer = nullptr;
@@ -545,12 +529,6 @@ public:
 		: Desc(BufferDesc)
 	{}
 
-	static bool OnDelete(GLuint Resource,uint32 Size,bool bStreamDraw,uint32 Offset)
-	{
-		OnPixelBufferDeletion(Resource);
-		return true;
-	}
-
 	const FRHIBufferDesc& GetDesc() const { return Desc; }
 	uint32 GetSize() const { return Desc.Size; }
 	EBufferUsageFlags GetUsage() const { return Desc.Usage; }
@@ -603,16 +581,10 @@ public:
 		LLM_SCOPED_PAUSE_TRACKING_WITH_ENUM_AND_AMOUNT(ELLMTag::Meshes, -(int64)GetSize(), ELLMTracker::Default, ELLMAllocType::None);
 		FRHIBuffer::ReleaseOwnership();
 	}
-
-	static bool OnDelete(GLuint Resource,uint32 Size,bool bStreamDraw,uint32 Offset)
-	{
-		OnBufferDeletion(Resource);
-		return true;
-	}
 };
 
-typedef TOpenGLBuffer<FOpenGLBasePixelBuffer, CachedBindPixelUnpackBuffer> FOpenGLPixelBuffer;
-typedef TOpenGLBuffer<FOpenGLBaseBuffer, CachedBindBuffer> FOpenGLBuffer;
+typedef TOpenGLBuffer<FOpenGLBasePixelBuffer> FOpenGLPixelBuffer;
+typedef TOpenGLBuffer<FOpenGLBaseBuffer     > FOpenGLBuffer;
 
 struct FOpenGLEUniformBufferData : public FRefCountedObject
 {
@@ -1113,6 +1085,7 @@ private:
 
 	void Fill2DGLTextureImage(const FOpenGLTextureFormat& GLFormat, const bool bSRGB, uint32 MipIndex, const void* LockedBuffer, uint32 LockedSize, uint32 ArrayIndex);
 
+public:
 	uint32 GetEffectiveSizeZ() const
 	{
 		FRHITextureDesc const& Desc = GetDesc();
@@ -1138,9 +1111,9 @@ public:
 	// Pointer to current sampler state in this unit
 	class FOpenGLSamplerState* SamplerState = nullptr;
 
-private:
 	TArray<TRefCountPtr<FOpenGLPixelBuffer>> PixelBuffers;
 
+private:
 	/** Bitfields marking whether we have allocated storage for each mip */
 	TBitArray<TInlineAllocator<1> > bAllocatedStorage;
 
@@ -1195,8 +1168,6 @@ private:
 	/** The query resource. */
 	GLuint Resource = 0;
 
-	uint8 bSharedContext : 1 = false;
-
 protected:
 	EType const Type : 2;
 
@@ -1234,8 +1205,6 @@ private:
 	void Unlink();
 
 	void SetResult(uint64 Value);
-
-	void CheckContext();
 };
 
 class FOpenGLRenderQuery_RHI : public FRHIRenderQuery, public FOpenGLRenderQuery
diff --git a/Engine/Source/Runtime/OpenGLDrv/Private/OpenGLShaders.cpp b/Engine/Source/Runtime/OpenGLDrv/Private/OpenGLShaders.cpp
--- a/Engine/Source/Runtime/OpenGLDrv/Private/OpenGLShaders.cpp
+++ b/Engine/Source/Runtime/OpenGLDrv/Private/OpenGLShaders.cpp
@@ -1126,7 +1126,7 @@ static void MarkShaderParameterCachesDirty(FOpenGLShaderParameterCache* ShaderPa
 	}
 }
 
-void FOpenGLDynamicRHI::BindUniformBufferBase(FOpenGLContextState& ContextState, int32 NumUniformBuffers, FRHIUniformBuffer** BoundUniformBuffers, uint32* DynamicOffsets, uint32 FirstUniformBuffer, bool ForceUpdate)
+void FOpenGLDynamicRHI::BindUniformBufferBase(int32 NumUniformBuffers, FRHIUniformBuffer** BoundUniformBuffers, uint32* DynamicOffsets, uint32 FirstUniformBuffer, bool ForceUpdate)
 {
 	SCOPE_CYCLE_COUNTER_DETAILED(STAT_OpenGLUniformBindTime);
 	VERIFY_GL_SCOPE();
@@ -1163,7 +1163,7 @@ void FOpenGLDynamicRHI::BindUniformBufferBase(FOpenGLContextState& ContextState,
 				FMemory::Memzero(ZeroBuffer,ZERO_FILLED_DUMMY_UNIFORM_BUFFER_SIZE);
 				FOpenGL::GenBuffers(1, &PendingState.ZeroFilledDummyUniformBuffer);
 				check(PendingState.ZeroFilledDummyUniformBuffer != 0);
-				CachedBindUniformBuffer(ContextState,PendingState.ZeroFilledDummyUniformBuffer);
+				CachedBindUniformBuffer(PendingState.ZeroFilledDummyUniformBuffer);
 				glBufferData(GL_UNIFORM_BUFFER, ZERO_FILLED_DUMMY_UNIFORM_BUFFER_SIZE, ZeroBuffer, GL_STATIC_DRAW);
 				FMemory::Free(ZeroBuffer);
 
@@ -2575,7 +2575,7 @@ void FOpenGLDynamicRHI::LinkComputeProgram(FRHIComputeShader* ComputeShaderRHI)
 	if (ComputeShader->LinkedProgram == nullptr)
 	{
 		// Make sure we have OpenGL context set up, and invalidate the parameters cache and current program (as we'll link a new one soon)
-		GetContextStateForCurrentContext().Program = -1;
+		ContextState.Program = -1;
 		MarkShaderParameterCachesDirty(PendingState.ShaderParameters, true);
 		PendingState.LinkedProgramAndDirtyFlag = nullptr;
 
@@ -2821,7 +2821,7 @@ FBoundShaderStateRHIRef FOpenGLDynamicRHI::RHICreateBoundShaderState_Internal(
 	else
 	{
 		// Make sure we have OpenGL context set up, and invalidate the parameters cache and current program (as we'll link a new one soon)
-		GetContextStateForCurrentContext().Program = -1;
+		ContextState.Program = -1;
 		MarkShaderParameterCachesDirty(PendingState.ShaderParameters, false);
 		PendingState.LinkedProgramAndDirtyFlag = nullptr;
 
@@ -2853,7 +2853,7 @@ void DestroyShadersAndPrograms()
 	}
 }
 
-void FOpenGLDynamicRHI::BindPendingShaderState( FOpenGLContextState& ContextState )
+void FOpenGLDynamicRHI::BindPendingShaderState()
 {
 	SCOPE_CYCLE_COUNTER_DETAILED(STAT_OpenGLShaderBindTime);
 	VERIFY_GL_SCOPE();
@@ -2897,7 +2897,6 @@ void FOpenGLDynamicRHI::BindPendingShaderState( FOpenGLContextState& ContextStat
 		if (PendingState.bAnyDirtyRealUniformBuffers[SF_Vertex])
 		{
 			BindUniformBufferBase(
-				ContextState,
 				NumUniformBuffers[SF_Vertex],
 				PendingState.BoundUniformBuffers[SF_Vertex],
 				PendingState.BoundUniformBuffersDynamicOffset[SF_Vertex],
@@ -2909,7 +2908,6 @@ void FOpenGLDynamicRHI::BindPendingShaderState( FOpenGLContextState& ContextStat
 		if (PendingState.bAnyDirtyRealUniformBuffers[SF_Pixel])
 		{
 			BindUniformBufferBase(
-				ContextState,
 				NumUniformBuffers[SF_Pixel],
 				PendingState.BoundUniformBuffers[SF_Pixel],
 				PendingState.BoundUniformBuffersDynamicOffset[SF_Pixel],
@@ -2921,7 +2919,6 @@ void FOpenGLDynamicRHI::BindPendingShaderState( FOpenGLContextState& ContextStat
 		if (NumUniformBuffers[SF_Geometry] >= 0 && PendingState.bAnyDirtyRealUniformBuffers[SF_Geometry])
 		{
 			BindUniformBufferBase(
-				ContextState,
 				NumUniformBuffers[SF_Geometry],
 				PendingState.BoundUniformBuffers[SF_Geometry],
 				PendingState.BoundUniformBuffersDynamicOffset[SF_Geometry],
@@ -3013,7 +3010,7 @@ FOpenGLBoundShaderState::~FOpenGLBoundShaderState()
 			StaticLastReleasedProgramsIndex = 0;
 		}
 
-		OnProgramDeletion(LinkedProgram->Program);
+		FOpenGLDynamicRHI::Get().OnProgramDeletion(LinkedProgram->Program);
 	}
 }
 
@@ -3084,7 +3081,7 @@ bool FOpenGLComputeShader::NeedsUAVStage(int32 UAVStageIndex) const
 	return LinkedProgram->UAVStageNeeds[UAVStageIndex];
 }
 
-void FOpenGLDynamicRHI::BindPendingComputeShaderState(FOpenGLContextState& ContextState, FOpenGLComputeShader* ComputeShader)
+void FOpenGLDynamicRHI::BindPendingComputeShaderState(FOpenGLComputeShader* ComputeShader)
 {
 	VERIFY_GL_SCOPE();
 	bool ForceUniformBindingUpdate = false;
@@ -3105,7 +3102,6 @@ void FOpenGLDynamicRHI::BindPendingComputeShaderState(FOpenGLContextState& Conte
 	if (PendingState.bAnyDirtyRealUniformBuffers[SF_Compute])
 	{
 		BindUniformBufferBase(
-			ContextState,
 			ComputeShader->Bindings.NumUniformBuffers,
 			PendingState.BoundUniformBuffers[SF_Compute],
 			PendingState.BoundUniformBuffersDynamicOffset[SF_Compute],
diff --git a/Engine/Source/Runtime/OpenGLDrv/Private/OpenGLTexture.cpp b/Engine/Source/Runtime/OpenGLDrv/Private/OpenGLTexture.cpp
--- a/Engine/Source/Runtime/OpenGLDrv/Private/OpenGLTexture.cpp
+++ b/Engine/Source/Runtime/OpenGLDrv/Private/OpenGLTexture.cpp
@@ -533,10 +533,8 @@ void FOpenGLDynamicRHI::InitializeGLTextureInternal(FOpenGLTexture* Texture, voi
 	
 	const bool bMultiviewMSAADepthTarget = (Texture->bDepthStencil && Desc.NumSamples > 1 && Desc.Dimension == ETextureDimension::Texture2DArray);
 
-	FOpenGLContextState& ContextState = GetContextStateForCurrentContext();
-
 	// Make sure PBO is disabled
-	CachedBindPixelUnpackBuffer(ContextState, 0);
+	CachedBindBuffer(GL_PIXEL_UNPACK_BUFFER, 0);
 
 	bool bAllocatedStorage = false;
 	if (Texture->bMultisampleRenderbuffer)
@@ -548,7 +546,7 @@ void FOpenGLDynamicRHI::InitializeGLTextureInternal(FOpenGLTexture* Texture, voi
 	else
 	{
 		// Use a texture stage that's not likely to be used for draws, to avoid waiting
-		CachedSetupTextureStage(ContextState, FOpenGL::GetMaxCombinedTextureImageUnits() - 1, Target, TextureID, 0, Desc.NumMips);
+		CachedSetupTextureStage(FOpenGL::GetMaxCombinedTextureImageUnits() - 1, Target, TextureID, 0, Desc.NumMips);
 
 		if((GLFormat.bBGRA && !EnumHasAnyFlags(Desc.Flags, TexCreate_RenderTargetable)) 
 #if !PLATFORM_ANDROID
@@ -784,31 +782,31 @@ void FOpenGLDynamicRHI::InitializeGLTextureInternal(FOpenGLTexture* Texture, voi
 }
 
 
-void FOpenGLTexture::Resolve(uint32 MipIndex,uint32 ArrayIndex)
+void FOpenGLDynamicRHI::ResolveTexture(FOpenGLTexture* Texture, uint32 MipIndex,uint32 ArrayIndex)
 {
 	VERIFY_GL_SCOPE();
-	check(!GetTexture2D() || GetNumSamples() == 1);
+	check(!Texture->GetTexture2D() || Texture->GetNumSamples() == 1);
 	
 	// Calculate the dimensions of the mip-map.
-	EPixelFormat PixelFormat = this->GetFormat();
+	EPixelFormat PixelFormat = Texture->GetFormat();
 	const uint32 BlockSizeX = GPixelFormats[PixelFormat].BlockSizeX;
 	const uint32 BlockSizeY = GPixelFormats[PixelFormat].BlockSizeY;
 	const uint32 BlockBytes = GPixelFormats[PixelFormat].BlockBytes;
-	const uint32 MipSizeX = FMath::Max(this->GetSizeX() >> MipIndex,BlockSizeX);
-	const uint32 MipSizeY = FMath::Max(this->GetSizeY() >> MipIndex,BlockSizeY);
+	const uint32 MipSizeX = FMath::Max(Texture->GetSizeX() >> MipIndex, BlockSizeX);
+	const uint32 MipSizeY = FMath::Max(Texture->GetSizeY() >> MipIndex, BlockSizeY);
 	uint32 NumBlocksX = (MipSizeX + BlockSizeX - 1) / BlockSizeX;
 	uint32 NumBlocksY = (MipSizeY + BlockSizeY - 1) / BlockSizeY;
 	const uint32 MipBytes = NumBlocksX * NumBlocksY * BlockBytes;
 	
-	const int32 BufferIndex = MipIndex * (bCubemap ? 6 : 1) * this->GetEffectiveSizeZ() + ArrayIndex;
+	const int32 BufferIndex = MipIndex * (Texture->bCubemap ? 6 : 1) * Texture->GetEffectiveSizeZ() + ArrayIndex;
 	
 	// Standard path with a PBO mirroring ever slice of a texture to allow multiple simulataneous maps
-	if (!IsValidRef(PixelBuffers[BufferIndex]))
+	if (!IsValidRef(Texture->PixelBuffers[BufferIndex]))
 	{
-		PixelBuffers[BufferIndex] = new FOpenGLPixelBuffer(nullptr, GL_PIXEL_UNPACK_BUFFER, FRHIBufferDesc(MipBytes, 0, BUF_Dynamic), nullptr);
+		Texture->PixelBuffers[BufferIndex] = new FOpenGLPixelBuffer(nullptr, GL_PIXEL_UNPACK_BUFFER, FRHIBufferDesc(MipBytes, 0, BUF_Dynamic), nullptr);
 	}
 	
-	TRefCountPtr<FOpenGLPixelBuffer> PixelBuffer = PixelBuffers[BufferIndex];
+	TRefCountPtr<FOpenGLPixelBuffer> PixelBuffer = Texture->PixelBuffers[BufferIndex];
 	check(PixelBuffer->GetSize() == MipBytes);
 	check(!PixelBuffer->IsLocked());
 	
@@ -816,17 +814,15 @@ void FOpenGLTexture::Resolve(uint32 MipIndex,uint32 ArrayIndex)
 	// This may be further optimized by caching information if surface content was changed since last lock.
 	
 	const FOpenGLTextureFormat& GLFormat = GOpenGLTextureFormats[PixelFormat];
-	const bool bSRGB = EnumHasAnyFlags(this->GetFlags(), TexCreate_SRGB);
+	const bool bSRGB = EnumHasAnyFlags(Texture->GetFlags(), TexCreate_SRGB);
 	
 	// Use a texture stage that's not likely to be used for draws, to avoid waiting
-	FOpenGLContextState& ContextState = FOpenGLDynamicRHI::Get().GetContextStateForCurrentContext();
-
-	FOpenGLDynamicRHI::Get().CachedSetupTextureStage(ContextState, FOpenGL::GetMaxCombinedTextureImageUnits() - 1, Target, this->GetResource(), -1, this->GetNumMips());
+	CachedSetupTextureStage(FOpenGL::GetMaxCombinedTextureImageUnits() - 1, Texture->Target, Texture->GetResource(), -1, Texture->GetNumMips());
 	
 	glBindBuffer( GL_PIXEL_PACK_BUFFER, PixelBuffer->Resource );
 
 	{
-		if (GetDesc().IsTextureArray() || GetDesc().IsTexture3D())
+		if (Texture->GetDesc().IsTextureArray() || Texture->GetDesc().IsTexture3D())
 		{
 			// apparently it's not possible to retrieve compressed image from GL_TEXTURE_2D_ARRAY in OpenGL for compressed images
 			// and for uncompressed ones it's not possible to specify the image index
@@ -837,15 +833,16 @@ void FOpenGLTexture::Resolve(uint32 MipIndex,uint32 ArrayIndex)
 			if (GLFormat.bCompressed)
 			{
 				FOpenGL::GetCompressedTexImage(
-											   bCubemap ? GL_TEXTURE_CUBE_MAP_POSITIVE_X + ArrayIndex : Target,
-											   MipIndex,
-											   0);	// offset into PBO
+					Texture->bCubemap
+						? GL_TEXTURE_CUBE_MAP_POSITIVE_X + ArrayIndex
+						: Texture->Target,
+					MipIndex,
+					0);	// offset into PBO
 			}
 			else
 			{
 				// Get framebuffer for texture
-				FOpenGLTexture* Texture = this;
-				GLuint SourceFramebuffer = FOpenGLDynamicRHI::Get().GetOpenGLFramebuffer(1, &Texture, (bCubemap ? &ArrayIndex : nullptr), &MipIndex, nullptr);
+				GLuint SourceFramebuffer = GetOpenGLFramebuffer(1, &Texture, (Texture->bCubemap ? &ArrayIndex : nullptr), &MipIndex, nullptr);
 				// Bind the framebuffer
 				glBindFramebuffer(UGL_READ_FRAMEBUFFER, SourceFramebuffer);
 				FOpenGL::ReadBuffer(GL_COLOR_ATTACHMENT0);
@@ -854,12 +851,12 @@ void FOpenGLTexture::Resolve(uint32 MipIndex,uint32 ArrayIndex)
 				glReadPixels(0, 0, MipSizeX, MipSizeY, GLFormat.Format, GLFormat.Type, 0);
 				glPixelStorei(GL_PACK_ALIGNMENT, 4);
 
-				FOpenGLDynamicRHI::Get().GetContextStateForCurrentContext().Framebuffer = (GLuint)-1;
+				ContextState.Framebuffer = (GLuint)-1;
 			}
 		}
 	}
 	
-	glBindBuffer( GL_PIXEL_PACK_BUFFER, 0 );
+	glBindBuffer(GL_PIXEL_PACK_BUFFER, 0);
 	
 	// No need to restore texture stage; leave it like this,
 	// and the next draw will take care of cleaning it up; or
@@ -976,7 +973,7 @@ void* FOpenGLTexture::Lock(uint32 InMipIndex,uint32 ArrayIndex,EResourceLockMode
 
 		if (LockMode != RLM_WriteOnly && !bCPUTexResolved)
 		{
-			Resolve(InMipIndex, ArrayIndex);
+			FOpenGLDynamicRHI::Get().ResolveTexture(this, InMipIndex, ArrayIndex);
 		}
 
 		result = PixelBuffer->Lock(0, PixelBuffer->GetSize(), LockMode == RLM_ReadOnly, LockMode != RLM_ReadOnly);
@@ -1068,10 +1065,9 @@ void FOpenGLTexture::Unlock(uint32 MipIndex,uint32 ArrayIndex)
 		check(this->GetSizeZ() == 0);
 
 		// Use a texture stage that's not likely to be used for draws, to avoid waiting
-		FOpenGLContextState& ContextState = FOpenGLDynamicRHI::Get().GetContextStateForCurrentContext();
-		FOpenGLDynamicRHI::Get().CachedSetupTextureStage(ContextState, FOpenGL::GetMaxCombinedTextureImageUnits() - 1, Target, GetResource(), -1, this->GetNumMips());
+		FOpenGLDynamicRHI::Get().CachedSetupTextureStage(FOpenGL::GetMaxCombinedTextureImageUnits() - 1, Target, GetResource(), -1, this->GetNumMips());
 
-		CachedBindPixelUnpackBuffer(GL_PIXEL_UNPACK_BUFFER, 0);
+		FOpenGLDynamicRHI::Get().CachedBindBuffer(GL_PIXEL_UNPACK_BUFFER, 0);
 
 		// get the source data and size
 		uint16* floatData = (uint16*)PixelBuffer->GetLockedBuffer();
@@ -1114,7 +1110,7 @@ void FOpenGLTexture::Unlock(uint32 MipIndex,uint32 ArrayIndex)
 		// and the next draw will take care of cleaning it up; or
 		// next operation that needs the stage will switch something else in on it.
 
-		CachedBindPixelUnpackBuffer(GL_PIXEL_UNPACK_BUFFER, 0);
+		FOpenGLDynamicRHI::Get().CachedBindBuffer(GL_PIXEL_UNPACK_BUFFER, 0);
 
 		return;
 	}
@@ -1129,8 +1125,7 @@ void FOpenGLTexture::Unlock(uint32 MipIndex,uint32 ArrayIndex)
 	if (!PixelBuffer->IsLockReadOnly())
 	{
 		// Use a texture stage that's not likely to be used for draws, to avoid waiting
-		FOpenGLContextState& ContextState = FOpenGLDynamicRHI::Get().GetContextStateForCurrentContext();
-		FOpenGLDynamicRHI::Get().CachedSetupTextureStage(ContextState, FOpenGL::GetMaxCombinedTextureImageUnits() - 1, Target, GetResource(), -1, this->GetNumMips());
+		FOpenGLDynamicRHI::Get().CachedSetupTextureStage(FOpenGL::GetMaxCombinedTextureImageUnits() - 1, Target, GetResource(), -1, this->GetNumMips());
 
 		if (GetDesc().IsTextureArray() || GetDesc().IsTexture3D())
 		{
@@ -1181,7 +1176,7 @@ void FOpenGLTexture::Unlock(uint32 MipIndex,uint32 ArrayIndex)
 	// and the next draw will take care of cleaning it up; or
 	// next operation that needs the stage will switch something else in on it.
 
-	CachedBindPixelUnpackBuffer(GL_PIXEL_UNPACK_BUFFER, 0);
+	FOpenGLDynamicRHI::Get().CachedBindBuffer(GL_PIXEL_UNPACK_BUFFER, 0);
 }
 
 
@@ -1214,7 +1209,7 @@ void FOpenGLTexture::RestoreEvictedGLResource(bool bAttemptToRetainMips)
 	const bool bSRGB = EnumHasAnyFlags(this->GetFlags(), TexCreate_SRGB);
 	checkf(EvictionParamsPtr->MipImageData.Num() == this->GetNumMips(), TEXT("EvictionParamsPtr->MipImageData.Num() =%d, this->GetNumMips() = %d"), EvictionParamsPtr->MipImageData.Num(), this->GetNumMips());
 
-	CachedBindPixelUnpackBuffer(GL_PIXEL_UNPACK_BUFFER, 0);
+	FOpenGLDynamicRHI::Get().CachedBindBuffer(GL_PIXEL_UNPACK_BUFFER, 0);
 
 	for (int i = EvictionParamsPtr->MipImageData.Num() - 1; i >= 0; i--)
 	{
@@ -1394,11 +1389,10 @@ void FOpenGLDynamicRHI::RHIGenerateMips(FRHITexture* TextureRHI)
 	VERIFY_GL_SCOPE();
 	RegisterGPUWork(0);
 
-	FOpenGLContextState& ContextState = GetContextStateForCurrentContext();
 	FOpenGLTexture* Texture = ResourceCast(TextureRHI);
 	// Setup the texture on a disused unit
 	// need to figure out how to setup mips properly in no views case
-	CachedSetupTextureStage(ContextState, FOpenGL::GetMaxCombinedTextureImageUnits() - 1, Texture->Target, Texture->GetResource(), -1, Texture->GetNumMips());
+	CachedSetupTextureStage(FOpenGL::GetMaxCombinedTextureImageUnits() - 1, Texture->Target, Texture->GetResource(), -1, Texture->GetNumMips());
 	FOpenGL::GenerateMipmap(Texture->Target);
 }
 
@@ -1541,9 +1535,8 @@ void FOpenGLDynamicRHI::RHIUpdateTexture2D(FRHICommandListBase& RHICmdList, FRHI
 		const FOpenGLTextureFormat& GLFormat = GOpenGLTextureFormats[PixelFormat];
 
 		// Use a texture stage that's not likely to be used for draws, to avoid waiting
-		FOpenGLContextState& ContextState = GetContextStateForCurrentContext();
-		CachedSetupTextureStage(ContextState, FOpenGL::GetMaxCombinedTextureImageUnits() - 1, Texture->Target, Texture->GetResource(), 0, Texture->GetNumMips());
-		CachedBindPixelUnpackBuffer(ContextState, 0);
+		CachedSetupTextureStage(FOpenGL::GetMaxCombinedTextureImageUnits() - 1, Texture->Target, Texture->GetResource(), 0, Texture->GetNumMips());
+		CachedBindBuffer(GL_PIXEL_UNPACK_BUFFER, 0);
 
 		glPixelStorei(GL_UNPACK_ROW_LENGTH, UpdatePitch / FormatInfo.BlockBytes);
 		glPixelStorei(GL_UNPACK_ALIGNMENT, 1);
@@ -1604,9 +1597,8 @@ void FOpenGLDynamicRHI::RHIUpdateTexture3D(FRHICommandListBase& RHICmdList, FRHI
 		FOpenGLTexture* Texture = ResourceCast(TextureRHI);
 
 		// Use a texture stage that's not likely to be used for draws, to avoid waiting
-		FOpenGLContextState& ContextState = GetContextStateForCurrentContext();
-		CachedSetupTextureStage(ContextState, FOpenGL::GetMaxCombinedTextureImageUnits() - 1, Texture->Target, Texture->GetResource(), 0, Texture->GetNumMips());
-		CachedBindPixelUnpackBuffer(ContextState, 0);
+		CachedSetupTextureStage(FOpenGL::GetMaxCombinedTextureImageUnits() - 1, Texture->Target, Texture->GetResource(), 0, Texture->GetNumMips());
+		CachedBindBuffer(GL_PIXEL_UNPACK_BUFFER, 0);
 
 		EPixelFormat PixelFormat = Texture->GetFormat();
 		const FOpenGLTextureFormat& GLFormat = GOpenGLTextureFormats[PixelFormat];
@@ -1672,8 +1664,7 @@ void FOpenGLDynamicRHI::InvalidateTextureResourceInCache(GLuint Resource)
 			}
 		}
 	};
-	InvalidateContextTextureResources(SharedContextState.Textures);
-	InvalidateContextTextureResources(RenderingContextState.Textures);
+	InvalidateContextTextureResources(ContextState.Textures);
 	InvalidateContextTextureResources(PendingState.Textures);
 	
 	TextureMipLimits.Remove(Resource);
@@ -1687,20 +1678,14 @@ void FOpenGLDynamicRHI::InvalidateTextureResourceInCache(GLuint Resource)
 void FOpenGLDynamicRHI::InvalidateUAVResourceInCache(GLuint Resource)
 {
 	VERIFY_GL_SCOPE();
-	int32 NumUAVs = RenderingContextState.UAVs.Num();
+	int32 NumUAVs = ContextState.UAVs.Num();
 	
 	for (int32 UAVIndex = 0; UAVIndex < NumUAVs; ++UAVIndex)
 	{
-		if (SharedContextState.UAVs[UAVIndex].Resource == Resource)
-		{
-			SharedContextState.UAVs[UAVIndex].Format = GL_NONE;
-			SharedContextState.UAVs[UAVIndex].Resource = 0;
-		}
-
-		if (RenderingContextState.UAVs[UAVIndex].Resource == Resource)
+		if (ContextState.UAVs[UAVIndex].Resource == Resource)
 		{
-			RenderingContextState.UAVs[UAVIndex].Format = GL_NONE;
-			RenderingContextState.UAVs[UAVIndex].Resource = 0;
+			ContextState.UAVs[UAVIndex].Format = GL_NONE;
+			ContextState.UAVs[UAVIndex].Resource = 0;
 		}
 
 		if (PendingState.UAVs[UAVIndex].Resource == Resource)
diff --git a/Engine/Source/Runtime/OpenGLDrv/Private/OpenGLUAV.cpp b/Engine/Source/Runtime/OpenGLDrv/Private/OpenGLUAV.cpp
--- a/Engine/Source/Runtime/OpenGLDrv/Private/OpenGLUAV.cpp
+++ b/Engine/Source/Runtime/OpenGLDrv/Private/OpenGLUAV.cpp
@@ -77,7 +77,6 @@ void FOpenGLShaderResourceView::UpdateView()
 
 				// Use a texture stage that's not likely to be used for draws, to avoid waiting
 				FOpenGLDynamicRHI::Get().CachedSetupTextureStage(
-					FOpenGLDynamicRHI::Get().GetContextStateForCurrentContext(),
 					FOpenGL::GetMaxCombinedTextureImageUnits() - 1,
 					GL_TEXTURE_BUFFER,
 					Resource,
@@ -129,8 +128,7 @@ void FOpenGLShaderResourceView::UpdateView()
 		if (Info.Format == PF_X24_G8)
 		{
 			// Use a texture stage that's not likely to be used for draws, to avoid waiting
-			FOpenGLContextState& ContextState = FOpenGLDynamicRHI::Get().GetContextStateForCurrentContext();
-			FOpenGLDynamicRHI::Get().CachedSetupTextureStage(ContextState, FOpenGL::GetMaxCombinedTextureImageUnits() - 1, Target, Resource, LimitMip, Texture->GetNumMips());
+			FOpenGLDynamicRHI::Get().CachedSetupTextureStage(FOpenGL::GetMaxCombinedTextureImageUnits() - 1, Target, Resource, LimitMip, Texture->GetNumMips());
 
 			//set the texture to return the stencil index, and then force the components to match D3D
 			glTexParameteri(Target, GL_DEPTH_STENCIL_TEXTURE_MODE, GL_STENCIL_INDEX);
@@ -208,10 +206,8 @@ void FOpenGLUnorderedAccessView::UpdateView()
 					GLuint TextureID = 0;
 					FOpenGL::GenTextures(1, &TextureID);
 
-					FOpenGLDynamicRHI& RHI = FOpenGLDynamicRHI::Get();
-
 					// Use a texture stage that's not likely to be used for draws, to avoid waiting
-					RHI.CachedSetupTextureStage(RHI.GetContextStateForCurrentContext(), FOpenGL::GetMaxCombinedTextureImageUnits() - 1, GL_TEXTURE_BUFFER, TextureID, -1, 1);
+					FOpenGLDynamicRHI::Get().CachedSetupTextureStage(FOpenGL::GetMaxCombinedTextureImageUnits() - 1, GL_TEXTURE_BUFFER, TextureID, -1, 1);
 					FOpenGL::TexBuffer(GL_TEXTURE_BUFFER, GLFormat.InternalFormat[0], Buffer->Resource);
 
 					// No need to restore texture stage; leave it like this,
diff --git a/Engine/Source/Runtime/OpenGLDrv/Private/OpenGLUniformBuffer.cpp b/Engine/Source/Runtime/OpenGLDrv/Private/OpenGLUniformBuffer.cpp
--- a/Engine/Source/Runtime/OpenGLDrv/Private/OpenGLUniformBuffer.cpp
+++ b/Engine/Source/Runtime/OpenGLDrv/Private/OpenGLUniformBuffer.cpp
@@ -368,7 +368,7 @@ static void SuballocateUBO( uint32 Size, GLuint& Resource, uint32& Offset, uint8
 
 	FOpenGL::GenBuffers( 1, &Pool.Resource);
 	
-	CachedBindUniformBuffer(Pool.Resource);
+	FOpenGLDynamicRHI::Get().CachedBindUniformBuffer(Pool.Resource);
 
 	if (FOpenGL::SupportsBufferStorage() && OpenGLConsoleVariables::bUBODirectWrite)
 	{
@@ -586,7 +586,7 @@ static FUniformBufferRHIRef CreateUniformBuffer(const void* Contents, const FRHI
 				GLuint AllocatedResource = 0;
 				uint32 OffsetInBuffer = 0;
 				FOpenGL::GenBuffers(1, &AllocatedResource);
-				::CachedBindUniformBuffer(AllocatedResource);
+				FOpenGLDynamicRHI::Get().CachedBindUniformBuffer(AllocatedResource);
 				glBufferData(GL_UNIFORM_BUFFER, AllocatedSize, NULL, bStreamDraw ? GL_STREAM_DRAW : GL_STATIC_DRAW);
 				NewUniformBuffer->SetGLUniformBufferParams(AllocatedResource, OffsetInBuffer, nullptr, AllocatedSize, EmulatedUniformDataRef, bStreamDraw);
 			});
@@ -748,7 +748,7 @@ FUniformBufferRHIRef FOpenGLDynamicRHI::RHICreateUniformBuffer(const void* Conte
 					RHICmdList.EnqueueLambda([AllocatedResource](FRHICommandListImmediate&)
 					{
 						VERIFY_GL_SCOPE();
-						::CachedBindUniformBuffer(AllocatedResource);
+						FOpenGLDynamicRHI::Get().CachedBindUniformBuffer(AllocatedResource);
 					});
 				}
 			}
@@ -796,7 +796,7 @@ void UpdateUniformBufferContents(FOpenGLUniformBuffer* UniformBuffer, const void
 		}
 		else
 		{
-			::CachedBindUniformBuffer(UniformBuffer->Resource);
+			FOpenGLDynamicRHI::Get().CachedBindUniformBuffer(UniformBuffer->Resource);
 			FOpenGL::BufferSubData(GL_UNIFORM_BUFFER, 0, ConstantBufferSize, Contents);
 		}
 	}
diff --git a/Engine/Source/Runtime/OpenGLDrv/Private/OpenGLViewport.cpp b/Engine/Source/Runtime/OpenGLDrv/Private/OpenGLViewport.cpp
--- a/Engine/Source/Runtime/OpenGLDrv/Private/OpenGLViewport.cpp
+++ b/Engine/Source/Runtime/OpenGLDrv/Private/OpenGLViewport.cpp
@@ -120,20 +120,8 @@ void FOpenGLDynamicRHI::RHIBeginDrawingViewport(FRHIViewport* ViewportRHI, FRHIT
 	check(!DrawingViewport);
 	DrawingViewport = Viewport;
 
-	bRevertToSharedContextAfterDrawingViewport = false;
-	EOpenGLCurrentContext CurrentContext = PlatformOpenGLCurrentContext( PlatformDevice );
-	if( CurrentContext != CONTEXT_Rendering )
-	{
-		check(CurrentContext == CONTEXT_Shared);
-		check(!bIsRenderingContextAcquired || !GUseThreadedRendering);
-
-		bRevertToSharedContextAfterDrawingViewport = true;
-		PlatformRenderingContextSetup(PlatformDevice);
-		CachedContextState = nullptr;
-	}
-
 	// Set the render target and viewport.
-	if( RenderTarget )
+	if (RenderTarget)
 	{
 		FRHIRenderTargetView RTV(RenderTarget, ERenderTargetLoadAction::ELoad);
 		SetRenderTargets(1, &RTV, nullptr);
@@ -164,8 +152,6 @@ void FOpenGLDynamicRHI::RHIEndDrawingViewport(FRHIViewport* ViewportRHI,bool bPr
 
 		FOpenGLTexture* BackBuffer = Viewport->GetBackBuffer();
 
-		FOpenGLContextState& ContextState = GetContextStateForCurrentContext();
-
 		if (ContextState.bScissorEnabled)
 		{
 			ContextState.bScissorEnabled = false;
@@ -181,7 +167,7 @@ void FOpenGLDynamicRHI::RHIEndDrawingViewport(FRHIViewport* ViewportRHI,bool bPr
 		);
 
 		// Always consider the Framebuffer in the rendering context dirty after the blit
-		RenderingContextState.Framebuffer = -1;
+		ContextState.Framebuffer = -1;
 
 		DrawingViewport = NULL;
 
@@ -211,14 +197,6 @@ void FOpenGLDynamicRHI::RHIEndDrawingViewport(FRHIViewport* ViewportRHI,bool bPr
 			GInputLatencyTimer.DeltaTime = EndTime - GInputLatencyTimer.StartTime;
 			GInputLatencyTimer.RenderThreadTrigger = false;
 		}
-
-		if (bRevertToSharedContextAfterDrawingViewport)
-		{
-			PlatformSharedContextSetup(PlatformDevice);
-			CachedContextState = nullptr;
-
-			bRevertToSharedContextAfterDrawingViewport = false;
-		}
 	}
 
 	EndFrameTick();
@@ -259,13 +237,9 @@ FOpenGLViewport::FOpenGLViewport(FOpenGLDynamicRHI* InOpenGLRHI,void* InWindowHa
 
 	OpenGLRHI->Viewports.Add(this);
 
-	check(PlatformOpenGLCurrentContext(OpenGLRHI->PlatformDevice) == CONTEXT_Shared);
-
 	OpenGLContext = PlatformCreateOpenGLContext(OpenGLRHI->PlatformDevice, InWindowHandle);
 	Resize(InSizeX, InSizeY, bInIsFullscreen);
 
-	check(PlatformOpenGLCurrentContext(OpenGLRHI->PlatformDevice) == CONTEXT_Shared);
-
 	ENQUEUE_RENDER_COMMAND(CreateFrameSyncEvent)([this](FRHICommandListImmediate& RHICmdList)
 	{
 		RHICmdList.EnqueueLambda([this](FRHICommandListImmediate&)
diff --git a/Engine/Source/Runtime/OpenGLDrv/Private/Windows/WindowsOpenGLPlatform.cpp b/Engine/Source/Runtime/OpenGLDrv/Private/Windows/WindowsOpenGLPlatform.cpp
--- a/Engine/Source/Runtime/OpenGLDrv/Private/Windows/WindowsOpenGLPlatform.cpp
+++ b/Engine/Source/Runtime/OpenGLDrv/Private/Windows/WindowsOpenGLPlatform.cpp
@@ -255,7 +255,8 @@ void PlatformReleaseOpenGLContext(FPlatformOpenGLDevice* Device, FPlatformOpenGL
 /** Platform specific OpenGL device. */
 struct FPlatformOpenGLDevice
 {
-	FPlatformOpenGLContext	SharedContext;
+	static FPlatformOpenGLDevice* Singleton;
+
 	FPlatformOpenGLContext	RenderingContext;
 	TArray<FPlatformOpenGLContext*>	ViewportContexts;
 	bool					TargetDirty;
@@ -266,6 +267,9 @@ struct FPlatformOpenGLDevice
 	FPlatformOpenGLDevice()
 		: TargetDirty(true)
 	{
+		check(!Singleton);
+		Singleton = this;
+
 		extern void InitDebugContext();
 		ContextUsageGuard = new FCriticalSection;
 
@@ -276,31 +280,17 @@ struct FPlatformOpenGLDevice
 		// Need to call this before we set the debug callback, otherwise if we're not running under RD, the debug extension will assert (invalid enum)
 		GRunningUnderRenderDoc = glIsEnabled(GL_DEBUG_TOOL_EXT) != GL_FALSE;
 
-		PlatformCreateDummyGLWindow(&SharedContext);
-		PlatformCreateOpenGLContextCore(&SharedContext, MajorVersion, MinorVersion, NULL);
-		check(SharedContext.OpenGLContext)
-		{
-			FScopeContext ScopeContext(&SharedContext);
-			InitDebugContext();
-			glGenVertexArrays(1, &SharedContext.VertexArrayObject);
-			glBindVertexArray(SharedContext.VertexArrayObject);
-			InitDefaultGLContextState();
-			glGenFramebuffers(1, &SharedContext.ViewportFramebuffer);
-		}
-
 		PlatformCreateDummyGLWindow(&RenderingContext);
-		PlatformCreateOpenGLContextCore(&RenderingContext, MajorVersion, MinorVersion, SharedContext.OpenGLContext);
-		check(RenderingContext.OpenGLContext);
-		{
-			FScopeContext ScopeContext(&RenderingContext);
-			InitDebugContext();
-			glGenVertexArrays(1, &RenderingContext.VertexArrayObject);
-			glBindVertexArray(RenderingContext.VertexArrayObject);
-			InitDefaultGLContextState();
-			glGenFramebuffers(1, &RenderingContext.ViewportFramebuffer);
-		}
+		PlatformCreateOpenGLContextCore(&RenderingContext, MajorVersion, MinorVersion, nullptr);
+		ContextMakeCurrent(RenderingContext.DeviceContext, RenderingContext.OpenGLContext);
 
-		ContextMakeCurrent(SharedContext.DeviceContext, SharedContext.OpenGLContext);
+		check(RenderingContext.OpenGLContext);
+		
+		InitDebugContext();
+		glGenVertexArrays(1, &RenderingContext.VertexArrayObject);
+		glBindVertexArray(RenderingContext.VertexArrayObject);
+		InitDefaultGLContextState();
+		glGenFramebuffers(1, &RenderingContext.ViewportFramebuffer);
 	}
 
 	~FPlatformOpenGLDevice()
@@ -308,14 +298,17 @@ struct FPlatformOpenGLDevice
 		check(ViewportContexts.Num() == 0);
 
 		ContextMakeCurrent(NULL, NULL);
-
 		PlatformReleaseOpenGLContext(this, &RenderingContext);
-		PlatformReleaseOpenGLContext(this, &SharedContext);
 
 		delete ContextUsageGuard;
+
+		check(Singleton == this);
+		Singleton = nullptr;
 	}
 };
 
+FPlatformOpenGLDevice* FPlatformOpenGLDevice::Singleton = nullptr;
+
 FPlatformOpenGLDevice* PlatformCreateOpenGLDevice()
 {
 	return new FPlatformOpenGLDevice;
@@ -351,7 +344,7 @@ FPlatformOpenGLContext* PlatformCreateOpenGLContext(FPlatformOpenGLDevice* Devic
 	int MinorVersion = 0;
 	GetOpenGLVersionForCoreProfile(MajorVersion, MinorVersion);
 
-	PlatformCreateOpenGLContextCore(Context, MajorVersion, MinorVersion, Device->SharedContext.OpenGLContext);
+	PlatformCreateOpenGLContextCore(Context, MajorVersion, MinorVersion, Device->RenderingContext.OpenGLContext);
 	check(Context->OpenGLContext);
 	{
 		FScopeContext Scope(Context);
@@ -407,7 +400,7 @@ void PlatformReleaseOpenGLContext(FPlatformOpenGLDevice* Device, FPlatformOpenGL
 	ReleaseDC(Context->WindowHandle, Context->DeviceContext);
 	Context->DeviceContext = NULL;
 
-	if (PlatformOpenGLCurrentContext(Device) == CONTEXT_Rendering && bSharedDC)
+	if (PlatformOpenGLThreadHasRenderingContext() && bSharedDC)
 	{
 		// The rendering context has been made current using the DC of the now destroyed context. Since this DC has been released the current context will be invalid.
 		// To properly set the rendering context we must make current here with it's own DC.
@@ -549,15 +542,6 @@ void PlatformRenderingContextSetup(FPlatformOpenGLDevice* Device)
 	}
 }
 
-void PlatformSharedContextSetup(FPlatformOpenGLDevice* Device)
-{
-	check(Device && Device->SharedContext.DeviceContext && Device->SharedContext.OpenGLContext);
-
-	// no need to glFlush() on Windows, it does flush by itself before switching contexts
-	ContextMakeCurrent(Device->SharedContext.DeviceContext, Device->SharedContext.OpenGLContext);
-
-}
-
 void PlatformNULLContextSetup()
 {
 	if (wglGetCurrentDC())
@@ -824,45 +808,18 @@ bool PlatformInitOpenGL()
 	return bOpenGLSupported;
 }
 
-bool PlatformOpenGLContextValid()
-{
-	return(GetCurrentContext() != NULL);
-}
-
 int32 PlatformGlGetError()
 {
 	return glGetError();
 }
 
-EOpenGLCurrentContext PlatformOpenGLCurrentContext(FPlatformOpenGLDevice* Device)
+bool PlatformOpenGLThreadHasRenderingContext()
 {
-	HGLRC Context = GetCurrentContext();
-
-	if (Context == Device->RenderingContext.OpenGLContext)	// most common case
-	{
-		return CONTEXT_Rendering;
-	}
-	else if (Context == Device->SharedContext.OpenGLContext)
-	{
-		return CONTEXT_Shared;
-	}
-	else if (Context)
-	{
-		return CONTEXT_Other;
-	}
-	else
-	{
-		return CONTEXT_Invalid;
-	}
+	return GetCurrentContext() == FPlatformOpenGLDevice::Singleton->RenderingContext.OpenGLContext;
 }
 
 // =============================================================
 
-bool PlatformContextIsCurrent(uint64 QueryContext)
-{
-	return (uint64)GetCurrentContext() == QueryContext;
-}
-
 FOpenGLTexture* PlatformCreateBuiltinBackBuffer(FOpenGLDynamicRHI* OpenGLRHI, uint32 SizeX, uint32 SizeY)
 {
 	if (FOpenGL::IsAndroidGLESCompatibilityModeEnabled())
diff --git a/Engine/Source/Runtime/OpenGLDrv/Public/OpenGLUtil.h b/Engine/Source/Runtime/OpenGLDrv/Public/OpenGLUtil.h
--- a/Engine/Source/Runtime/OpenGLDrv/Public/OpenGLUtil.h
+++ b/Engine/Source/Runtime/OpenGLDrv/Public/OpenGLUtil.h
@@ -49,13 +49,13 @@ int32 GetOGLDebugOutputLevel();
 */
 GLenum GetOpenGLCubeFace(ECubeFace Face);
 
-extern bool PlatformOpenGLContextValid();
+extern bool PlatformOpenGLThreadHasRenderingContext();
 
 #if ENABLE_VERIFY_GL_THREAD
 	#define CHECK_EXPECTED_GL_THREAD() \
-		if (!PlatformOpenGLContextValid()) \
+		if (!PlatformOpenGLThreadHasRenderingContext()) \
 		{ \
-			UE_LOG(LogRHI, Fatal, TEXT("Potential use of GL context from incorrect thread. [ValidContext = %d] && [IsInGameThread() = %d, IsInRenderingThread() = %d, IsInRHIThread() = %d, IsRunningRHIInSeparateThread() = %d]"), PlatformOpenGLContextValid(), IsInGameThread(), IsInRenderingThread(), IsInRHIThread(), IsRunningRHIInSeparateThread()); \
+			UE_LOG(LogRHI, Fatal, TEXT("Potential use of GL context from incorrect thread. [IsInGameThread() = %d, IsInRenderingThread() = %d, IsInRHIThread() = %d, IsRunningRHIInSeparateThread() = %d]"), IsInGameThread(), IsInRenderingThread(), IsInRHIThread(), IsRunningRHIInSeparateThread()); \
 		}
 #else
 	#define CHECK_EXPECTED_GL_THREAD() 
@@ -99,7 +99,7 @@ extern bool PlatformOpenGLContextValid();
 
 		void CheckForErrors(int32 BeforeOrAfter)
 		{
-			check(PlatformOpenGLContextValid());
+			check(PlatformOpenGLThreadHasRenderingContext());
 
 			GLenum ErrorCode = PlatformGlGetError();
 			if (ErrorCode != GL_NO_ERROR)
diff --git a/Engine/Source/Runtime/RenderCore/Private/RenderingThread.cpp b/Engine/Source/Runtime/RenderCore/Private/RenderingThread.cpp
--- a/Engine/Source/Runtime/RenderCore/Private/RenderingThread.cpp
+++ b/Engine/Source/Runtime/RenderCore/Private/RenderingThread.cpp
@@ -381,7 +381,6 @@ public:
 	FRenderingThread()
 	{
 		TaskGraphBoundSyncEvent	= FPlatformProcess::GetSynchEventFromPool(true);
-		RHIFlushResources();
 	}
 
 	virtual ~FRenderingThread()
