Assign OpenGL shader debug name

diff --git a/Engine/Source/Runtime/OpenGLDrv/Private/OpenGLShaders.cpp b/Engine/Source/Runtime/OpenGLDrv/Private/OpenGLShaders.cpp
--- a/Engine/Source/Runtime/OpenGLDrv/Private/OpenGLShaders.cpp
+++ b/Engine/Source/Runtime/OpenGLDrv/Private/OpenGLShaders.cpp
@@ -1049,6 +1049,10 @@ FOpenGLShader::FOpenGLShader(TArrayView<const uint8> Code, const FSHAHash& Hash,
 	GlslCodeString = GlslCode.GetData();
 #endif
 
+#if RHI_INCLUDE_SHADER_DEBUG_DATA
+	RHIShader->Debug.ShaderName = ShaderCode.FindOptionalData(FShaderCodeName::Key);
+#endif
+
 	// The shader is compiled when we link program
 }
 
