Workaround for GLSL shaders that use 'precise' qualifier on xclipse

diff --git a/Engine/Source/Runtime/OpenGLDrv/Private/Android/AndroidOpenGL.cpp b/Engine/Source/Runtime/OpenGLDrv/Private/Android/AndroidOpenGL.cpp
--- a/Engine/Source/Runtime/OpenGLDrv/Private/Android/AndroidOpenGL.cpp
+++ b/Engine/Source/Runtime/OpenGLDrv/Private/Android/AndroidOpenGL.cpp
@@ -588,6 +588,7 @@ void FAndroidOpenGL::ProcessExtensions(const FString& ExtensionsString)
 	const bool bIsPoverVRBased = RendererString.Contains(TEXT("PowerVR"));
 	const bool bIsAdrenoBased = RendererString.Contains(TEXT("Adreno"));
 	const bool bIsMaliBased = RendererString.Contains(TEXT("Mali"));
+	const bool bIsSamsungBased = RendererString.Contains(TEXT("Xclipse"));
 
 	if (bIsPoverVRBased)
 	{
@@ -634,6 +635,16 @@ void FAndroidOpenGL::ProcessExtensions(const FString& ExtensionsString)
 		bRequiresReadOnlyBuffersWorkaround = true;
 	}
 
+	if (bIsSamsungBased)
+	{
+		FString AndroidVersion = FAndroidMisc::GetAndroidVersion();
+		if (AndroidVersion.Contains(TEXT("14")))
+		{
+			bRequiresPreciseQualifierWorkaround = true;
+			UE_LOG(LogRHI, Log, TEXT("Disable \'precise\' qualifier for [Android: %s, GPU: %s]"), *AndroidVersion, *RendererString);
+		}
+	}
+
 	// Disable ASTC if requested by device profile
 	static const auto CVarDisableASTC = IConsoleManager::Get().FindTConsoleVariableDataInt(TEXT("r.Android.DisableASTCSupport"));
 	if (bSupportsASTC && CVarDisableASTC->GetValueOnAnyThread())
diff --git a/Engine/Source/Runtime/OpenGLDrv/Private/OpenGLES.cpp b/Engine/Source/Runtime/OpenGLDrv/Private/OpenGLES.cpp
--- a/Engine/Source/Runtime/OpenGLDrv/Private/OpenGLES.cpp
+++ b/Engine/Source/Runtime/OpenGLDrv/Private/OpenGLES.cpp
@@ -71,6 +71,9 @@ bool FOpenGLES::bRequiresDisabledEarlyFragmentTests = false;
 /* This is a workaround for a Mali bug where read-only buffers do not work when passed to functions*/
 bool FOpenGLES::bRequiresReadOnlyBuffersWorkaround = false;
 
+/* This is a workaround for Samsung xclipse that fails to compile shaders that use the 'precise' qualifier*/
+bool FOpenGLES::bRequiresPreciseQualifierWorkaround = false;
+
 /* This is to avoid a bug in Adreno drivers that define GL_ARM_shader_framebuffer_fetch_depth_stencil even when device does not support this extension  */
 bool FOpenGLES::bRequiresARMShaderFramebufferFetchDepthStencilUndef = false;
 
diff --git a/Engine/Source/Runtime/OpenGLDrv/Private/OpenGLES.h b/Engine/Source/Runtime/OpenGLDrv/Private/OpenGLES.h
--- a/Engine/Source/Runtime/OpenGLDrv/Private/OpenGLES.h
+++ b/Engine/Source/Runtime/OpenGLDrv/Private/OpenGLES.h
@@ -78,6 +78,7 @@ struct FOpenGLES : public FOpenGLBase
 	static FORCEINLINE bool HasBinaryProgramRetrievalFailed() { return bBinaryProgramRetrievalFailed; }
 	static FORCEINLINE bool RequiresDisabledEarlyFragmentTests() { return bRequiresDisabledEarlyFragmentTests; }
 	static FORCEINLINE bool RequiresReadOnlyBuffersWorkaround() { return bRequiresReadOnlyBuffersWorkaround; }
+	static FORCEINLINE bool RequiresPreciseQualifierWorkaround() { return bRequiresPreciseQualifierWorkaround; }
 	static FORCEINLINE bool RequiresARMShaderFramebufferFetchDepthStencilUndef() { return bRequiresARMShaderFramebufferFetchDepthStencilUndef; }
 
 	// Adreno doesn't support HALF_FLOAT
@@ -705,6 +706,9 @@ public:
 	/* This is a workaround for a Mali bug where read-only buffers do not work when passed to functions*/
 	static bool bRequiresReadOnlyBuffersWorkaround;
 
+	/* This is a workaround for Samsung xclipse that fails to compile shaders that use the 'precise' qualifier*/
+	static bool bRequiresPreciseQualifierWorkaround;
+
 	/* This is to avoid a bug in Adreno drivers that define GL_ARM_shader_framebuffer_fetch_depth_stencil even when device does not support this extension  */
 	static bool bRequiresARMShaderFramebufferFetchDepthStencilUndef;
 
diff --git a/Engine/Source/Runtime/OpenGLDrv/Private/OpenGLShaders.cpp b/Engine/Source/Runtime/OpenGLDrv/Private/OpenGLShaders.cpp
--- a/Engine/Source/Runtime/OpenGLDrv/Private/OpenGLShaders.cpp
+++ b/Engine/Source/Runtime/OpenGLDrv/Private/OpenGLShaders.cpp
@@ -779,6 +779,7 @@ void OPENGLDRV_API GetCurrentOpenGLShaderDeviceCapabilities(FOpenGLShaderDeviceC
 		Capabilities.bRequiresARMShaderFramebufferFetchDepthStencilUndef = false;
 		Capabilities.bRequiresReadOnlyBuffersWorkaround = false;
 		Capabilities.MaxVaryingVectors = FOpenGL::GetMaxVaryingVectors();
+		Capabilities.bRequiresPreciseQualifierWorkaround = false;
 	}
 
 #elif PLATFORM_ANDROID
@@ -788,6 +789,7 @@ void OPENGLDRV_API GetCurrentOpenGLShaderDeviceCapabilities(FOpenGLShaderDeviceC
 		Capabilities.bRequiresARMShaderFramebufferFetchDepthStencilUndef = FOpenGL::RequiresARMShaderFramebufferFetchDepthStencilUndef();
 		Capabilities.MaxVaryingVectors = FOpenGL::GetMaxVaryingVectors();
 		Capabilities.bRequiresDisabledEarlyFragmentTests = FOpenGL::RequiresDisabledEarlyFragmentTests();
+		Capabilities.bRequiresPreciseQualifierWorkaround = FOpenGL::RequiresPreciseQualifierWorkaround();
 #elif PLATFORM_IOS
 	Capabilities.TargetPlatform = EOpenGLShaderTargetPlatform::OGLSTP_iOS;
 #else
@@ -833,6 +835,12 @@ void OPENGLDRV_API GLSLToDeviceCompatibleGLSL(FAnsiCharArray& GlslCodeOriginal,
 		AppendCString(GlslCode, "#define fma(A, B, C) ((A) * (B) + (C))\n");
 	}
 
+	if (Capabilities.bRequiresPreciseQualifierWorkaround)
+	{
+		// Disable use of 'precise' qualifier
+		AppendCString(GlslCode, "#define precise\n");
+	}
+
 	if (Capabilities.bRequiresReadOnlyBuffersWorkaround)
 	{
 		ReplaceCString(GlslCodeOriginal, "readonly buffer", "buffer");
diff --git a/Engine/Source/Runtime/OpenGLDrv/Private/OpenGLShaders.h b/Engine/Source/Runtime/OpenGLDrv/Private/OpenGLShaders.h
--- a/Engine/Source/Runtime/OpenGLDrv/Private/OpenGLShaders.h
+++ b/Engine/Source/Runtime/OpenGLDrv/Private/OpenGLShaders.h
@@ -29,6 +29,7 @@ struct FOpenGLShaderDeviceCapabilities
 	GLint MaxVaryingVectors;
 	bool bRequiresDisabledEarlyFragmentTests;
 	bool bRequiresReadOnlyBuffersWorkaround;
+	bool bRequiresPreciseQualifierWorkaround;
 };
 
 /**
