Fixed: Flickering decals on Android OpenGL devices with mobile deferred shading enabled

diff --git a/Engine/Source/Runtime/OpenGLDrv/Private/OpenGLCommands.cpp b/Engine/Source/Runtime/OpenGLDrv/Private/OpenGLCommands.cpp
--- a/Engine/Source/Runtime/OpenGLDrv/Private/OpenGLCommands.cpp
+++ b/Engine/Source/Runtime/OpenGLDrv/Private/OpenGLCommands.cpp
@@ -1633,7 +1633,8 @@ void FOpenGLDynamicRHI::SetRenderTargets(
 	}
 
 	FOpenGLTexture* NewDepthStencilRT = ResourceCast(NewDepthStencilTargetRHI ? NewDepthStencilTargetRHI->Texture : nullptr);
-	
+	FExclusiveDepthStencil NewDepthStencilAccess = NewDepthStencilTargetRHI ? NewDepthStencilTargetRHI->GetDepthStencilAccess() : FExclusiveDepthStencil();
+		
 	PendingState.DepthStencil = NewDepthStencilRT;
 	PendingState.StencilStoreAction = NewDepthStencilTargetRHI ? NewDepthStencilTargetRHI->GetStencilStoreAction() : ERenderTargetStoreAction::ENoAction;
 	PendingState.DepthTargetWidth   = NewDepthStencilRT ? NewDepthStencilRT->GetDesc().Extent.X : 0u;
@@ -1648,7 +1649,7 @@ void FOpenGLDynamicRHI::SetRenderTargets(
 		return;
 	}
 
-	PendingState.Framebuffer = GetOpenGLFramebuffer(NumSimultaneousRenderTargets, PendingState.RenderTargets, PendingState.RenderTargetArrayIndex, PendingState.RenderTargetMipmapLevels, PendingState.DepthStencil, PendingState.NumRenderingSamples);
+	PendingState.Framebuffer = GetOpenGLFramebuffer(NumSimultaneousRenderTargets, PendingState.RenderTargets, PendingState.RenderTargetArrayIndex, PendingState.RenderTargetMipmapLevels, PendingState.DepthStencil, NewDepthStencilAccess, PendingState.NumRenderingSamples);
 	PendingState.bFramebufferSetupInvalid = false;
 
 	if (PendingState.FirstNonzeroRenderTarget != -1)
diff --git a/Engine/Source/Runtime/OpenGLDrv/Private/OpenGLDrv.h b/Engine/Source/Runtime/OpenGLDrv/Private/OpenGLDrv.h
--- a/Engine/Source/Runtime/OpenGLDrv/Private/OpenGLDrv.h
+++ b/Engine/Source/Runtime/OpenGLDrv/Private/OpenGLDrv.h
@@ -658,7 +658,7 @@ public:
 	virtual void RHIPostExternalCommandsReset() final override;
 
 	GLuint GetOpenGLFramebuffer(uint32 NumSimultaneousRenderTargets, FOpenGLTexture** RenderTargets, const uint32* ArrayIndices, const uint32* MipmapLevels, FOpenGLTexture* DepthStencilTarget);
-	GLuint GetOpenGLFramebuffer(uint32 NumSimultaneousRenderTargets, FOpenGLTexture** RenderTargets, const uint32* ArrayIndices, const uint32* MipmapLevels, FOpenGLTexture* DepthStencilTarget, int32 NumRenderingSamples);
+	GLuint GetOpenGLFramebuffer(uint32 NumSimultaneousRenderTargets, FOpenGLTexture** RenderTargets, const uint32* ArrayIndices, const uint32* MipmapLevels, FOpenGLTexture* DepthStencilTarget, FExclusiveDepthStencil DepthStencilAccess, int32 NumRenderingSamples);
 	
 	void ResolveTexture(FOpenGLTexture* Texture, uint32 MipIndex, uint32 ArrayIndex);
 
diff --git a/Engine/Source/Runtime/OpenGLDrv/Private/OpenGLRenderTarget.cpp b/Engine/Source/Runtime/OpenGLDrv/Private/OpenGLRenderTarget.cpp
--- a/Engine/Source/Runtime/OpenGLDrv/Private/OpenGLRenderTarget.cpp
+++ b/Engine/Source/Runtime/OpenGLDrv/Private/OpenGLRenderTarget.cpp
@@ -43,9 +43,11 @@ public:
 		const uint32* InRenderTargetArrayIndices,
 		const uint32* InRenderTargetMipmapLevels,
 		FOpenGLTexture* InDepthStencilTarget,
+		FExclusiveDepthStencil DepthStencilAccess,
 		int32 InNumRenderingSamples
 		)
 		:	DepthStencilTarget(InDepthStencilTarget)
+		,	DepthStencilAccessIndex(DepthStencilAccess.GetIndex())
 		,	NumRenderingSamples(InNumRenderingSamples)
 	{
 		uint32 RenderTargetIndex;
@@ -74,6 +76,7 @@ public:
 		return	
 			!FMemory::Memcmp(A.RenderTargets, B.RenderTargets, sizeof(A.RenderTargets) ) && 
 			A.DepthStencilTarget == B.DepthStencilTarget &&
+			A.DepthStencilAccessIndex == B.DepthStencilAccessIndex &&
 			A.NumRenderingSamples == B.NumRenderingSamples;
 	}
 
@@ -84,7 +87,7 @@ public:
 	*/
 	friend uint32 GetTypeHash(const FOpenGLFramebufferKey &Key)
 	{
-		return FCrc::MemCrc_DEPRECATED(Key.RenderTargets, sizeof(Key.RenderTargets)) ^ GetTypeHash(Key.DepthStencilTarget) ^ GetTypeHash(Key.NumRenderingSamples);
+		return FCrc::MemCrc_DEPRECATED(Key.RenderTargets, sizeof(Key.RenderTargets)) ^ GetTypeHash(Key.DepthStencilTarget) ^ GetTypeHash(Key.DepthStencilAccessIndex) ^ GetTypeHash(Key.NumRenderingSamples);
 	}
 
 	const FOpenGLTexture* GetRenderTarget( int32 Index ) const { return RenderTargets[Index].Texture; }
@@ -95,6 +98,7 @@ private:
 
 	RenderTargetInfo RenderTargets[MaxSimultaneousRenderTargets];
 	FOpenGLTexture* DepthStencilTarget;
+	uint32 DepthStencilAccessIndex; // we don't want to merge render passes that have different access to a depth buffer
 	int32 NumRenderingSamples; // MSAA on tile
 };
 
@@ -131,16 +135,16 @@ static FOpenGLFramebufferCache& GetOpenGLFramebufferCache()
 GLuint FOpenGLDynamicRHI::GetOpenGLFramebuffer(uint32 NumSimultaneousRenderTargets, FOpenGLTexture** RenderTargets, const uint32* ArrayIndices, const uint32* MipmapLevels, FOpenGLTexture* DepthStencilTarget)
 {
 	const int32 NumRenderingSamples = 1;
-	return GetOpenGLFramebuffer(NumSimultaneousRenderTargets, RenderTargets, ArrayIndices, MipmapLevels, DepthStencilTarget, NumRenderingSamples);
+	return GetOpenGLFramebuffer(NumSimultaneousRenderTargets, RenderTargets, ArrayIndices, MipmapLevels, DepthStencilTarget, FExclusiveDepthStencil(), NumRenderingSamples);
 }
 
-GLuint FOpenGLDynamicRHI::GetOpenGLFramebuffer(uint32 NumSimultaneousRenderTargets, FOpenGLTexture** RenderTargets, const uint32* ArrayIndices, const uint32* MipmapLevels, FOpenGLTexture* DepthStencilTarget, int32 NumRenderingSamples)
+GLuint FOpenGLDynamicRHI::GetOpenGLFramebuffer(uint32 NumSimultaneousRenderTargets, FOpenGLTexture** RenderTargets, const uint32* ArrayIndices, const uint32* MipmapLevels, FOpenGLTexture* DepthStencilTarget, FExclusiveDepthStencil DepthStencilAccess, int32 NumRenderingSamples)
 {
 	VERIFY_GL_SCOPE();
 
 	check( NumSimultaneousRenderTargets <= MaxSimultaneousRenderTargets );
 
-	uint32 FramebufferRet = GetOpenGLFramebufferCache().FindRef(FOpenGLFramebufferKey(NumSimultaneousRenderTargets, RenderTargets, ArrayIndices, MipmapLevels, DepthStencilTarget, NumRenderingSamples));
+	uint32 FramebufferRet = GetOpenGLFramebufferCache().FindRef(FOpenGLFramebufferKey(NumSimultaneousRenderTargets, RenderTargets, ArrayIndices, MipmapLevels, DepthStencilTarget, DepthStencilAccess, NumRenderingSamples));
 	if( FramebufferRet > 0 )
 	{
 		// Found and is valid. We never store zero as a result, increasing all results by 1 to avoid range overlap.
@@ -206,7 +210,7 @@ GLuint FOpenGLDynamicRHI::GetOpenGLFramebuffer(uint32 NumSimultaneousRenderTarge
 		FOpenGL::ReadBuffer(GL_NONE);
 		FOpenGL::DrawBuffer(GL_COLOR_ATTACHMENT0);
 
-		GetOpenGLFramebufferCache().Add(FOpenGLFramebufferKey(NumSimultaneousRenderTargets, RenderTargets, ArrayIndices, MipmapLevels, DepthStencilTarget, NumRenderingSamples), Framebuffer + 1);
+		GetOpenGLFramebufferCache().Add(FOpenGLFramebufferKey(NumSimultaneousRenderTargets, RenderTargets, ArrayIndices, MipmapLevels, DepthStencilTarget, DepthStencilAccess, NumRenderingSamples), Framebuffer + 1);
 		
 		return Framebuffer;
 	}
@@ -363,7 +367,7 @@ GLuint FOpenGLDynamicRHI::GetOpenGLFramebuffer(uint32 NumSimultaneousRenderTarge
 	
 	FOpenGL::CheckFrameBuffer();
 
-	GetOpenGLFramebufferCache().Add(FOpenGLFramebufferKey(NumSimultaneousRenderTargets, RenderTargets, ArrayIndices, MipmapLevels, DepthStencilTarget, NumRenderingSamples), Framebuffer+1);
+	GetOpenGLFramebufferCache().Add(FOpenGLFramebufferKey(NumSimultaneousRenderTargets, RenderTargets, ArrayIndices, MipmapLevels, DepthStencilTarget, DepthStencilAccess, NumRenderingSamples), Framebuffer+1);
 
 	return Framebuffer;
 }
