Fixing a bug with GLES texture eviction LRU not touching/updating recently used textures and potentially evicting textures too early.
Do not allow gl textures to be evicted after modification with Lock or texture copy operations.

diff --git a/Engine/Source/Runtime/OpenGLDrv/Private/OpenGLResources.h b/Engine/Source/Runtime/OpenGLDrv/Private/OpenGLResources.h
--- a/Engine/Source/Runtime/OpenGLDrv/Private/OpenGLResources.h
+++ b/Engine/Source/Runtime/OpenGLDrv/Private/OpenGLResources.h
@@ -1060,13 +1060,14 @@ public:
 	void RestoreEvictedGLResource(bool bAttemptToRetainMips);
 	bool CanBeEvicted();
 	void TryEvictGLResource();
+	void RemoveEvictionData();
 
 private:
 	static void UpdateTextureStats(FOpenGLTexture* Texture, bool bAllocating);
 
 	void TryRestoreGLResource()
 	{
-		if (EvictionParamsPtr.IsValid() && !EvictionParamsPtr->bHasRestored)
+		if (EvictionParamsPtr.IsValid())
 		{
 			VERIFY_GL_SCOPE();
 			if (!EvictionParamsPtr->bHasRestored)
@@ -1075,8 +1076,10 @@ private:
 			}
 			else
 			{
-				check(CanBeEvicted());
-				FTextureEvictionLRU::Get().Touch(this);
+				if(CanBeEvicted())
+				{
+					FTextureEvictionLRU::Get().Touch(this);
+				}
 			}
 		}
 	}
diff --git a/Engine/Source/Runtime/OpenGLDrv/Private/OpenGLTexture.cpp b/Engine/Source/Runtime/OpenGLDrv/Private/OpenGLTexture.cpp
--- a/Engine/Source/Runtime/OpenGLDrv/Private/OpenGLTexture.cpp
+++ b/Engine/Source/Runtime/OpenGLDrv/Private/OpenGLTexture.cpp
@@ -956,6 +956,12 @@ void* FOpenGLTexture::Lock(uint32 InMipIndex,uint32 ArrayIndex,EResourceLockMode
 	}
 	else
 	{
+		if(LockMode != RLM_ReadOnly)
+		{
+			// if we modify an active but evictable texture then just make it non-evictable, this is to avoid keeping the cpu backup copies in sync.
+			RemoveEvictionData();
+		}
+
 		// Standard path with a PBO mirroring ever slice of a texture to allow multiple simulataneous maps
 		bool bBufferExists = true;
 		if (!IsValidRef(PixelBuffers[BufferIndex]))
@@ -1227,17 +1233,18 @@ void FOpenGLTexture::RestoreEvictedGLResource(bool bAttemptToRetainMips)
 		? DeferTextureCreationKeepLowerMipCount
 		: 0;
 
+	// keep the mips for streamable textures
+	EvictionParamsPtr->ReleaseMipData(RetainMips);
+
 	if (CanBeEvicted())
 	{
 		if (FTextureEvictionLRU::Get().Add(this) == false)
 		{
 			// could not store this in the LRU. Deleting all backup mips, as this texture will never be evicted.
-			RetainMips = 0;
+			EvictionParamsPtr->ReleaseMipData(0);
 		}
 	}
 
-	// keep the mips for streamable textures
-	EvictionParamsPtr->ReleaseMipData(RetainMips);
 #if GLDEBUG_LABELS_ENABLED
 	FAnsiCharArray& TextureDebugName = EvictionParamsPtr->GetDebugLabelName();
 	if(TextureDebugName.Num())
@@ -1320,6 +1327,15 @@ bool FOpenGLTexture::CanBeEvicted()
 	return bRet;
 }
 
+void FOpenGLTexture::RemoveEvictionData()
+{
+	if (EvictionParamsPtr.IsValid() && CanBeEvicted())
+	{
+		EvictionParamsPtr->ReleaseMipData(0);
+		FTextureEvictionLRU::Get().Remove(this);
+	}
+}
+
 void FOpenGLTexture::CloneViaCopyImage(FOpenGLTexture* Src, uint32 InNumMips, int32 SrcOffset, int32 DstOffset)
 {
 	VERIFY_GL_SCOPE();
@@ -1702,7 +1718,9 @@ void FOpenGLDynamicRHI::RHICopyTexture(FRHITexture* SourceTextureRHI, FRHITextur
 	VERIFY_GL_SCOPE();
 	FOpenGLTexture* SourceTexture = ResourceCast(SourceTextureRHI);
 	FOpenGLTexture* DestTexture = ResourceCast(DestTextureRHI);
-
+	
+	//since the texture will be modified we cannot use the eviction data.
+	DestTexture->RemoveEvictionData();
 
 	GLsizei Width, Height, Depth;
 
