GLES eviction fix: make sure the GL resource is created before removing the MIP data.

diff --git a/Engine/Source/Runtime/OpenGLDrv/Private/OpenGLTexture.cpp b/Engine/Source/Runtime/OpenGLDrv/Private/OpenGLTexture.cpp
--- a/Engine/Source/Runtime/OpenGLDrv/Private/OpenGLTexture.cpp
+++ b/Engine/Source/Runtime/OpenGLDrv/Private/OpenGLTexture.cpp
@@ -444,6 +444,8 @@ void FOpenGLTexture::DeleteGLResource()
 			break;
 		}
 	}
+
+	Resource = GL_NONE;
 }
 
 static inline bool IsAstcLdrRGBAFormat(GLenum Format)
@@ -1336,6 +1338,9 @@ void FOpenGLTexture::RemoveEvictionData()
 {
 	if (EvictionParamsPtr.IsValid() && CanBeEvicted())
 	{
+		// Make sure the GL resource is created before removing the MIP data.
+		TryRestoreGLResource();
+
 		EvictionParamsPtr->ReleaseMipData(0);
 		FTextureEvictionLRU::Get().Remove(this);
 	}
