Implementation of the new RHI GPU profiler for OpenGL RHI
- Relies on the GL_EXT_disjoint_timer_query extension.
- This extension exposes a GPU timestamp query, plus a function to query that timestamp on the CPU for the purpose of calibration.
- Some Android devices have been found to report incorrect values from this API, where the timestamps returned by the GPU query are significantly behind and out-of-sync with the CPU timestamp function. On boot, we detect this and skip enabling the profiler if the timestamps don't align.
- The API works fine on PC. Android is still disabled via SupportsDisjointTimeQueries() being hardcoded to return false. We should evaluate whether this can be enabled in future.

diff --git a/Engine/Source/Runtime/OpenGLDrv/Private/Android/AndroidOpenGLFramePacer.cpp b/Engine/Source/Runtime/OpenGLDrv/Private/Android/AndroidOpenGLFramePacer.cpp
--- a/Engine/Source/Runtime/OpenGLDrv/Private/Android/AndroidOpenGLFramePacer.cpp
+++ b/Engine/Source/Runtime/OpenGLDrv/Private/Android/AndroidOpenGLFramePacer.cpp
@@ -97,9 +97,9 @@ namespace AndroidGL
 	{
 		const double GPUTimeInSeconds = (double)gpu_time_ns / 1000000000.0;
 
-		GGPUFrameTime = GPUTimeInSeconds / FPlatformTime::GetSecondsPerCycle64();
+		uint64 Cycles = GPUTimeInSeconds / FPlatformTime::GetSecondsPerCycle64();
 
-		GetDynamicRHI<FOpenGLDynamicRHI>()->RHISetExternalGPUTime(GGPUFrameTime);
+		GetDynamicRHI<FOpenGLDynamicRHI>()->RHISetExternalGPUTime(Cycles);
 	}
 
 	void SetSwappyPostWaitCallback()
diff --git a/Engine/Source/Runtime/OpenGLDrv/Private/Android/AndroidOpenGLFunctions.cpp b/Engine/Source/Runtime/OpenGLDrv/Private/Android/AndroidOpenGLFunctions.cpp
--- a/Engine/Source/Runtime/OpenGLDrv/Private/Android/AndroidOpenGLFunctions.cpp
+++ b/Engine/Source/Runtime/OpenGLDrv/Private/Android/AndroidOpenGLFunctions.cpp
@@ -41,6 +41,10 @@ namespace GLFuncPointers
 	PFNGLOBJECTPTRLABELKHRPROC				glObjectPtrLabelKHR = NULL;
 	PFNGLGETOBJECTPTRLABELKHRPROC			glGetObjectPtrLabelKHR = NULL;
 
+	// GL_EXT_disjoint_timer_query
+	PFNGLQUERYCOUNTEREXTPROC                glQueryCounterEXT = nullptr;
+	PFNGLGETQUERYOBJECTUI64VEXTPROC         glGetQueryObjectui64vEXT = nullptr;
+
 	// ES 3.2
 	PFNGLTEXBUFFEREXTPROC					glTexBufferEXT = nullptr;
 	PFNGLTEXBUFFERRANGEEXTPROC				glTexBufferRangeEXT = nullptr;
diff --git a/Engine/Source/Runtime/OpenGLDrv/Private/Android/AndroidOpenGLFunctions.h b/Engine/Source/Runtime/OpenGLDrv/Private/Android/AndroidOpenGLFunctions.h
--- a/Engine/Source/Runtime/OpenGLDrv/Private/Android/AndroidOpenGLFunctions.h
+++ b/Engine/Source/Runtime/OpenGLDrv/Private/Android/AndroidOpenGLFunctions.h
@@ -66,6 +66,10 @@ namespace GLFuncPointers
 	extern PFNGLOBJECTPTRLABELKHRPROC		glObjectPtrLabelKHR;
 	extern PFNGLGETOBJECTPTRLABELKHRPROC	glGetObjectPtrLabelKHR;
 
+	// GL_EXT_disjoint_timer_query
+	extern PFNGLQUERYCOUNTEREXTPROC         glQueryCounterEXT;
+	extern PFNGLGETQUERYOBJECTUI64VEXTPROC  glGetQueryObjectui64vEXT;
+
 	// ES 3.2
 	extern PFNGLTEXBUFFEREXTPROC				glTexBufferEXT;
 	extern PFNGLTEXBUFFERRANGEEXTPROC			glTexBufferRangeEXT;
diff --git a/Engine/Source/Runtime/OpenGLDrv/Private/Linux/LinuxOpenGLFunctions.h b/Engine/Source/Runtime/OpenGLDrv/Private/Linux/LinuxOpenGLFunctions.h
--- a/Engine/Source/Runtime/OpenGLDrv/Private/Linux/LinuxOpenGLFunctions.h
+++ b/Engine/Source/Runtime/OpenGLDrv/Private/Linux/LinuxOpenGLFunctions.h
@@ -342,7 +342,8 @@
 	EnumMacro(PFNGLGETPROGRAMPIPELINEIVPROC, glGetProgramPipelineiv)\
 	EnumMacro(PFNGLVALIDATEPROGRAMPIPELINEPROC, glValidateProgramPipeline)\
 	EnumMacro(PFNGLGETPROGRAMPIPELINEINFOLOGPROC, glGetProgramPipelineInfoLog)\
-	EnumMacro(PFNGLISPROGRAMPIPELINEPROC, glIsProgramPipeline)
+	EnumMacro(PFNGLISPROGRAMPIPELINEPROC, glIsProgramPipeline)\
+	EnumMacro(PFNGLGETINTEGER64VPROC, glGetInteger64v)
 
 /** List of all OpenGL entry points. */
 #define ENUM_GL_ENTRYPOINTS_ALL(EnumMacro) \
diff --git a/Engine/Source/Runtime/OpenGLDrv/Private/OpenGLBase.h b/Engine/Source/Runtime/OpenGLDrv/Private/OpenGLBase.h
--- a/Engine/Source/Runtime/OpenGLDrv/Private/OpenGLBase.h
+++ b/Engine/Source/Runtime/OpenGLDrv/Private/OpenGLBase.h
@@ -68,7 +68,7 @@ public:
 	static FORCEINLINE bool SupportsUniformBuffers()					{ return true; }
 	static FORCEINLINE bool SupportsStructuredBuffers()					{ return true; }
 	static FORCEINLINE bool SupportsTimestampQueries()					{ return true; }
-	static FORCEINLINE bool SupportsDisjointTimeQueries()				{ return false; } // @todo: if enabled, causes crash on PC
+	static FORCEINLINE bool SupportsDisjointTimeQueries()				{ return true; }
 	static FORCEINLINE bool SupportsExactOcclusionQueries()				{ return true; }
 	static FORCEINLINE bool SupportsDepthStencilReadSurface()			{ return true; }
 	static FORCEINLINE bool SupportsFloatReadSurface()					{ return true; }
diff --git a/Engine/Source/Runtime/OpenGLDrv/Private/OpenGLCommands.cpp b/Engine/Source/Runtime/OpenGLDrv/Private/OpenGLCommands.cpp
--- a/Engine/Source/Runtime/OpenGLDrv/Private/OpenGLCommands.cpp
+++ b/Engine/Source/Runtime/OpenGLDrv/Private/OpenGLCommands.cpp
@@ -2795,6 +2795,10 @@ IRHIComputeContext* FOpenGLDynamicRHI::RHIGetCommandContext(ERHIPipeline Pipelin
 
 void FOpenGLDynamicRHI::RHIFinalizeContext(FRHIFinalizeContextArgs&& Args, TRHIPipelineArray<IRHIPlatformCommandList*>& Output)
 {
+#if RHI_NEW_GPU_PROFILER
+	FlushProfilerStats();
+#endif
+
 	for(IRHIComputeContext* Context : Args.Contexts)
 	{
 		// "Context" will always be the default context, since we don't implement parallel execution.
diff --git a/Engine/Source/Runtime/OpenGLDrv/Private/OpenGLDevice.cpp b/Engine/Source/Runtime/OpenGLDrv/Private/OpenGLDevice.cpp
--- a/Engine/Source/Runtime/OpenGLDrv/Private/OpenGLDevice.cpp
+++ b/Engine/Source/Runtime/OpenGLDrv/Private/OpenGLDevice.cpp
@@ -92,7 +92,9 @@ extern void BeginFrame_VertexBufferCleanup();
 
 void FOpenGLDynamicRHI::RHIEndFrame(const FRHIEndFrameArgs& Args)
 {
-#if (RHI_NEW_GPU_PROFILER == 0)
+#if RHI_NEW_GPU_PROFILER
+	Profiler.EndFrame(Args);
+#else
 	GPUProfilingData->EndFrame();
 #endif
 
@@ -1339,6 +1341,10 @@ void FOpenGLDynamicRHI::Init()
 	// Flush here since we might be switching to a different context/thread for rendering
 	FOpenGL::Flush();
 
+#if RHI_NEW_GPU_PROFILER
+	Profiler.Setup();
+#endif
+
 	FHardwareInfo::RegisterHardwareInfo( NAME_RHI, TEXT( "OpenGL" ) );
 
 	FRenderResource::InitPreRHIResources();
@@ -1349,6 +1355,9 @@ void FOpenGLDynamicRHI::Shutdown()
 {
 	check(IsInGameThread() && IsInRenderingThread()); // require that the render thread has been shut down
 
+#if RHI_NEW_GPU_PROFILER
+	Profiler.Shutdown();
+#endif
 	Cleanup();
 
 	FOpenGLRenderQuery::Cleanup();
@@ -1442,3 +1451,213 @@ bool FOpenGLDynamicRHIModule::IsSupported()
 {
 	return true;
 }
+
+#if RHI_NEW_GPU_PROFILER
+
+FOpenGLRenderQuery* FOpenGLProfiler::GetQuery()
+{
+	if (QueryPool.IsEmpty())
+	{
+		return new FOpenGLRenderQuery(FOpenGLRenderQuery::EType::Profiler);
+	}
+	else
+	{
+		return QueryPool.Pop();
+	}
+}
+
+FOpenGLRenderQuery* FOpenGLProfiler::InsertQuery(uint64* Target)
+{
+	check(bEnabled);
+
+	// Temporarily store the current frame pointer in the Target.
+	// We retrieve this later when resolving the query.
+	*Target = reinterpret_cast<uint64>(Current.Get());
+
+	FOpenGLRenderQuery* Query = GetQuery();
+	Query->End(Target);
+	return Query;
+}
+
+FOpenGLProfiler::FFrame::FFrame(FOpenGLProfiler& Profiler)
+	: EventStream(UE::RHI::GPUProfiler::FQueue(UE::RHI::GPUProfiler::FQueue::EType::Graphics, 0, 0))
+{
+	if (Profiler.bEnabled)
+	{
+		glGetInteger64v(GL_TIMESTAMP_EXT, (GLint64*)&Calibration.GLTimestamp);
+		Calibration.CPUTimestamp = FPlatformTime::Cycles64();
+		Calibration.CPUFrequency = uint64(1.0 / FPlatformTime::GetSecondsPerCycle64());
+	}
+}
+
+void FOpenGLProfiler::Setup()
+{
+	bEnabled = false;
+
+	if (!FOpenGL::SupportsDisjointTimeQueries())
+	{
+		UE_LOG(LogOpenGL, Warning, TEXT("OpenGL RHI profiler not supported."));
+	}
+	else
+	{
+		// Test for the correctness of the GL timestamp functions.
+		// On some devices, the values reported by the GPU queries do not tick at the same rate as those reported by the CPU function.
+		
+		GLuint Query;
+		FOpenGL::GenQueries(1, &Query);
+
+		GLuint64 GPU = 0;
+		GLint64 CPU = 0;
+
+		for (int32 Attempts = 0; Attempts < 10; ++Attempts)
+		{
+			GPU = 0;
+			CPU = 0;
+
+			// Sample both the GPU and CPU timestamps, as close to simultenously as possible
+			FOpenGL::QueryTimestampCounter(Query);
+			FOpenGL::GetQueryObject(Query, FOpenGL::QM_Result, &GPU);
+			glGetInteger64v(GL_TIMESTAMP_EXT, &CPU);
+
+			GLint64 Delta = FMath::Abs(GLint64(GPU) - CPU);
+
+			// We consider the timestamps reliable if the GPU and CPU clocks are within 100ms of each other.
+			// On devices that don't implement this properly, the difference between the clocks can be minutes / hours.
+			const GLint64 Threshold = 100 * 1000 * 1000; // 100ms in nanoseconds
+			if (!FOpenGL::TimerQueryDisjoint() && GPU != 0 && CPU != 0 && Delta < Threshold)
+			{
+				bEnabled = true;
+				break;
+			}
+		}
+
+		FOpenGL::DeleteQueries(1, &Query);
+
+		if (bEnabled)
+		{
+			UE_LOG(LogOpenGL, Display, TEXT("OpenGL RHI profiler detected reliable timestamps. GPU profiler is enabled. GPU clock: 0x%llx, CPU Clock: 0x%llx"), GPU, CPU);
+		}
+		else
+		{
+			UE_LOG(LogOpenGL, Warning, TEXT("OpenGL RHI profiler detected unreliable timestamps. GPU profiler is disabled. GPU clock: 0x%llx, CPU Clock: 0x%llx"), GPU, CPU);
+		}
+	}
+
+	// Register the single graphics GPU queue we have access to in OpenGL.
+	UE::RHI::GPUProfiler::FQueue Queue(UE::RHI::GPUProfiler::FQueue::EType::Graphics, 0, 0);
+	UE::RHI::GPUProfiler::InitializeQueues(MakeConstArrayView(&Queue, 1));
+
+	Current = MakeUnique<FFrame>(*this);
+
+	if (bEnabled)
+	{
+		//
+		// Since we can't tell when the GPU is actually executing engine work in OpenGL,
+		// just mark the GPU as always busy in the frame. We also push begin/end work
+		// markers in RHIEndFrame either side of the frame boundary.
+		//
+		auto& Event = EmplaceEvent<UE::RHI::GPUProfiler::FEvent::FBeginWork>(Current->Calibration.CPUTimestamp);
+		InsertQuery(&Event.GPUTimestampTOP);
+	}
+}
+
+void FOpenGLProfiler::Shutdown()
+{
+	// Drain the query readback queue
+	FOpenGLRenderQuery::PollQueryResults();
+}
+
+void FOpenGLProfiler::EndFrame(FDynamicRHI::FRHIEndFrameArgs const& Args)
+{
+	if (bEnabled)
+	{
+		auto& EventWork = EmplaceEvent<UE::RHI::GPUProfiler::FEvent::FEndWork>();
+		Current->EndWorkQuery = InsertQuery(&EventWork.GPUTimestampBOP);
+	}
+
+	if (ExternalGPUTime)
+	{
+		EmplaceEvent<UE::RHI::GPUProfiler::FEvent::FFrameTime>(*ExternalGPUTime);
+	}
+
+	// Flushing after the EndWork event helps improve the accuracy of the timer
+	FOpenGL::Flush();
+
+	// We must always emit the frame boundary event, even if the profiler is disabled in the RHI.
+	TUniquePtr<FFrame> NewFrame = MakeUnique<FFrame>(*this);
+
+	// Insert frame boundary
+	Current->FrameBoundaryEvent = &EmplaceEvent<UE::RHI::GPUProfiler::FEvent::FFrameBoundary>(NewFrame->Calibration.CPUTimestamp, Args.FrameNumber
+#if WITH_RHI_BREADCRUMBS
+		, Args.GPUBreadcrumbs[ERHIPipeline::Graphics]
+#endif
+#if STATS
+		, Args.StatsFrame
+#endif
+	);
+
+	Pending.Enqueue(MoveTemp(Current));
+	Current = MoveTemp(NewFrame);
+
+	// Attempt to process historic results
+	while (TUniquePtr<FFrame>* PreviousFramePtr = Pending.Peek())
+	{
+		TUniquePtr<FFrame>& PreviousFrame = *PreviousFramePtr;
+
+		// When the profiler is disabled, EndWorkQuery will be nullptr
+		if (PreviousFrame->EndWorkQuery && !FOpenGLRenderQuery::PollQueryResults(PreviousFrame->EndWorkQuery))
+		{
+			// Frame not yet finished on the GPU
+			break;
+		}
+
+		// Previous frame has completed and the data is available.
+
+		if (PreviousFrame->bDisjoint)
+		{
+			// Discard all the events from this frame, except the frame boundary
+			UE::RHI::GPUProfiler::FEventStream Stream(PreviousFrame->EventStream.Queue);
+			Stream.Emplace<UE::RHI::GPUProfiler::FEvent::FFrameBoundary>(*PreviousFrame->FrameBoundaryEvent);
+
+			UE::RHI::GPUProfiler::ProcessEvents(MakeArrayView(&Stream, 1));
+		}
+		else
+		{
+			// This is a valid frame. Publish the profiler events.
+			UE::RHI::GPUProfiler::ProcessEvents(MakeArrayView(&PreviousFrame->EventStream, 1));
+		}
+
+		Pending.Pop();
+	}
+
+	if (bEnabled)
+	{
+		// Start the next frame's GPU work
+		auto& BeginWork = EmplaceEvent<UE::RHI::GPUProfiler::FEvent::FBeginWork>(Current->Calibration.CPUTimestamp);
+		InsertQuery(&BeginWork.GPUTimestampTOP);
+	}
+}
+
+uint64 FOpenGLProfiler::ResolveQuery(uint64 Value, uint64* Target, bool bDisjoint)
+{
+	check(bEnabled);
+
+	// Read the current frame pointer back out of the target, which was stored here during InsertQuery().
+	FFrame* Frame = reinterpret_cast<FFrame*>(*Target);
+
+	Frame->bDisjoint |= bDisjoint;
+	if (Frame->bDisjoint)
+	{
+		return 0;
+	}
+	else
+	{
+		// Convert from GPU timestamp to CPU timestamp (relative to FPlatformTime::Cycles64())
+		int64 GLDelta = Value - Frame->Calibration.GLTimestamp;
+		int64 CPUDelta = (GLDelta * Frame->Calibration.CPUFrequency) / Frame->Calibration.GLFrequency;
+
+		return CPUDelta + Frame->Calibration.CPUTimestamp;
+	}
+}
+
+#endif // RHI_NEW_GPU_PROFILER
diff --git a/Engine/Source/Runtime/OpenGLDrv/Private/OpenGLDrv.cpp b/Engine/Source/Runtime/OpenGLDrv/Private/OpenGLDrv.cpp
--- a/Engine/Source/Runtime/OpenGLDrv/Private/OpenGLDrv.cpp
+++ b/Engine/Source/Runtime/OpenGLDrv/Private/OpenGLDrv.cpp
@@ -56,11 +56,12 @@ bool FOpenGLDynamicRHI::RHIIsValidTexture(GLuint InTexture) const
 	return glIsTexture(InTexture) == GL_TRUE;
 }
 
-void FOpenGLDynamicRHI::RHISetExternalGPUTime(uint32 InExternalGPUTime)
+void FOpenGLDynamicRHI::RHISetExternalGPUTime(uint64 InExternalGPUTime)
 {
 #if RHI_NEW_GPU_PROFILER
-	checkNoEntry(); // @todo - new gpu profiler
+	Profiler.ExternalGPUTime = InExternalGPUTime;
 #else
+	GGPUFrameTime = InExternalGPUTime;
 	GPUProfilingData->ExternalGPUTime = InExternalGPUTime;
 #endif
 }
@@ -142,7 +143,15 @@ void FOpenGLDynamicRHI::RHIEGLTerminateContext()
 	#endif
 		}
 
-	#if (RHI_NEW_GPU_PROFILER == 0)
+	#if RHI_NEW_GPU_PROFILER
+		FlushProfilerStats();
+
+		if (Profiler.bEnabled)
+		{
+			auto& Event = Profiler.EmplaceEvent<UE::RHI::GPUProfiler::FEvent::FBeginBreadcrumb>(Breadcrumb);
+			Profiler.InsertQuery(&Event.GPUTimestampTOP);
+		}
+	#else
 		if (GPUProfilingData->IsProfilingGPU())
 		{
 			GPUProfilingData->PushEvent(GetNameStr(), FColor::White);
@@ -152,7 +161,15 @@ void FOpenGLDynamicRHI::RHIEGLTerminateContext()
 
 	void FOpenGLDynamicRHI::RHIEndBreadcrumbGPU(FRHIBreadcrumbNode* Breadcrumb)
 	{
-	#if (RHI_NEW_GPU_PROFILER == 0)
+	#if RHI_NEW_GPU_PROFILER
+		FlushProfilerStats();
+
+		if (Profiler.bEnabled)
+		{
+			auto& Event = Profiler.EmplaceEvent<UE::RHI::GPUProfiler::FEvent::FEndBreadcrumb>(Breadcrumb);
+			Profiler.InsertQuery(&Event.GPUTimestampBOP);
+		}
+	#else
 		if (GPUProfilingData->IsProfilingGPU())
 		{
 			GPUProfilingData->PopEvent();
diff --git a/Engine/Source/Runtime/OpenGLDrv/Private/OpenGLDrv.h b/Engine/Source/Runtime/OpenGLDrv/Private/OpenGLDrv.h
--- a/Engine/Source/Runtime/OpenGLDrv/Private/OpenGLDrv.h
+++ b/Engine/Source/Runtime/OpenGLDrv/Private/OpenGLDrv.h
@@ -272,7 +272,7 @@ public:
 	virtual FTextureRHIRef RHICreateTextureCubeFromResource(EPixelFormat Format, uint32 Size, bool bArray, uint32 ArraySize, uint32 NumMips, uint32 NumSamples, uint32 NumSamplesTileMem, const FClearValueBinding& ClearValueBinding, GLuint Resource, ETextureCreateFlags Flags) final override;
 	virtual GLuint RHIGetResource(FRHITexture* InTexture) const final override;
 	virtual bool RHIIsValidTexture(GLuint InTexture) const final override;
-	virtual void RHISetExternalGPUTime(uint32 InExternalGPUTime) final override;
+	virtual void RHISetExternalGPUTime(uint64 InExternalGPUTime) final override;
 
 #if PLATFORM_ANDROID
 	virtual EGLDisplay RHIGetEGLDisplay() const final override;
@@ -694,7 +694,23 @@ private:
 	/** Underlying platform-specific data */
 	struct FPlatformOpenGLDevice* PlatformDevice = nullptr;
 
-#if (RHI_NEW_GPU_PROFILER == 0)
+#if RHI_NEW_GPU_PROFILER
+
+	friend class FOpenGLRenderQuery;
+
+	FOpenGLProfiler Profiler;
+	
+	void FlushProfilerStats()
+	{
+		// Flush accumulated draw stats
+		if (Profiler.bEnabled && StatEvent)
+		{
+			Profiler.EmplaceEvent<UE::RHI::GPUProfiler::FEvent::FStats>() = StatEvent;
+			StatEvent = {};
+		}
+	}
+
+#else
 
 	TOptional<FOpenGLGPUProfiler> GPUProfilingData;
 	friend FOpenGLGPUProfiler;
diff --git a/Engine/Source/Runtime/OpenGLDrv/Private/OpenGLES.cpp b/Engine/Source/Runtime/OpenGLDrv/Private/OpenGLES.cpp
--- a/Engine/Source/Runtime/OpenGLDrv/Private/OpenGLES.cpp
+++ b/Engine/Source/Runtime/OpenGLDrv/Private/OpenGLES.cpp
@@ -168,7 +168,13 @@ void FOpenGLES::ProcessExtensions(const FString& ExtensionsString)
 	ProcessQueryGLInt();
 	FOpenGLBase::ProcessExtensions(ExtensionsString);
 
-	bSupportsDisjointTimeQueries = ExtensionsString.Contains(TEXT("GL_EXT_disjoint_timer_query")) || ExtensionsString.Contains(TEXT("GL_NV_timer_query"));
+	bSupportsDisjointTimeQueries = ExtensionsString.Contains(TEXT("GL_EXT_disjoint_timer_query"));
+	if (bSupportsDisjointTimeQueries)
+	{
+		glQueryCounterEXT = (PFNGLQUERYCOUNTEREXTPROC)((void*)eglGetProcAddress("glQueryCounterEXT"));;
+		glGetQueryObjectui64vEXT = (PFNGLGETQUERYOBJECTUI64VEXTPROC)((void*)eglGetProcAddress("glGetQueryObjectui64vEXT"));
+	}
+
 	bTimerQueryCanBeDisjoint = !ExtensionsString.Contains(TEXT("GL_NV_timer_query"));
 	bSupportsBGRA8888 = ExtensionsString.Contains(TEXT("GL_APPLE_texture_format_BGRA8888")) || ExtensionsString.Contains(TEXT("GL_IMG_texture_format_BGRA8888")) || ExtensionsString.Contains(TEXT("GL_EXT_texture_format_BGRA8888"));
 	bSupportsColorBufferFloat = ExtensionsString.Contains(TEXT("GL_EXT_color_buffer_float"));
diff --git a/Engine/Source/Runtime/OpenGLDrv/Private/OpenGLES.h b/Engine/Source/Runtime/OpenGLDrv/Private/OpenGLES.h
--- a/Engine/Source/Runtime/OpenGLDrv/Private/OpenGLES.h
+++ b/Engine/Source/Runtime/OpenGLDrv/Private/OpenGLES.h
@@ -114,6 +114,11 @@ struct FOpenGLES : public FOpenGLBase
 		glDeleteQueries(NumQueries, QueryIDs);
 	}
 
+	static FORCEINLINE void QueryTimestampCounter(GLuint QueryID)
+	{
+		glQueryCounterEXT(QueryID, GL_TIMESTAMP_EXT);
+	}
+
 	using FOpenGLBase::GetQueryObject;
 
 	static FORCEINLINE void GetQueryObject(GLuint QueryId, EQueryMode QueryMode, GLuint* OutResult)
@@ -121,6 +126,12 @@ struct FOpenGLES : public FOpenGLBase
 		GLenum QueryName = (QueryMode == QM_Result) ? GL_QUERY_RESULT : GL_QUERY_RESULT_AVAILABLE;
 		glGetQueryObjectuiv(QueryId, QueryName, OutResult);
 	}
+
+	static FORCEINLINE void GetQueryObject(GLuint QueryId, EQueryMode QueryMode, GLuint64* OutResult)
+	{
+		GLenum QueryName = (QueryMode == QM_Result) ? GL_QUERY_RESULT : GL_QUERY_RESULT_AVAILABLE;
+		glGetQueryObjectui64vEXT(QueryId, QueryName, OutResult);
+	}
 		
 	static FORCEINLINE void LabelObject(GLenum Type, GLuint Object, const ANSICHAR* Name)
 	{
diff --git a/Engine/Source/Runtime/OpenGLDrv/Private/OpenGLQuery.cpp b/Engine/Source/Runtime/OpenGLDrv/Private/OpenGLQuery.cpp
--- a/Engine/Source/Runtime/OpenGLDrv/Private/OpenGLQuery.cpp
+++ b/Engine/Source/Runtime/OpenGLDrv/Private/OpenGLQuery.cpp
@@ -140,7 +140,6 @@ void FOpenGLRenderQuery::Begin()
 	switch(Type)
 	{
 	default:
-	case EType::Timestamp:
 		checkNoEntry();
 		break;
 
@@ -153,13 +152,15 @@ void FOpenGLRenderQuery::Begin()
 		);
 		break;
 
+#if RHI_NEW_GPU_PROFILER == 0
 	case EType::Disjoint:
 		FOpenGL::BeginQuery(UGL_TIME_ELAPSED, Resource);
 		break;
+#endif
 	};
 }
 
-void FOpenGLRenderQuery::End()
+void FOpenGLRenderQuery::End(uint64* InTarget)
 {
 	VERIFY_GL_SCOPE();
 	AcquireGlQuery();
@@ -180,21 +181,34 @@ void FOpenGLRenderQuery::End()
 		FOpenGL::QueryTimestampCounter(Resource);
 		break;
 
+#if RHI_NEW_GPU_PROFILER
+	case EType::Profiler:
+		FOpenGL::QueryTimestampCounter(Resource);
+		break;
+#else
 	case EType::Disjoint:
 		FOpenGL::EndQuery(UGL_TIME_ELAPSED);
 		break;
+#endif
 	}
 
 	BOPCounter++;
 
+	Target = InTarget;
+
 	Link();
 }
 
+bool FOpenGLRenderQuery::IsCached()
+{
+	return BOPCounter == LastCachedBOPCounter.load(std::memory_order_relaxed);
+}
+
 bool FOpenGLRenderQuery::CacheResult(bool bWait)
 {
 	VERIFY_GL_SCOPE();
 
-	if (BOPCounter == LastCachedBOPCounter.load(std::memory_order_relaxed))
+	if (IsCached())
 	{
 		// Value has been cached and no newer query operation has started.
 		check(!IsLinked());
@@ -241,6 +255,39 @@ bool FOpenGLRenderQuery::CacheResult(bool bWait)
 		}
 		break;
 
+#if RHI_NEW_GPU_PROFILER
+	case EType::Profiler:
+		{
+			FOpenGLDynamicRHI& RHI = FOpenGLDynamicRHI::Get();
+
+			// TimerQueryDisjoint is a one-shot state in the driver, it is not pipelined.
+			// If it returns true, all timers we've submitted after this timer but haven't
+			// yet resolved should be discarded for having invalid data.
+			if (FOpenGL::TimerQueryDisjoint())
+			{
+				for (FOpenGLRenderQuery* Other = this; Other; Other = Other->Next)
+				{
+					if (Other->Type == EType::Profiler)
+					{
+						Other->SetResult(RHI.Profiler.ResolveQuery(0, Other->Target, true));
+
+						// Return the query to the pool
+						RHI.Profiler.QueryPool.Push(Other);
+					}
+				}
+			}
+			else
+			{
+				GLuint64 Value;
+				FOpenGL::GetQueryObject(Resource, FOpenGL::QM_Result, &Value);
+				SetResult(RHI.Profiler.ResolveQuery(Value, Target, false));
+
+				// Return the query to the pool
+				RHI.Profiler.QueryPool.Push(this);
+			}
+		}
+		break;
+#else
 	case EType::Disjoint:
 		{
 			// TimerQueryDisjoint is a one-shot state in the driver, it is not pipelined.
@@ -266,6 +313,7 @@ bool FOpenGLRenderQuery::CacheResult(bool bWait)
 			}
 		}
 		break;
+#endif
 	}
 
 	return true;
@@ -273,6 +321,12 @@ bool FOpenGLRenderQuery::CacheResult(bool bWait)
 
 void FOpenGLRenderQuery::SetResult(uint64 Value)
 {
+	if (Target)
+	{
+		*Target = Value;
+		Target = nullptr;
+	}
+
 	Result = Value;
 	ReleaseGlQuery();
 
@@ -321,12 +375,17 @@ bool FOpenGLRenderQuery_RHI::GetResult(bool bWait, uint64& OutResult)
 	return true;
 }
 
-void FOpenGLRenderQuery::PollQueryResults()
+bool FOpenGLRenderQuery::PollQueryResults(FOpenGLRenderQuery* TargetQuery)
 {
 	if (!PlatformOpenGLThreadHasRenderingContext())
 	{
 		// Don't poll queries if this thread doesn't own the GL context.
-		return;
+		return false;
+	}
+
+	if (TargetQuery && TargetQuery->IsCached())
+	{
+		return true;
 	}
 
 	if (ActiveQueries.First)
@@ -335,11 +394,22 @@ void FOpenGLRenderQuery::PollQueryResults()
 
 		do
 		{
-			if (!ActiveQueries.First->CacheResult(/*bWait = */ false))
-				break;
+			FOpenGLRenderQuery* Current = ActiveQueries.First;
+			if (!Current->CacheResult(/*bWait = */ false))
+			{
+				// Not complete yet
+				return false;
+			}
+
+			if (Current == TargetQuery)
+			{
+				return true;
+			}
 		}
 		while (ActiveQueries.First);
 	}
+
+	return TargetQuery == nullptr;
 }
 
 void FOpenGLRenderQuery::Cleanup()
diff --git a/Engine/Source/Runtime/OpenGLDrv/Private/OpenGLResources.h b/Engine/Source/Runtime/OpenGLDrv/Private/OpenGLResources.h
--- a/Engine/Source/Runtime/OpenGLDrv/Private/OpenGLResources.h
+++ b/Engine/Source/Runtime/OpenGLDrv/Private/OpenGLResources.h
@@ -1132,6 +1132,67 @@ public:
 	uint8 const bMultisampleRenderbuffer : 1;
 };
 
+#if RHI_NEW_GPU_PROFILER
+
+struct FOpenGLProfiler
+{
+	struct FCalibration
+	{
+		// FPlatformTime::Cycles64()
+		int64 CPUTimestamp = 0;
+		int64 CPUFrequency = 0;
+
+		// The GL timestamp as queried directly from the driver by the CPU
+		int64 GLTimestamp = 0;
+
+		// Frequency of the GLTimestamp clock
+		static constexpr int64 GLFrequency = 1'000'000'000; // one tick per nanosecond
+	};
+
+	struct FFrame
+	{
+		UE::RHI::GPUProfiler::FEventStream EventStream;
+		FCalibration Calibration;
+
+		class FOpenGLRenderQuery* EndWorkQuery = nullptr;
+		UE::RHI::GPUProfiler::FEvent::FFrameBoundary* FrameBoundaryEvent = nullptr;
+
+		// True if this frame contained any timestamp queries which the driver flagged as disjoint. 
+		// If this happens, we discard all the timestamp events from this frame, leaving only the 
+		// frame boundary, since the timing will be unreliable.
+		bool bDisjoint = false;
+
+		FFrame(FOpenGLProfiler& Profiler);
+	};
+
+	bool bEnabled = false;
+
+	TUniquePtr<FFrame> Current;
+	TQueue<TUniquePtr<FFrame>> Pending;
+
+	TArray<FOpenGLRenderQuery*> QueryPool;
+
+	TOptional<uint32> ExternalGPUTime;
+
+	FOpenGLRenderQuery* GetQuery();
+	FOpenGLRenderQuery* InsertQuery(uint64* Target);
+
+	uint64 ResolveQuery(uint64 Value, uint64* Target, bool bDisjoint);
+
+	void Setup();
+	void Shutdown();
+
+	void EndFrame(FDynamicRHI::FRHIEndFrameArgs const& Args);
+
+	template <typename TEventType, typename... TArgs>
+	TEventType& EmplaceEvent(TArgs&&... Args)
+	{
+		return Current->EventStream.Emplace<TEventType>(Forward<TArgs>(Args)...);
+	}
+};
+
+#endif
+
 class FOpenGLRenderQuery
 {
 public:
@@ -1139,12 +1200,18 @@ public:
 	{
 		Timestamp,
 		Occlusion,
+	#if RHI_NEW_GPU_PROFILER
+		Profiler,
+	#else
 		Disjoint,
+	#endif
 
 		Num
 	};
 
+#if RHI_NEW_GPU_PROFILER == 0
 	static constexpr uint64 InvalidDisjointMask = 0x8000000000000000;
+#endif
 
 private:
 	// Render queries that should be polled by the RHI thread.
@@ -1169,8 +1236,11 @@ private:
 	/** The query resource. */
 	GLuint Resource = 0;
 
+	// Additional data used for profiler timestamps
+	uint64* Target = nullptr;
+
 protected:
-	EType const Type : 2;
+	EType const Type;
 
 	std::atomic<uint8> LastCachedBOPCounter = 0;
 	uint8 BOPCounter = 0;
@@ -1189,7 +1259,7 @@ public:
 	bool IsLinked() const { return Prev != nullptr; }
 
 	void Begin();
-	void End();
+	void End(uint64* InTarget = nullptr);
 
 	uint64 GetResult() const
 	{
@@ -1198,7 +1268,7 @@ public:
 
 	bool CacheResult(bool bWait);
 
-	static void PollQueryResults();
+	static bool PollQueryResults(FOpenGLRenderQuery* TargetQuery = nullptr);
 	static void Cleanup();
 
 private:
@@ -1206,6 +1276,8 @@ private:
 	void Unlink();
 
 	void SetResult(uint64 Value);
+
+	bool IsCached();
 };
 
 class FOpenGLRenderQuery_RHI : public FRHIRenderQuery, public FOpenGLRenderQuery
diff --git a/Engine/Source/Runtime/OpenGLDrv/Private/Windows/WindowsOpenGLFunctions.h b/Engine/Source/Runtime/OpenGLDrv/Private/Windows/WindowsOpenGLFunctions.h
--- a/Engine/Source/Runtime/OpenGLDrv/Private/Windows/WindowsOpenGLFunctions.h
+++ b/Engine/Source/Runtime/OpenGLDrv/Private/Windows/WindowsOpenGLFunctions.h
@@ -377,7 +377,8 @@
 	EnumMacro(PFNGLGETPROGRAMPIPELINEINFOLOGPROC, glGetProgramPipelineInfoLog) \
 	EnumMacro(PFNGLISPROGRAMPIPELINEPROC, glIsProgramPipeline) \
 	EnumMacro(PFNGLGETPROGRAMBINARYPROC, glGetProgramBinary) \
-	EnumMacro(PFNGLPROGRAMBINARYPROC, glProgramBinary)
+	EnumMacro(PFNGLPROGRAMBINARYPROC, glProgramBinary) \
+	EnumMacro(PFNGLGETINTEGER64VPROC, glGetInteger64v)
 
 /** List of all OpenGL entry points. */
 #define ENUM_GL_ENTRYPOINTS_ALL(EnumMacro) \
diff --git a/Engine/Source/Runtime/OpenGLDrv/Public/IOpenGLDynamicRHI.h b/Engine/Source/Runtime/OpenGLDrv/Public/IOpenGLDynamicRHI.h
--- a/Engine/Source/Runtime/OpenGLDrv/Public/IOpenGLDynamicRHI.h
+++ b/Engine/Source/Runtime/OpenGLDrv/Public/IOpenGLDynamicRHI.h
@@ -20,7 +20,7 @@ struct IOpenGLDynamicRHI : public FDynamicRHIPSOFallback
 
 	virtual GLuint RHIGetResource(FRHITexture* InTexture) const = 0;
 	virtual bool RHIIsValidTexture(GLuint InTexture) const = 0;
-	virtual void RHISetExternalGPUTime(uint32 InExternalGPUTime) = 0;
+	virtual void RHISetExternalGPUTime(uint64 InExternalGPUTime) = 0;
 
 	virtual void RHIGenerateMips(FRHITexture* Texture) = 0;
 
