Fix OpenGL LLM buffer tracking by lining it up with RHI core stats.

diff --git a/Engine/Source/Runtime/OpenGLDrv/Private/OpenGLResources.h b/Engine/Source/Runtime/OpenGLDrv/Private/OpenGLResources.h
--- a/Engine/Source/Runtime/OpenGLDrv/Private/OpenGLResources.h
+++ b/Engine/Source/Runtime/OpenGLDrv/Private/OpenGLResources.h
@@ -171,7 +171,7 @@ public:
 				LoadData(0, BaseType::GetSize(), InData);
 			}
 #endif
-			OpenGLBufferStats::UpdateBufferStats(BaseType::GetDesc(), true);
+			BaseType::UpdateBufferStats(true);
 		};
 
 		if (!CreateDesc.IsNull())
@@ -497,7 +497,7 @@ public:
 			Resource = 0;
 
 			LockBuffer = nullptr;
-			OpenGLBufferStats::UpdateBufferStats(BaseType::GetDesc(), false);
+			BaseType::UpdateBufferStats(false);
 
 			ReleaseCachedBuffer();
 		}
@@ -544,6 +544,12 @@ public:
 
 private:
 	FRHIBufferDesc Desc;
+	
+protected:
+	void UpdateBufferStats(bool bAllocating)
+	{
+		OpenGLBufferStats::UpdateBufferStats(Desc, bAllocating);
+	}
 };
 
 class FOpenGLBaseBuffer : public FRHIBuffer, public FOpenGLViewableResource
@@ -552,32 +558,17 @@ public:
 	FOpenGLBaseBuffer(const FRHIBufferCreateDesc& InCreateDesc)
 		: FRHIBuffer(InCreateDesc)
 	{
-		LLM_SCOPED_PAUSE_TRACKING_WITH_ENUM_AND_AMOUNT(ELLMTag::GraphicsPlatform, (int64)GetSize(), ELLMTracker::Platform, ELLMAllocType::None);
-		LLM_SCOPED_PAUSE_TRACKING_WITH_ENUM_AND_AMOUNT(ELLMTag::Meshes, (int64)GetSize(), ELLMTracker::Default, ELLMAllocType::None);
 	}
 
-	~FOpenGLBaseBuffer()
-	{
-		// If ReleaseOwnership was called, then Size is already 0
-		if (GetSize() != 0)
-		{
-			LLM_SCOPED_PAUSE_TRACKING_WITH_ENUM_AND_AMOUNT(ELLMTag::GraphicsPlatform, -(int64)GetSize(), ELLMTracker::Platform, ELLMAllocType::None);
-			LLM_SCOPED_PAUSE_TRACKING_WITH_ENUM_AND_AMOUNT(ELLMTag::Meshes, -(int64)GetSize(), ELLMTracker::Default, ELLMAllocType::None);
-		}
-    }
 
-	void TakeOwnership(FOpenGLBaseBuffer& Other)
-	{
-		FRHIBuffer::TakeOwnership(Other);
-		LLM_SCOPED_PAUSE_TRACKING_WITH_ENUM_AND_AMOUNT(ELLMTag::GraphicsPlatform, (int64)GetSize(), ELLMTracker::Platform, ELLMAllocType::None);
-		LLM_SCOPED_PAUSE_TRACKING_WITH_ENUM_AND_AMOUNT(ELLMTag::Meshes, (int64)GetSize(), ELLMTracker::Default, ELLMAllocType::None);
-	}
-
-	void ReleaseOwnership()
+protected:
+	void UpdateBufferStats(bool bAllocating)
 	{
-		LLM_SCOPED_PAUSE_TRACKING_WITH_ENUM_AND_AMOUNT(ELLMTag::GraphicsPlatform, -(int64)GetSize(), ELLMTracker::Platform, ELLMAllocType::None);
-		LLM_SCOPED_PAUSE_TRACKING_WITH_ENUM_AND_AMOUNT(ELLMTag::Meshes, -(int64)GetSize(), ELLMTracker::Default, ELLMAllocType::None);
-		FRHIBuffer::ReleaseOwnership();
+		int64 AllocSize = ( bAllocating ? -1 : 1 ) * GetSize();
+		LLM_SCOPED_PAUSE_TRACKING_WITH_ENUM_AND_AMOUNT(ELLMTag::GraphicsPlatform, AllocSize, ELLMTracker::Platform, ELLMAllocType::None);
+		LLM_SCOPED_PAUSE_TRACKING_WITH_ENUM_AND_AMOUNT(ELLMTag::Meshes, AllocSize, ELLMTracker::Default, ELLMAllocType::None);
+		
+		OpenGLBufferStats::UpdateBufferStats(GetDesc(), bAllocating);
 	}
 };
 
