Use correct Manual depth test for 32F depth for Vulkan and GL.

diff --git a/Engine/Source/Runtime/OpenGLDrv/Private/OpenGLDevice.cpp b/Engine/Source/Runtime/OpenGLDrv/Private/OpenGLDevice.cpp
--- a/Engine/Source/Runtime/OpenGLDrv/Private/OpenGLDevice.cpp
+++ b/Engine/Source/Runtime/OpenGLDrv/Private/OpenGLDevice.cpp
@@ -1010,11 +1010,13 @@ static void InitRHICapabilitiesForGL()
 	{
 		SetupTextureFormat(PF_DepthStencil, FOpenGLTextureFormat(GL_DEPTH24_STENCIL8, GL_NONE, GL_DEPTH_STENCIL, GL_UNSIGNED_INT_24_8, false, false));
 		GPixelFormats[PF_DepthStencil].BlockBytes = 4;
+		GPixelFormats[PF_DepthStencil].bIs24BitUnormDepthStencil = true;
 	}
 	else
 	{
 		SetupTextureFormat(PF_DepthStencil, FOpenGLTextureFormat(GL_DEPTH32F_STENCIL8, GL_NONE, GL_DEPTH_STENCIL, GL_FLOAT_32_UNSIGNED_INT_24_8_REV, false, false));
 		GPixelFormats[PF_DepthStencil].BlockBytes = 8;
+		GPixelFormats[PF_DepthStencil].bIs24BitUnormDepthStencil = false;
 	}
 	GPixelFormats[PF_X24_G8].Supported = true;
 	SetupTextureFormat( PF_ShadowDepth,			FOpenGLTextureFormat( ShadowDepthFormat,		ShadowDepthFormat,		GL_DEPTH_COMPONENT,	GL_UNSIGNED_INT,				false,			false));
diff --git a/Engine/Source/Runtime/VulkanRHI/Private/VulkanDevice.cpp b/Engine/Source/Runtime/VulkanRHI/Private/VulkanDevice.cpp
--- a/Engine/Source/Runtime/VulkanRHI/Private/VulkanDevice.cpp
+++ b/Engine/Source/Runtime/VulkanRHI/Private/VulkanDevice.cpp
@@ -746,6 +746,7 @@ void FVulkanDevice::SetupFormats()
 		// prefer VK_FORMAT_D32_SFLOAT_S8_UINT
 		MapFormatSupport(PF_DepthStencil, { VK_FORMAT_D32_SFLOAT_S8_UINT, VK_FORMAT_D24_UNORM_S8_UINT, VK_FORMAT_D16_UNORM_S8_UINT }, ComponentMappingRIII);
 		MapFormatSupport(PF_X24_G8, { VK_FORMAT_D32_SFLOAT_S8_UINT, VK_FORMAT_D24_UNORM_S8_UINT, VK_FORMAT_D16_UNORM_S8_UINT }, ComponentMappingRR01);
+		GPixelFormats[PF_DepthStencil].bIs24BitUnormDepthStencil = false;
 	}
 
 	if (FVulkanPlatform::SupportsBCTextureFormats())
