Prevent GL RHI attempting to access data from a null buffer descriptor and triggering an assert.

diff --git a/Engine/Source/Runtime/OpenGLDrv/Private/OpenGLBuffer.cpp b/Engine/Source/Runtime/OpenGLDrv/Private/OpenGLBuffer.cpp
--- a/Engine/Source/Runtime/OpenGLDrv/Private/OpenGLBuffer.cpp
+++ b/Engine/Source/Runtime/OpenGLDrv/Private/OpenGLBuffer.cpp
@@ -219,7 +219,7 @@ FRHIBufferInitializer FOpenGLDynamicRHI::RHICreateBufferInitializer(FRHICommandL
 	const GLenum BufferType = GetNewBufferType(CreateDesc);
 
 	const void* InitialData = nullptr;
-	if (CreateDesc.InitAction == ERHIBufferInitAction::ResourceArray)
+	if (CreateDesc.InitAction == ERHIBufferInitAction::ResourceArray && !CreateDesc.IsNull())
 	{
 		// TODO: move this initialization code path into this method
 		InitialData = CreateDesc.InitialData->GetResourceData();
