Fixing some issues with null RHI buffers.
- RawIndexBuffers shouldn't set CreateDesc.ResourceArray if there is no data to set
- OpenGLDrv should early out for null buffers, matching all other RHIs.

diff --git a/Engine/Source/Runtime/Engine/Private/RawIndexBuffer.cpp b/Engine/Source/Runtime/Engine/Private/RawIndexBuffer.cpp
--- a/Engine/Source/Runtime/Engine/Private/RawIndexBuffer.cpp
+++ b/Engine/Source/Runtime/Engine/Private/RawIndexBuffer.cpp
@@ -491,15 +491,18 @@ FBufferRHIRef FRawStaticIndexBuffer16or32Interface::CreateRHIIndexBufferInternal
 		.SetOwnerName(InOwnerName)
 		.SetInitialState(ERHIAccess::VertexOrIndexBuffer | ERHIAccess::SRVMask);
 
-	if (!CreateDesc.IsNull() && bNeedSRV)
+	if (!CreateDesc.IsNull())
 	{
-		// EBufferUsageFlags::ShaderResource is needed for SkinCache RecomputeSkinTangents
-		CreateDesc.AddUsage(EBufferUsageFlags::ShaderResource);
-	}
+		if (bNeedSRV)
+		{
+			// EBufferUsageFlags::ShaderResource is needed for SkinCache RecomputeSkinTangents
+			CreateDesc.AddUsage(EBufferUsageFlags::ShaderResource);
+		}
 
-	if (ResourceArray)
-	{
-		CreateDesc.SetInitActionResourceArray(ResourceArray);
+		if (ResourceArray)
+		{
+			CreateDesc.SetInitActionResourceArray(ResourceArray);
+		}
 	}
 
 	return RHICmdList.CreateBuffer(CreateDesc);
diff --git a/Engine/Source/Runtime/OpenGLDrv/Private/OpenGLBuffer.cpp b/Engine/Source/Runtime/OpenGLDrv/Private/OpenGLBuffer.cpp
--- a/Engine/Source/Runtime/OpenGLDrv/Private/OpenGLBuffer.cpp
+++ b/Engine/Source/Runtime/OpenGLDrv/Private/OpenGLBuffer.cpp
@@ -218,8 +218,14 @@ FRHIBufferInitializer FOpenGLDynamicRHI::RHICreateBufferInitializer(FRHICommandL
 
 	const GLenum BufferType = GetNewBufferType(CreateDesc);
 
+	if (CreateDesc.IsNull())
+	{
+		FOpenGLBuffer* Buffer = new FOpenGLBuffer(&RHICmdList, BufferType, CreateDesc, nullptr);
+		return UE::RHICore::FDefaultBufferInitializer(RHICmdList, Buffer);
+	}
+
 	const void* InitialData = nullptr;
-	if (CreateDesc.InitAction == ERHIBufferInitAction::ResourceArray && !CreateDesc.IsNull())
+	if (CreateDesc.InitAction == ERHIBufferInitAction::ResourceArray)
 	{
 		// TODO: move this initialization code path into this method
 		InitialData = CreateDesc.InitialData->GetResourceData();
@@ -233,7 +239,7 @@ FRHIBufferInitializer FOpenGLDynamicRHI::RHICreateBufferInitializer(FRHICommandL
 		CreateDesc.InitialData->Discard();
 	}
 
-	if (CreateDesc.IsNull() || CreateDesc.InitAction == ERHIBufferInitAction::ResourceArray || CreateDesc.InitAction == ERHIBufferInitAction::Default)
+	if (CreateDesc.InitAction == ERHIBufferInitAction::ResourceArray || CreateDesc.InitAction == ERHIBufferInitAction::Default)
 	{
 		return UE::RHICore::FDefaultBufferInitializer(RHICmdList, Buffer);
 	}
