Fix compilation warning

diff --git a/Engine/Source/Runtime/OpenGLDrv/Private/OpenGLTexture.cpp b/Engine/Source/Runtime/OpenGLDrv/Private/OpenGLTexture.cpp
--- a/Engine/Source/Runtime/OpenGLDrv/Private/OpenGLTexture.cpp
+++ b/Engine/Source/Runtime/OpenGLDrv/Private/OpenGLTexture.cpp
@@ -349,7 +349,7 @@ FOpenGLTexture::FOpenGLTexture(FRHICommandListBase& RHICmdList, FOpenGLTextureCr
 	, bMultisampleRenderbuffer(CreateDesc.bMultisampleRenderbuffer)
 	, HardwareBufferImage(nullptr)
 {
-	QUICK_SCOPE_CYCLE_COUNTER(STAT_OpenGLCreateTextureTime);
+	SCOPE_CYCLE_COUNTER(STAT_OpenGLCreateTextureTime);
 
 	RHICmdList.EnqueueLambda([this, HardwareBuffer](FRHICommandListBase& CmdList)
 		{
