Remove RHIBeginDrawingViewport
- Platform RHIs either did nothing, or set the viewport as the current RTV. This is redundant since the renderer must start a render pass if it ever wants to draw to an RTV, and therefore it will call RHIBeginRenderPass.
- CustomPresent BeginDrawing was also called in some RHIs, but this is not implemented by any use of the custom present interface in the engine.

diff --git a/Engine/Plugins/Tests/RHITests/Source/RHITests/Private/RHIClearTextureTests.cpp b/Engine/Plugins/Tests/RHITests/Source/RHITests/Private/RHIClearTextureTests.cpp
--- a/Engine/Plugins/Tests/RHITests/Source/RHITests/Private/RHIClearTextureTests.cpp
+++ b/Engine/Plugins/Tests/RHITests/Source/RHITests/Private/RHIClearTextureTests.cpp
@@ -47,13 +47,11 @@ static void Test_ClearTextureBeginCapture(FRHICommandListImmediate& RHICmdList)
 {
 	if (GBeginEndCaptureHack)
 	{
-		RHICmdList.BeginDrawingViewport(nullptr, nullptr);
 		RHICmdList.EndDrawingViewport(nullptr, true, false);
 		GDynamicRHI->RHIWaitForFlip(100000);
 		IRenderCaptureProvider::Get().CaptureFrame();
  		GDynamicRHI->RHISignalFlipEvent();
  		GDynamicRHI->RHIWaitForFlip(100000);
-		RHICmdList.BeginDrawingViewport(nullptr, nullptr);
 	}
 	else
 	{
diff --git a/Engine/Source/Runtime/Apple/MetalRHI/Private/MetalViewport.cpp b/Engine/Source/Runtime/Apple/MetalRHI/Private/MetalViewport.cpp
--- a/Engine/Source/Runtime/Apple/MetalRHI/Private/MetalViewport.cpp
+++ b/Engine/Source/Runtime/Apple/MetalRHI/Private/MetalViewport.cpp
@@ -818,25 +818,6 @@ void FMetalDynamicRHI::RHITick( float DeltaTime )
  *	Viewport functions.
  *=============================================================================*/
 
-void FMetalRHICommandContext::RHIBeginDrawingViewport(FRHIViewport* ViewportRHI, FRHITexture* RenderTargetRHI)
-{
-    MTL_SCOPED_AUTORELEASE_POOL;
-	FMetalViewport* Viewport = ResourceCast(ViewportRHI);
-	check(Viewport);
-
-	// Set the render target and viewport.
-	if (RenderTargetRHI)
-	{
-		FRHIRenderTargetView RTV(RenderTargetRHI, GIsEditor ? ERenderTargetLoadAction::EClear : ERenderTargetLoadAction::ELoad);
-		SetRenderTargets(1, &RTV, nullptr);
-	}
-	else
-	{
-		FRHIRenderTargetView RTV(Viewport->GetBackBuffer(EMetalViewportAccessRHI), GIsEditor ? ERenderTargetLoadAction::EClear : ERenderTargetLoadAction::ELoad);
-		SetRenderTargets(1, &RTV, nullptr);
-	}
-}
-
 void FMetalRHICommandContext::RHIEndDrawingViewport(FRHIViewport* ViewportRHI,bool bPresent,bool bLockToVsync)
 {
     MTL_SCOPED_AUTORELEASE_POOL;
diff --git a/Engine/Source/Runtime/Apple/MetalRHI/Public/MetalRHIContext.h b/Engine/Source/Runtime/Apple/MetalRHI/Public/MetalRHIContext.h
--- a/Engine/Source/Runtime/Apple/MetalRHI/Public/MetalRHIContext.h
+++ b/Engine/Source/Runtime/Apple/MetalRHI/Public/MetalRHIContext.h
@@ -90,9 +90,6 @@ public:
 
 	virtual void RHIDiscardRenderTargets(bool Depth, bool Stencil, uint32 ColorBitMask) final override;
 	
-	// This method is queued with an RHIThread, otherwise it will flush after it is queued; without an RHI thread there is no benefit to queuing this frame advance commands
-	virtual void RHIBeginDrawingViewport(FRHIViewport* Viewport, FRHITexture* RenderTargetRHI) override;
-	
 	// This method is queued with an RHIThread, otherwise it will flush after it is queued; without an RHI thread there is no benefit to queuing this frame advance commands
 	virtual void RHIEndDrawingViewport(FRHIViewport* Viewport, bool bPresent, bool bLockToVsync) override;
 	
diff --git a/Engine/Source/Runtime/D3D12RHI/Private/D3D12Adapter.cpp b/Engine/Source/Runtime/D3D12RHI/Private/D3D12Adapter.cpp
--- a/Engine/Source/Runtime/D3D12RHI/Private/D3D12Adapter.cpp
+++ b/Engine/Source/Runtime/D3D12RHI/Private/D3D12Adapter.cpp
@@ -922,11 +922,6 @@ void FD3D12Adapter::CreateD3DInfoQueue()
 #endif
 }
 
-void FD3D12Adapter::SetDrawingViewport(FD3D12Viewport* InViewport)
-{
-	DrawingViewport = InViewport;
-}
-
 FD3D12TransientHeapCache& FD3D12Adapter::GetOrCreateTransientHeapCache()
 {
 	if (!TransientMemoryCache)
@@ -1531,7 +1526,6 @@ FD3D12Adapter::~FD3D12Adapter()
 	}
 
 	Viewports.Empty();
-	DrawingViewport = nullptr;
 
 	for (uint32 GPUIndex : FRHIGPUMask::All())
 	{
diff --git a/Engine/Source/Runtime/D3D12RHI/Private/D3D12Adapter.h b/Engine/Source/Runtime/D3D12RHI/Private/D3D12Adapter.h
--- a/Engine/Source/Runtime/D3D12RHI/Private/D3D12Adapter.h
+++ b/Engine/Source/Runtime/D3D12RHI/Private/D3D12Adapter.h
@@ -220,8 +220,6 @@ public:
 	FORCEINLINE const FD3D12AdapterDesc& GetDesc() const { return Desc; }
 
 	FORCEINLINE TArray<FD3D12Viewport*>& GetViewports() { return Viewports; }
-	FORCEINLINE FD3D12Viewport* GetDrawingViewport() { return DrawingViewport; }
-	void SetDrawingViewport(FD3D12Viewport* InViewport);
 
 	FORCEINLINE int32 GetMaxDescriptorsForHeapType(ERHIDescriptorHeapType InHeapType) { return InHeapType == ERHIDescriptorHeapType::Sampler ? MaxSamplerDescriptors : MaxNonSamplerDescriptors; }
 
@@ -576,9 +574,6 @@ protected:
 	/** A list of all viewport RHIs that have been created. */
 	TArray<FD3D12Viewport*> Viewports;
 
-	/** The viewport which is currently being drawn. */
-	TRefCountPtr<FD3D12Viewport> DrawingViewport;
-
 	/** A Fence whose value increases every frame*/
 	TUniquePtr<FD3D12ManualFence> FrameFence;
 
diff --git a/Engine/Source/Runtime/D3D12RHI/Private/D3D12CommandContext.h b/Engine/Source/Runtime/D3D12RHI/Private/D3D12CommandContext.h
--- a/Engine/Source/Runtime/D3D12RHI/Private/D3D12CommandContext.h
+++ b/Engine/Source/Runtime/D3D12RHI/Private/D3D12CommandContext.h
@@ -474,7 +474,6 @@ class FD3D12CommandContextBase : public IRHICommandContext, public FD3D12Adapter
 public:
 	FD3D12CommandContextBase(FD3D12Adapter* InParent, FRHIGPUMask InGPUMask);
 
-	void RHIBeginDrawingViewport(FRHIViewport* Viewport, FRHITexture* RenderTargetRHI) final override;
 	void RHIEndDrawingViewport(FRHIViewport* Viewport, bool bPresent, bool bLockToVsync) final override;
 
 	FRHIGPUMask GetGPUMask() const { return GPUMask; }
diff --git a/Engine/Source/Runtime/D3D12RHI/Private/D3D12Viewport.cpp b/Engine/Source/Runtime/D3D12RHI/Private/D3D12Viewport.cpp
--- a/Engine/Source/Runtime/D3D12RHI/Private/D3D12Viewport.cpp
+++ b/Engine/Source/Runtime/D3D12RHI/Private/D3D12Viewport.cpp
@@ -730,41 +730,6 @@ void FD3D12DynamicRHI::RHITick(float DeltaTime)
  *	Viewport functions.
  *=============================================================================*/
 
-void FD3D12CommandContextBase::RHIBeginDrawingViewport(FRHIViewport* ViewportRHI, FRHITexture* RenderTargetRHI)
-{
-	FD3D12Viewport* Viewport = FD3D12DynamicRHI::ResourceCast(ViewportRHI);
-
-	SCOPE_CYCLE_COUNTER(STAT_D3D12PresentTime);
-
-	// Set the viewport.
-	check(!ParentAdapter->GetDrawingViewport());
-	ParentAdapter->SetDrawingViewport(Viewport);
-
-	if (RenderTargetRHI == nullptr)
-	{
-		RenderTargetRHI = Viewport->GetBackBuffer_RHIThread();
-	}
-
-#if !UE_BUILD_SHIPPING
-	if (RHIConsoleVariables::LogViewportEvents)
-	{
-		const FString& ThreadName = FThreadManager::GetThreadName(FPlatformTLS::GetCurrentThreadId());
-		UE_LOG(LogD3D12RHI, Log, TEXT("Thread %s: RHIBeginDrawingViewport (Viewport %#016llx: BackBuffer %#016llx: CmdList: %016llx)"),
-			ThreadName.GetCharArray().GetData(),
-			Viewport,
-			RenderTargetRHI,
-			GetSingleDeviceContext(0)->BaseCommandList().GetNoRefCount()
-		);
-	}
-#endif
-
-	FRHICustomPresent* CustomPresent = Viewport->GetCustomPresent();
-	if (CustomPresent)
-	{
-		CustomPresent->BeginDrawing();
-	}
-}
-
 void FD3D12CommandContextBase::RHIEndDrawingViewport(FRHIViewport* ViewportRHI, bool bPresent, bool bLockToVsync)
 {
 	FD3D12Viewport* Viewport = FD3D12DynamicRHI::ResourceCast(ViewportRHI);
@@ -784,9 +749,6 @@ void FD3D12CommandContextBase::RHIEndDrawingViewport(FRHIViewport* ViewportRHI,
 
 	SCOPE_CYCLE_COUNTER(STAT_D3D12PresentTime);
 
-	check(ParentAdapter->GetDrawingViewport() == Viewport);
-	ParentAdapter->SetDrawingViewport(nullptr);
-
 	bool bNativelyPresented = true;
 	if (bPresent)
 	{
diff --git a/Engine/Source/Runtime/Engine/Private/UnrealClient.cpp b/Engine/Source/Runtime/Engine/Private/UnrealClient.cpp
--- a/Engine/Source/Runtime/Engine/Private/UnrealClient.cpp
+++ b/Engine/Source/Runtime/Engine/Private/UnrealClient.cpp
@@ -1647,7 +1647,6 @@ static void ViewportEndDrawing(FRHICommandListImmediate& RHICmdList, FEndDrawing
 void FViewport::BeginRenderFrame(FRHICommandListImmediate& RHICmdList)
 {
 	check( IsInRenderingThread() );
-	RHICmdList.BeginDrawingViewport(GetViewportRHI(), FTextureRHIRef());
 	UpdateRenderTargetSurfaceRHIToCurrentBackBuffer();
 }
 
diff --git a/Engine/Source/Runtime/NullDrv/Public/NullRHI.h b/Engine/Source/Runtime/NullDrv/Public/NullRHI.h
--- a/Engine/Source/Runtime/NullDrv/Public/NullRHI.h
+++ b/Engine/Source/Runtime/NullDrv/Public/NullRHI.h
@@ -420,11 +420,6 @@ public:
 		return true; 
 	}
 
-
-	virtual void RHIBeginDrawingViewport(FRHIViewport* Viewport, FRHITexture* RenderTargetRHI) final override
-	{
-	}
-
 	virtual void RHIEndDrawingViewport(FRHIViewport* Viewport, bool bPresent, bool bLockToVsync) final override
 	{
 	}
diff --git a/Engine/Source/Runtime/OpenGLDrv/Private/OpenGLDrv.h b/Engine/Source/Runtime/OpenGLDrv/Private/OpenGLDrv.h
--- a/Engine/Source/Runtime/OpenGLDrv/Private/OpenGLDrv.h
+++ b/Engine/Source/Runtime/OpenGLDrv/Private/OpenGLDrv.h
@@ -408,7 +408,6 @@ public:
 	virtual void RHIBeginRenderQuery(FRHIRenderQuery* RenderQuery) final override;
 	virtual void RHIEndRenderQuery(FRHIRenderQuery* RenderQuery) final override;
 
-	virtual void RHIBeginDrawingViewport(FRHIViewport* Viewport, FRHITexture* RenderTargetRHI) final override;
 	virtual void RHIEndDrawingViewport(FRHIViewport* Viewport, bool bPresent, bool bLockToVsync) final override;
 	virtual void RHIEndFrame(const FRHIEndFrameArgs& Args) final override;
 	virtual void RHISetStreamSource(uint32 StreamIndex, FRHIBuffer* VertexBuffer, uint32 Offset) final override;
@@ -679,7 +678,6 @@ private:
 
 	/** A list of all viewport RHIs that have been created. */
 	TArray<FOpenGLViewport*> Viewports;
-	TRefCountPtr<FOpenGLViewport>		DrawingViewport;
 	bool								bRevertToSharedContextAfterDrawingViewport = false;
 
 	EPrimitiveType						PrimitiveType = PT_Num;
diff --git a/Engine/Source/Runtime/OpenGLDrv/Private/OpenGLViewport.cpp b/Engine/Source/Runtime/OpenGLDrv/Private/OpenGLViewport.cpp
--- a/Engine/Source/Runtime/OpenGLDrv/Private/OpenGLViewport.cpp
+++ b/Engine/Source/Runtime/OpenGLDrv/Private/OpenGLViewport.cpp
@@ -109,35 +109,6 @@ void FOpenGLDynamicRHI::RHITick( float DeltaTime )
  *	Viewport functions.
  *=============================================================================*/
 
-void FOpenGLDynamicRHI::RHIBeginDrawingViewport(FRHIViewport* ViewportRHI, FRHITexture* RenderTarget)
-{
-	VERIFY_GL_SCOPE();
-
-	FOpenGLViewport* Viewport = ResourceCast(ViewportRHI);
-
-	SCOPE_CYCLE_COUNTER(STAT_OpenGLPresentTime);
-
-	check(!DrawingViewport);
-	DrawingViewport = Viewport;
-
-	// Set the render target and viewport.
-	if (RenderTarget)
-	{
-		FRHIRenderTargetView RTV(RenderTarget, ERenderTargetLoadAction::ELoad);
-		SetRenderTargets(1, &RTV, nullptr);
-	}
-	else
-	{
-		FRHIRenderTargetView RTV(DrawingViewport->GetBackBuffer(), ERenderTargetLoadAction::ELoad);
-		SetRenderTargets(1, &RTV, nullptr);
-	}
-
-	if (IsValidRef(CustomPresent))
-	{
-		CustomPresent->BeginDrawing();
-	}
-}
-
 void FOpenGLDynamicRHI::RHIEndDrawingViewport(FRHIViewport* ViewportRHI,bool bPresent,bool bLockToVsync)
 {
 	VERIFY_GL_SCOPE();
@@ -148,8 +119,6 @@ void FOpenGLDynamicRHI::RHIEndDrawingViewport(FRHIViewport* ViewportRHI,bool bPr
 	{
 		FRenderThreadIdleScope IdleScope(ERenderThreadIdleTypes::WaitingForGPUPresent);
 
-		check(DrawingViewport.GetReference() == Viewport);
-
 		FOpenGLTexture* BackBuffer = Viewport->GetBackBuffer();
 
 		if (ContextState.bScissorEnabled)
@@ -169,8 +138,6 @@ void FOpenGLDynamicRHI::RHIEndDrawingViewport(FRHIViewport* ViewportRHI,bool bPr
 		// Always consider the Framebuffer in the rendering context dirty after the blit
 		ContextState.Framebuffer = -1;
 
-		DrawingViewport = NULL;
-
 		if (bNeedFinishFrame)
 		{
 			static const auto CFinishFrameVar = IConsoleManager::Get().FindTConsoleVariableDataInt(TEXT("r.FinishCurrentFrame"));
diff --git a/Engine/Source/Runtime/RHI/Private/RHICommandList.cpp b/Engine/Source/Runtime/RHI/Private/RHICommandList.cpp
--- a/Engine/Source/Runtime/RHI/Private/RHICommandList.cpp
+++ b/Engine/Source/Runtime/RHI/Private/RHICommandList.cpp
@@ -1900,24 +1900,6 @@ void FRHICommandListImmediate::UnStallRHIThread()
 	}
 }
 
-void FRHICommandListImmediate::BeginDrawingViewport(FRHIViewport* Viewport, FRHITexture* RenderTargetRHI)
-{
-	check(IsImmediate() && IsInRenderingThread());
-	if (Bypass())
-	{
-		GetContext().RHIBeginDrawingViewport(Viewport, RenderTargetRHI);
-		return;
-	}
-	ALLOC_COMMAND(FRHICommandBeginDrawingViewport)(Viewport, RenderTargetRHI);
-	if (!IsRunningRHIInSeparateThread())
-	{
-		// if we aren't running an RHIThread, there is no good reason to buffer this frame advance stuff and that complicates state management, so flush everything out now
-		QUICK_SCOPE_CYCLE_COUNTER(BeginDrawingViewport_Flush);
-		CSV_SCOPED_TIMING_STAT(RHITFlushes, BeginDrawingViewport);
-		ImmediateFlush(EImmediateFlushType::FlushRHIThread);
-	}
-}
-
 void FRHICommandListImmediate::EndDrawingViewport(FRHIViewport* Viewport, bool bPresent, bool bLockToVsync)
 {
 	// Make sure all prior graphics and async compute work has been submitted.
diff --git a/Engine/Source/Runtime/RHI/Public/RHICommandList.h b/Engine/Source/Runtime/RHI/Public/RHICommandList.h
--- a/Engine/Source/Runtime/RHI/Public/RHICommandList.h
+++ b/Engine/Source/Runtime/RHI/Public/RHICommandList.h
@@ -2504,19 +2504,6 @@ FRHICOMMAND_MACRO(FRHICommandPostExternalCommandsReset)
 	RHI_EXECUTE_API void Execute(FRHICommandListBase& CmdList);
 };
 
-FRHICOMMAND_MACRO(FRHICommandBeginDrawingViewport)
-{
-	FRHIViewport* Viewport;
-	FRHITexture* RenderTargetRHI;
-
-	FORCEINLINE_DEBUGGABLE FRHICommandBeginDrawingViewport(FRHIViewport* InViewport, FRHITexture* InRenderTargetRHI)
-		: Viewport(InViewport)
-		, RenderTargetRHI(InRenderTargetRHI)
-	{
-	}
-	RHI_EXECUTE_API void Execute(FRHICommandListBase& CmdList);
-};
-
 FRHICOMMAND_MACRO(FRHICommandEndDrawingViewport)
 {
 	FRHIViewport* Viewport;
@@ -4650,7 +4637,8 @@ public:
 		return static_cast<FRHICommandListImmediate&>(RHICmdList);
 	}
 
-	RHI_API void BeginDrawingViewport(FRHIViewport* Viewport, FRHITexture* RenderTargetRHI);
+	UE_DEPRECATED(5.7, "FRHICommandListImmediate::BeginDrawingViewport is deprecated. Calls to this function are unnecessary and can be removed. There is no replacement.")
+	inline void BeginDrawingViewport(FRHIViewport* Viewport, FRHITexture* RenderTargetRHI) {}
 	RHI_API void EndDrawingViewport(FRHIViewport* Viewport, bool bPresent, bool bLockToVsync);
 
 	RHI_API void EndFrame();
diff --git a/Engine/Source/Runtime/RHI/Public/RHICommandListCommandExecutes.inl b/Engine/Source/Runtime/RHI/Public/RHICommandListCommandExecutes.inl
--- a/Engine/Source/Runtime/RHI/Public/RHICommandListCommandExecutes.inl
+++ b/Engine/Source/Runtime/RHI/Public/RHICommandListCommandExecutes.inl
@@ -18,7 +18,6 @@ class IRHIComputeContext;
 struct FComputedBSS;
 struct FComputedUniformBuffer;
 struct FMemory;
-struct FRHICommandBeginDrawingViewport;
 struct FRHICommandBeginRenderQuery;
 struct FRHICommandBuildLocalBoundShaderState;
 struct FRHICommandBuildLocalGraphicsPipelineState;
@@ -552,12 +551,6 @@ void FRHICommandSetBindingsOnShaderBindingTable::Execute(FRHICommandListBase& Cm
 	INTERNAL_DECORATOR(RHISetBindingsOnShaderBindingTable)(SBT, RayTracingPipelineState, NumBindings, Bindings, BindingType);
 }
 
-void FRHICommandBeginDrawingViewport::Execute(FRHICommandListBase& CmdList)
-{
-	RHISTAT(BeginDrawingViewport);
-	INTERNAL_DECORATOR(RHIBeginDrawingViewport)(Viewport, RenderTargetRHI);
-}
-
 void FRHICommandEndDrawingViewport::Execute(FRHICommandListBase& CmdList)
 {
 	RHISTAT(EndDrawingViewport);
diff --git a/Engine/Source/Runtime/RHI/Public/RHIContext.h b/Engine/Source/Runtime/RHI/Public/RHIContext.h
--- a/Engine/Source/Runtime/RHI/Public/RHIContext.h
+++ b/Engine/Source/Runtime/RHI/Public/RHIContext.h
@@ -679,9 +679,6 @@ public:
 	// Not all RHIs need this (Mobile specific)
 	virtual void RHIDiscardRenderTargets(bool Depth, bool Stencil, uint32 ColorBitMask) {}
 
-	// This method is queued with an RHIThread, otherwise it will flush after it is queued; without an RHI thread there is no benefit to queuing this frame advance commands
-	virtual void RHIBeginDrawingViewport(FRHIViewport* Viewport, FRHITexture* RenderTargetRHI) = 0;
-
 	// This method is queued with an RHIThread, otherwise it will flush after it is queued; without an RHI thread there is no benefit to queuing this frame advance commands
 	virtual void RHIEndDrawingViewport(FRHIViewport* Viewport, bool bPresent, bool bLockToVsync) = 0;
 
diff --git a/Engine/Source/Runtime/RHI/Public/RHIGlobals.h b/Engine/Source/Runtime/RHI/Public/RHIGlobals.h
--- a/Engine/Source/Runtime/RHI/Public/RHIGlobals.h
+++ b/Engine/Source/Runtime/RHI/Public/RHIGlobals.h
@@ -495,15 +495,7 @@ struct FRHIGlobals
 	/** Whether or not the RHI supports 16bit VALU and resource loads (HLSL shader model 6.2). */
 	bool SupportsNative16BitOps = false;
 
-	/** Whether or not the RHI supports an RHI thread.
-	Requirements for RHI thread
-	* Microresources (those in RHIStaticStates.h) need to be able to be created by any thread at any time and be able to work with a radically simplified rhi resource lifecycle. CreateSamplerState, CreateRasterizerState, CreateDepthStencilState, CreateBlendState
-	* CreateUniformBuffer needs to be threadsafe
-	* GetRenderQueryResult should be threadsafe, but this isn't required. If it isn't threadsafe, then you need to flush yourself in the RHI
-	* GetViewportBackBuffer and AdvanceFrameForGetViewportBackBuffer need to be threadsafe and need to support the fact that the render thread has a different concept of "current backbuffer" than the RHI thread. Without an RHIThread this is moot due to the next two items.
-	* AdvanceFrameForGetViewportBackBuffer needs be added as an RHI method and this needs to work with GetViewportBackBuffer to give the render thread the right back buffer even though many commands relating to the beginning and end of the frame are queued.
-	* BeginDrawingViewport, and 5 or so other frame advance methods are queued with an RHIThread. Without an RHIThread, these just flush internally.
-	***/
+	/** Whether or not the RHI supports an RHI thread. */
 	bool SupportsRHIThread = false;
 
 	/* as above, but we run the commands on arbitrary task threads */
diff --git a/Engine/Source/Runtime/RHI/Public/RHIResources.h b/Engine/Source/Runtime/RHI/Public/RHIResources.h
--- a/Engine/Source/Runtime/RHI/Public/RHIResources.h
+++ b/Engine/Source/Runtime/RHI/Public/RHIResources.h
@@ -4200,6 +4200,7 @@ public:
 	virtual bool NeedsAdvanceBackbuffer() { return false; };
 
 	// Called from RHI thread when the engine begins drawing to the viewport.
+	UE_DEPRECATED(5.7, "FRHICustomPresent::BeginDrawing is deprecated and never called.")
 	virtual void BeginDrawing() {};
 
 	// Called from RHI thread to perform custom present.
diff --git a/Engine/Source/Runtime/RHI/Public/RHIValidationContext.h b/Engine/Source/Runtime/RHI/Public/RHIValidationContext.h
--- a/Engine/Source/Runtime/RHI/Public/RHIValidationContext.h
+++ b/Engine/Source/Runtime/RHI/Public/RHIValidationContext.h
@@ -756,12 +756,6 @@ public:
 		RHIContext->RHIDiscardRenderTargets(bDepth, bStencil, ColorBitMask);
 	}
 
-	// This method is queued with an RHIThread, otherwise it will flush after it is queued; without an RHI thread there is no benefit to queuing this frame advance commands
-	virtual void RHIBeginDrawingViewport(FRHIViewport* Viewport, FRHITexture* RenderTargetRHI) override final
-	{
-		RHIContext->RHIBeginDrawingViewport(Viewport, RenderTargetRHI);
-	}
-
 	// This method is queued with an RHIThread, otherwise it will flush after it is queued; without an RHI thread there is no benefit to queuing this frame advance commands
 	virtual void RHIEndDrawingViewport(FRHIViewport* Viewport, bool bPresent, bool bLockToVsync) override final
 	{
diff --git a/Engine/Source/Runtime/Renderer/Private/ScenePrivate.h b/Engine/Source/Runtime/Renderer/Private/ScenePrivate.h
--- a/Engine/Source/Runtime/Renderer/Private/ScenePrivate.h
+++ b/Engine/Source/Runtime/Renderer/Private/ScenePrivate.h
@@ -2517,7 +2517,7 @@ class FSceneVelocityData
 public:
 
 	/**
-	 * Must be called once per frame, even when there are multiple BeginDrawingViewports.
+	 * Must be called once per frame, even when there are multiple EndDrawingViewports.
 	 */
 	void StartFrame(FScene* Scene);
 
diff --git a/Engine/Source/Runtime/SlateRHIRenderer/Private/SlateRHIRenderer.cpp b/Engine/Source/Runtime/SlateRHIRenderer/Private/SlateRHIRenderer.cpp
--- a/Engine/Source/Runtime/SlateRHIRenderer/Private/SlateRHIRenderer.cpp
+++ b/Engine/Source/Runtime/SlateRHIRenderer/Private/SlateRHIRenderer.cpp
@@ -970,8 +970,6 @@ void FSlateRHIRenderer::PresentWindow_RenderThread(FRHICommandListImmediate& RHI
 
 	FRHITexture* OptionalSDRBuffer = nullptr;
 
-	RHICmdList.BeginDrawingViewport(DrawPassOutputs.ViewportRHI, FTextureRHIRef());
-
 	if (GRHIGlobals.NeedsExtraTransitions)
 	{
 		RHICmdList.Transition(FRHITransitionInfo(DrawPassOutputs.OutputTextureRHI, ERHIAccess::Unknown, ERHIAccess::Present));
diff --git a/Engine/Source/Runtime/VulkanRHI/Private/VulkanContext.h b/Engine/Source/Runtime/VulkanRHI/Private/VulkanContext.h
--- a/Engine/Source/Runtime/VulkanRHI/Private/VulkanContext.h
+++ b/Engine/Source/Runtime/VulkanRHI/Private/VulkanContext.h
@@ -318,7 +318,6 @@ public:
 
 	virtual void Finalize(TArray<FVulkanPayload*>& OutPayloads) override;
 
-	virtual void RHIBeginDrawingViewport(FRHIViewport* Viewport, FRHITexture* RenderTargetRHI) final override;
 	virtual void RHIEndDrawingViewport(FRHIViewport* Viewport, bool bPresent, bool bLockToVsync) final override;
 
 	virtual void RHIBeginRenderPass(const FRHIRenderPassInfo& InInfo, const TCHAR* InName) final override;
diff --git a/Engine/Source/Runtime/VulkanRHI/Private/VulkanDynamicRHI.h b/Engine/Source/Runtime/VulkanRHI/Private/VulkanDynamicRHI.h
--- a/Engine/Source/Runtime/VulkanRHI/Private/VulkanDynamicRHI.h
+++ b/Engine/Source/Runtime/VulkanRHI/Private/VulkanDynamicRHI.h
@@ -373,9 +373,6 @@ protected:
 	/** A list of all viewport RHIs that have been created. */
 	TArray<FVulkanViewport*> Viewports;
 
-	/** The viewport which is currently being drawn. */
-	TRefCountPtr<FVulkanViewport> DrawingViewport;
-
 	void CreateInstance();
 	void SelectDevice();
 
diff --git a/Engine/Source/Runtime/VulkanRHI/Private/VulkanRHI.cpp b/Engine/Source/Runtime/VulkanRHI/Private/VulkanRHI.cpp
--- a/Engine/Source/Runtime/VulkanRHI/Private/VulkanRHI.cpp
+++ b/Engine/Source/Runtime/VulkanRHI/Private/VulkanRHI.cpp
@@ -430,7 +430,6 @@ FVulkanCommandListContextImmediate::FVulkanCommandListContextImmediate(FVulkanDe
 FVulkanDynamicRHI::FVulkanDynamicRHI()
 	: Instance(VK_NULL_HANDLE)
 	, Device(nullptr)
-	, DrawingViewport(nullptr)
 {
 	// This should be called once at the start 
 	check(IsInGameThread());
@@ -1241,27 +1240,12 @@ FVulkanTiming::FVulkanTiming(FVulkanQueue& InQueue)
 }
 #endif
 
-void FVulkanCommandListContext::RHIBeginDrawingViewport(FRHIViewport* ViewportRHI, FRHITexture* RenderTargetRHI)
-{
-	//FRCLog::Printf(FString::Printf(TEXT("FVulkanCommandListContext::RHIBeginDrawingViewport\n")));
-	check(ViewportRHI);
-	FVulkanViewport* Viewport = ResourceCast(ViewportRHI);
-	FVulkanDynamicRHI::Get().DrawingViewport = Viewport;
-
-	FRHICustomPresent* CustomPresent = Viewport->GetCustomPresent();
-	if (CustomPresent)
-	{
-		CustomPresent->BeginDrawing();
-	}
-}
-
 void FVulkanCommandListContext::RHIEndDrawingViewport(FRHIViewport* ViewportRHI, bool bPresent, bool bLockToVsync)
 {
 	LLM_SCOPE_VULKAN(ELLMTagVulkan::VulkanMisc);
 	//FRCLog::Printf(FString::Printf(TEXT("FVulkanCommandListContext::RHIEndDrawingViewport()")));
 	check(IsImmediate());
 	FVulkanViewport* Viewport = ResourceCast(ViewportRHI);
-	check(Viewport == FVulkanDynamicRHI::Get().DrawingViewport);
 
 	//#todo-rco: Unbind all pending state
 /*
@@ -1276,8 +1260,6 @@ void FVulkanCommandListContext::RHIEndDrawingViewport(FRHIViewport* ViewportRHI,
 	{
 		//#todo-rco: Check for r.FinishCurrentFrame
 	}
-
-	FVulkanDynamicRHI::Get().DrawingViewport = nullptr;
 }
 
 #if WITH_RHI_BREADCRUMBS
diff --git a/Engine/Source/Runtime/Windows/D3D11RHI/Private/D3D11RHIPrivate.h b/Engine/Source/Runtime/Windows/D3D11RHI/Private/D3D11RHIPrivate.h
--- a/Engine/Source/Runtime/Windows/D3D11RHI/Private/D3D11RHIPrivate.h
+++ b/Engine/Source/Runtime/Windows/D3D11RHI/Private/D3D11RHIPrivate.h
@@ -537,7 +537,6 @@ public:
 	virtual void RHIBeginRenderQuery(FRHIRenderQuery* RenderQuery) final override;
 	virtual void RHIEndRenderQuery  (FRHIRenderQuery* RenderQuery) final override;
 
-	virtual void RHIBeginDrawingViewport(FRHIViewport* Viewport, FRHITexture* RenderTargetRHI) final override;
 	virtual void RHIEndDrawingViewport(FRHIViewport* Viewport, bool bPresent, bool bLockToVsync) final override;
 	virtual void RHIEndFrame(const FRHIEndFrameArgs& Args) final override;
 	virtual void RHISetStreamSource(uint32 StreamIndex, FRHIBuffer* VertexBuffer, uint32 Offset) final override;
@@ -736,9 +735,6 @@ protected:
 	/** A list of all viewport RHIs that have been created. */
 	TArray<FD3D11Viewport*> Viewports;
 
-	/** The viewport which is currently being drawn. */
-	TRefCountPtr<FD3D11Viewport> DrawingViewport;
-
 	/** The feature level of the device. */
 	D3D_FEATURE_LEVEL FeatureLevel;
 
diff --git a/Engine/Source/Runtime/Windows/D3D11RHI/Private/D3D11Viewport.cpp b/Engine/Source/Runtime/Windows/D3D11RHI/Private/D3D11Viewport.cpp
--- a/Engine/Source/Runtime/Windows/D3D11RHI/Private/D3D11Viewport.cpp
+++ b/Engine/Source/Runtime/Windows/D3D11RHI/Private/D3D11Viewport.cpp
@@ -748,33 +748,6 @@ void FD3D11DynamicRHI::RHITick( float DeltaTime )
  *	Viewport functions.
  *=============================================================================*/
 
-void FD3D11DynamicRHI::RHIBeginDrawingViewport(FRHIViewport* ViewportRHI, FRHITexture* RenderTarget)
-{
-	FD3D11Viewport* Viewport = ResourceCast(ViewportRHI);
-
-	SCOPE_CYCLE_COUNTER(STAT_D3D11PresentTime);
-
-	check(!DrawingViewport);
-	DrawingViewport = Viewport;
-
-	// Set the render target and viewport.
-	if( RenderTarget == NULL )
-	{
-		RenderTarget = Viewport->GetBackBuffer();
-	}
-	FRHIRenderTargetView View(RenderTarget, ERenderTargetLoadAction::ELoad);
-	SetRenderTargets(1,&View,nullptr);
-
-	// Set an initially disabled scissor rect.
-	RHISetScissorRect(false,0,0,0,0);
-
-	FRHICustomPresent* CustomPresent = Viewport->GetCustomPresent();
-	if (CustomPresent)
-	{
-		CustomPresent->BeginDrawing();
-	}
-}
-
 void FD3D11DynamicRHI::RHIEndDrawingViewport(FRHIViewport* ViewportRHI,bool bPresent,bool bLockToVsync)
 {
 	++PresentCounter;
@@ -782,9 +755,6 @@ void FD3D11DynamicRHI::RHIEndDrawingViewport(FRHIViewport* ViewportRHI,bool bPre
 
 	SCOPE_CYCLE_COUNTER(STAT_D3D11PresentTime);
 
-	check(DrawingViewport.GetReference() == Viewport);
-	DrawingViewport = NULL;
-
 	// Clear references the device might have to resources.
 	CurrentDepthTexture = NULL;
 	CurrentDepthStencilTarget = NULL;
