Fixed static validation warning: pointer was utilized before it was verified against nullptr

diff --git a/Engine/Source/Runtime/OpenGLDrv/Private/OpenGLRenderTarget.cpp b/Engine/Source/Runtime/OpenGLDrv/Private/OpenGLRenderTarget.cpp
--- a/Engine/Source/Runtime/OpenGLDrv/Private/OpenGLRenderTarget.cpp
+++ b/Engine/Source/Runtime/OpenGLDrv/Private/OpenGLRenderTarget.cpp
@@ -54,8 +54,8 @@ public:
 		for( RenderTargetIndex = 0; RenderTargetIndex < InNumRenderTargets; ++RenderTargetIndex )
 		{
 			FMemory::Memzero(RenderTargets[RenderTargetIndex]); // since memcmp is used, we need to zero memory
-			RenderTargets[RenderTargetIndex].Texture = InRenderTargets[RenderTargetIndex];
-			RenderTargets[RenderTargetIndex].Resource = (InRenderTargets[RenderTargetIndex]) ? InRenderTargets[RenderTargetIndex]->GetResource() : 0;
+			RenderTargets[RenderTargetIndex].Texture = InRenderTargets ? InRenderTargets[RenderTargetIndex] : nullptr;
+			RenderTargets[RenderTargetIndex].Resource = (InRenderTargets && InRenderTargets[RenderTargetIndex]) ? InRenderTargets[RenderTargetIndex]->GetResource() : 0;
 			RenderTargets[RenderTargetIndex].MipmapLevel = InRenderTargetMipmapLevels[RenderTargetIndex];
 			RenderTargets[RenderTargetIndex].ArrayIndex = (InRenderTargetArrayIndices == NULL || InRenderTargetArrayIndices[RenderTargetIndex] == -1) ? ALL_SLICES : InRenderTargetArrayIndices[RenderTargetIndex];
 		}
