Fix a bug of custom stencil srv reused as custom depth on opengl platform.

diff --git a/Engine/Source/Runtime/OpenGLDrv/Private/OpenGLCommands.cpp b/Engine/Source/Runtime/OpenGLDrv/Private/OpenGLCommands.cpp
--- a/Engine/Source/Runtime/OpenGLDrv/Private/OpenGLCommands.cpp
+++ b/Engine/Source/Runtime/OpenGLDrv/Private/OpenGLCommands.cpp
@@ -615,6 +615,18 @@ void FOpenGLDynamicRHI::SetupTexturesForDraw(const StateType& ShaderState, int32
 					FOpenGL::BindSampler(TextureStageIndex, PendingSampler ? PendingSampler->Resource : 0);
 					ContextState.SamplerStates[TextureStageIndex] = PendingSampler;
 				}
+
+				if (TextureStage.Texture && TextureStage.Texture->bInStencilTextureMode)
+				{
+					//set the texture to return the depth component, and then force the components to match D3D
+					glTexParameteri(TextureStage.Target, GL_DEPTH_STENCIL_TEXTURE_MODE, GL_DEPTH_COMPONENT);
+					glTexParameteri(TextureStage.Target, GL_TEXTURE_SWIZZLE_R, GL_RED);
+					glTexParameteri(TextureStage.Target, GL_TEXTURE_SWIZZLE_G, GL_ZERO);
+					glTexParameteri(TextureStage.Target, GL_TEXTURE_SWIZZLE_B, GL_ZERO);
+					glTexParameteri(TextureStage.Target, GL_TEXTURE_SWIZZLE_A, GL_ZERO);
+
+					TextureStage.Texture->bInStencilTextureMode = false;
+				}
 			}
 			else if (TextureStage.Target != GL_TEXTURE_BUFFER)
 			{
diff --git a/Engine/Source/Runtime/OpenGLDrv/Private/OpenGLResources.h b/Engine/Source/Runtime/OpenGLDrv/Private/OpenGLResources.h
--- a/Engine/Source/Runtime/OpenGLDrv/Private/OpenGLResources.h
+++ b/Engine/Source/Runtime/OpenGLDrv/Private/OpenGLResources.h
@@ -1090,6 +1090,8 @@ public:
 
 	TArray<TRefCountPtr<FOpenGLPixelBuffer>> PixelBuffers;
 
+	bool bInStencilTextureMode = false;
+
 #if PLATFORM_ANDROID
 	// The image created from an external hardware buffer.
 	EGLImageKHR HardwareBufferImage;
diff --git a/Engine/Source/Runtime/OpenGLDrv/Private/OpenGLUAV.cpp b/Engine/Source/Runtime/OpenGLDrv/Private/OpenGLUAV.cpp
--- a/Engine/Source/Runtime/OpenGLDrv/Private/OpenGLUAV.cpp
+++ b/Engine/Source/Runtime/OpenGLDrv/Private/OpenGLUAV.cpp
@@ -136,6 +136,8 @@ void FOpenGLShaderResourceView::UpdateView()
 			glTexParameteri(Target, GL_TEXTURE_SWIZZLE_G, GL_RED);
 			glTexParameteri(Target, GL_TEXTURE_SWIZZLE_B, GL_ZERO);
 			glTexParameteri(Target, GL_TEXTURE_SWIZZLE_A, GL_ZERO);
+
+			Texture->bInStencilTextureMode = true;
 		}
 	}
 }
