Mobile: Set bSupportsShaderFramebufferFetchProgrammableBlending only when use of FBF is allowed on GLES

diff --git a/Engine/Source/Runtime/OpenGLDrv/Private/OpenGLES.cpp b/Engine/Source/Runtime/OpenGLDrv/Private/OpenGLES.cpp
--- a/Engine/Source/Runtime/OpenGLDrv/Private/OpenGLES.cpp
+++ b/Engine/Source/Runtime/OpenGLDrv/Private/OpenGLES.cpp
@@ -78,7 +78,7 @@ bool FOpenGLES::bRequiresPreciseQualifierWorkaround = false;
 bool FOpenGLES::bRequiresARMShaderFramebufferFetchDepthStencilUndef = false;
 
 /** Framebuffer fetch can be used to do programmable blending without running into driver issues */
-bool FOpenGLES::bSupportsShaderFramebufferFetchProgrammableBlending = true;
+bool FOpenGLES::bSupportsShaderFramebufferFetchProgrammableBlending = false;
 
 /** GL_EXT_buffer_storage */
 bool FOpenGLES::bSupportsBufferStorage = false;
@@ -191,6 +191,7 @@ void FOpenGLES::ProcessExtensions(const FString& ExtensionsString)
 		bSupportsShaderMRTFramebufferFetch = ExtensionsString.Contains(TEXT("GL_EXT_shader_framebuffer_fetch")) || ExtensionsString.Contains(TEXT("GL_NV_shader_framebuffer_fetch"));
 		bSupportsPixelLocalStorage = ExtensionsString.Contains(TEXT("GL_EXT_shader_pixel_local_storage"));
 		bSupportsShaderDepthStencilFetch = ExtensionsString.Contains(TEXT("GL_ARM_shader_framebuffer_fetch_depth_stencil"));
+		bSupportsShaderFramebufferFetchProgrammableBlending = bSupportsShaderFramebufferFetch;
 	}
 
 	// Report shader precision
