Fixed: Add support for BC6H texture format on Windows GL

diff --git a/Engine/Source/Runtime/OpenGLDrv/Private/OpenGLDevice.cpp b/Engine/Source/Runtime/OpenGLDrv/Private/OpenGLDevice.cpp
--- a/Engine/Source/Runtime/OpenGLDrv/Private/OpenGLDevice.cpp
+++ b/Engine/Source/Runtime/OpenGLDrv/Private/OpenGLDevice.cpp
@@ -475,6 +475,7 @@ static EOpenGLFormatCapabilities GetOpenGLFormatCapabilities(const FOpenGLTextur
 	case GL_RG16:
 	case GL_RG16_SNORM:
 	case GL_COMPRESSED_RGBA_BPTC_UNORM:
+	case GL_COMPRESSED_RGB_BPTC_UNSIGNED_FLOAT:
 #endif
 
 	case GL_COMPRESSED_RGBA_S3TC_DXT1_EXT:
@@ -1089,6 +1090,7 @@ static void InitRHICapabilitiesForGL()
 		SetupTextureFormat( PF_BC4,				FOpenGLTextureFormat( GL_COMPRESSED_RED_RGTC1,	GL_COMPRESSED_RED_RGTC1,GL_RED,			GL_UNSIGNED_BYTE,					true,	false));
 		
 		SetupTextureFormat( PF_BC7,				FOpenGLTextureFormat( GL_COMPRESSED_RGBA_BPTC_UNORM, GL_COMPRESSED_SRGB_ALPHA_BPTC_UNORM, GL_RGBA, GL_UNSIGNED_BYTE, true, false));
+		SetupTextureFormat( PF_BC6H,			FOpenGLTextureFormat( GL_COMPRESSED_RGB_BPTC_UNSIGNED_FLOAT, GL_COMPRESSED_RGB_BPTC_UNSIGNED_FLOAT, GL_RGB, GL_HALF_FLOAT, true, false));
 
 		SetupTextureFormat( PF_G16R16,			FOpenGLTextureFormat(GL_RG16, GL_RG16, GL_RG, GL_UNSIGNED_SHORT, false, false));
 		SetupTextureFormat( PF_G16R16_SNORM,	FOpenGLTextureFormat(GL_RG16_SNORM, GL_RG16_SNORM, GL_RG, GL_SHORT, false, false));
