Mobile: Fixed initialization order of AllowFBF option which was causing wrong rendering in OpenGL

diff --git a/Engine/Source/Runtime/Apple/MetalRHI/Private/MetalRHI.cpp b/Engine/Source/Runtime/Apple/MetalRHI/Private/MetalRHI.cpp
--- a/Engine/Source/Runtime/Apple/MetalRHI/Private/MetalRHI.cpp
+++ b/Engine/Source/Runtime/Apple/MetalRHI/Private/MetalRHI.cpp
@@ -674,7 +674,7 @@ FMetalDynamicRHI::FMetalDynamicRHI(ERHIFeatureLevel::Type RequestedFeatureLevel)
 	SetEmitDrawEvents(GetEmitDrawEvents() | ENABLE_METAL_GPUEVENTS);
 #endif
 
-	GSupportsShaderFramebufferFetch = !PLATFORM_MAC && GMaxRHIShaderPlatform != SP_METAL_SM5_IOS && GMaxRHIShaderPlatform != SP_METAL_SM5_TVOS && MobileAllowFramebufferFetch(GMaxRHIShaderPlatform);
+	GSupportsShaderFramebufferFetch = !PLATFORM_MAC && GMaxRHIShaderPlatform != SP_METAL_SM5_IOS && GMaxRHIShaderPlatform != SP_METAL_SM5_TVOS;
 	GSupportsShaderMRTFramebufferFetch = GSupportsShaderFramebufferFetch;
 	GHardwareHiddenSurfaceRemoval = true;
 	GSupportsRenderTargetFormat_PF_G8 = false;
diff --git a/Engine/Source/Runtime/Core/Private/HAL/ConsoleManager.cpp b/Engine/Source/Runtime/Core/Private/HAL/ConsoleManager.cpp
--- a/Engine/Source/Runtime/Core/Private/HAL/ConsoleManager.cpp
+++ b/Engine/Source/Runtime/Core/Private/HAL/ConsoleManager.cpp
@@ -3792,17 +3792,10 @@ static TAutoConsoleVariable<int32> CVarMobileAllowDeferredShadingOpenGL(
 		 "1: Allow Deferred Shading on OpenGL (default)"),
 	ECVF_RenderThreadSafe | ECVF_ReadOnly);
 
-static TAutoConsoleVariable<int32> CVarMobileAllowFramebufferFetchOpenGL(
-	TEXT("r.Mobile.AllowFramebufferFetchOpenGL"),
-	0,
-	TEXT("0: Use multi-pass rendering without FBF and PLS extensions (only for a deferred shading) (default)\n"
-		"1: Allow use of framebuffer fetch and PLS extensions"),
-	ECVF_RenderThreadSafe | ECVF_ReadOnly);
-
 static TAutoConsoleVariable<int32> CVarMobileAllowFramebufferFetch(
 	TEXT("r.Mobile.AllowFramebufferFetch"),
 	0,
-	TEXT("0: Multi-pass rendering without using sub-passes in Vulkan or programmable blending on iOS. For OpenGL use r.Mobile.AllowFramebufferFetchOpenGL (only for a deferred shading) (default)\n"
+	TEXT("0: Multi-pass rendering without using sub-passes in Vulkan or programmable blending on iOS and GLES. (only for a deferred shading) (default)\n"
 		 "1: Allow use of sub-passes and programmable blending"),
 	ECVF_RenderThreadSafe | ECVF_ReadOnly);
 
diff --git a/Engine/Source/Runtime/OpenGLDrv/Private/OpenGLES.cpp b/Engine/Source/Runtime/OpenGLDrv/Private/OpenGLES.cpp
--- a/Engine/Source/Runtime/OpenGLDrv/Private/OpenGLES.cpp
+++ b/Engine/Source/Runtime/OpenGLDrv/Private/OpenGLES.cpp
@@ -185,14 +185,11 @@ void FOpenGLES::ProcessExtensions(const FString& ExtensionsString)
 	bSupportsDepthClamp = ExtensionsString.Contains(TEXT("GL_EXT_depth_clamp"));
 	bSupportsASTCDecodeMode = ExtensionsString.Contains(TEXT("GL_EXT_texture_compression_astc_decode_mode"));
 	
-	if (MobileAllowFramebufferFetch(GMaxRHIShaderPlatform))
-	{
-		bSupportsShaderFramebufferFetch = ExtensionsString.Contains(TEXT("GL_EXT_shader_framebuffer_fetch")) || ExtensionsString.Contains(TEXT("GL_NV_shader_framebuffer_fetch")) || ExtensionsString.Contains(TEXT("GL_ARM_shader_framebuffer_fetch ")); // has space at the end to exclude GL_ARM_shader_framebuffer_fetch_depth_stencil match
-		bSupportsShaderMRTFramebufferFetch = ExtensionsString.Contains(TEXT("GL_EXT_shader_framebuffer_fetch")) || ExtensionsString.Contains(TEXT("GL_NV_shader_framebuffer_fetch"));
-		bSupportsPixelLocalStorage = ExtensionsString.Contains(TEXT("GL_EXT_shader_pixel_local_storage"));
-		bSupportsShaderDepthStencilFetch = ExtensionsString.Contains(TEXT("GL_ARM_shader_framebuffer_fetch_depth_stencil"));
-		bSupportsShaderFramebufferFetchProgrammableBlending = bSupportsShaderFramebufferFetch;
-	}
+	bSupportsShaderFramebufferFetch = ExtensionsString.Contains(TEXT("GL_EXT_shader_framebuffer_fetch")) || ExtensionsString.Contains(TEXT("GL_NV_shader_framebuffer_fetch")) || ExtensionsString.Contains(TEXT("GL_ARM_shader_framebuffer_fetch ")); // has space at the end to exclude GL_ARM_shader_framebuffer_fetch_depth_stencil match
+	bSupportsShaderMRTFramebufferFetch = ExtensionsString.Contains(TEXT("GL_EXT_shader_framebuffer_fetch")) || ExtensionsString.Contains(TEXT("GL_NV_shader_framebuffer_fetch"));
+	bSupportsPixelLocalStorage = ExtensionsString.Contains(TEXT("GL_EXT_shader_pixel_local_storage"));
+	bSupportsShaderDepthStencilFetch = ExtensionsString.Contains(TEXT("GL_ARM_shader_framebuffer_fetch_depth_stencil"));
+	bSupportsShaderFramebufferFetchProgrammableBlending = bSupportsShaderFramebufferFetch;
 
 	// Report shader precision
 	int Range[2];
diff --git a/Engine/Source/Runtime/OpenGLDrv/Private/OpenGLShaders.cpp b/Engine/Source/Runtime/OpenGLDrv/Private/OpenGLShaders.cpp
--- a/Engine/Source/Runtime/OpenGLDrv/Private/OpenGLShaders.cpp
+++ b/Engine/Source/Runtime/OpenGLDrv/Private/OpenGLShaders.cpp
@@ -775,7 +775,7 @@ void OPENGLDRV_API GetCurrentOpenGLShaderDeviceCapabilities(FOpenGLShaderDeviceC
 	if (FOpenGL::IsAndroidGLESCompatibilityModeEnabled())
 	{
 		Capabilities.TargetPlatform = EOpenGLShaderTargetPlatform::OGLSTP_Android;
-		Capabilities.bSupportsShaderFramebufferFetch = FOpenGL::SupportsShaderFramebufferFetch();
+		Capabilities.bSupportsShaderFramebufferFetch = GSupportsShaderFramebufferFetch;
 		Capabilities.bRequiresARMShaderFramebufferFetchDepthStencilUndef = false;
 		Capabilities.bRequiresReadOnlyBuffersWorkaround = false;
 		Capabilities.MaxVaryingVectors = FOpenGL::GetMaxVaryingVectors();
@@ -784,7 +784,7 @@ void OPENGLDRV_API GetCurrentOpenGLShaderDeviceCapabilities(FOpenGLShaderDeviceC
 
 #elif PLATFORM_ANDROID
 		Capabilities.TargetPlatform = EOpenGLShaderTargetPlatform::OGLSTP_Android;
-		Capabilities.bSupportsShaderFramebufferFetch = FOpenGL::SupportsShaderFramebufferFetch();
+		Capabilities.bSupportsShaderFramebufferFetch = GSupportsShaderFramebufferFetch;
 		Capabilities.bRequiresReadOnlyBuffersWorkaround = FOpenGL::RequiresReadOnlyBuffersWorkaround();
 		Capabilities.bRequiresARMShaderFramebufferFetchDepthStencilUndef = FOpenGL::RequiresARMShaderFramebufferFetchDepthStencilUndef();
 		Capabilities.MaxVaryingVectors = FOpenGL::GetMaxVaryingVectors();
@@ -1010,7 +1010,7 @@ FOpenGLShader::FOpenGLShader(TArrayView<const uint8> Code, const FSHAHash& Hash,
 	uint32 CodeCRC = FCrc::MemCrc32(GlslCodeOriginal.GetData(), GlslCodeOriginal.Num());
 	ShaderCodeKey = FOpenGLCompiledShaderKey(TypeEnum, GlslCodeOriginal.Num(), CodeCRC);
 
-	if (TypeEnum == GL_FRAGMENT_SHADER && FOpenGL::SupportsShaderFramebufferFetch())
+	if (TypeEnum == GL_FRAGMENT_SHADER && GSupportsShaderFramebufferFetch)
 	{
 		// _Globals_gl_LastFragColor should only exist when 'FramebufferFetchGLES2()' is being used, not for MRT/deferred
 		if (FCStringAnsi::Strstr(GlslCodeOriginal.GetData(), "_Globals_gl_LastFragColor") != nullptr)
diff --git a/Engine/Source/Runtime/RenderCore/Private/ReadOnlyCVARCache.cpp b/Engine/Source/Runtime/RenderCore/Private/ReadOnlyCVARCache.cpp
--- a/Engine/Source/Runtime/RenderCore/Private/ReadOnlyCVARCache.cpp
+++ b/Engine/Source/Runtime/RenderCore/Private/ReadOnlyCVARCache.cpp
@@ -30,6 +30,7 @@ bool FReadOnlyCVARCache::bMobileDeferredShadingValue = false;
 bool FReadOnlyCVARCache::bMobileEnableMovableSpotlightsShadowValue = false;
 int32 FReadOnlyCVARCache::MobileForwardDecalLightingValue = 1;
 bool FReadOnlyCVARCache::bMobileAllowDitheredLODTransitionValue = false;
+bool FReadOnlyCVARCache::bMobileAllowFramebufferFetchValue = false;
 
 int32 FReadOnlyCVARCache::MobileEarlyZPassIniValue(EShaderPlatform Platform)
 {
@@ -79,6 +80,21 @@ bool FReadOnlyCVARCache::MobileAllowDitheredLODTransitionValueIniValue(EShaderPl
 	return CVar.Get(Platform);
 }
 
+bool FReadOnlyCVARCache::MobileAllowFramebufferFetchIniValue(EShaderPlatform Platform)
+{
+	if (IsVulkanMobilePlatform(Platform) ||
+		IsMetalMobilePlatform(Platform) ||
+		IsAndroidOpenGLESPlatform(Platform))
+	{
+		static FShaderPlatformCachedIniValue<bool> CVar(TEXT("r.Mobile.AllowFramebufferFetch"));
+		// Always allow in Forward, optional in Deferred
+		return !(bMobileDeferredShadingValue && bMobileHDR) || CVar.Get(Platform);
+	}
+
+	// Only GLES, Vulkan and Metal have FBF
+	return false;
+}
+
 void FReadOnlyCVARCache::Initialize()
 {
 	UE_LOG(LogInit, Log, TEXT("Initializing FReadOnlyCVARCache"));
@@ -119,6 +135,7 @@ void FReadOnlyCVARCache::Initialize()
 	bMobileForwardParticleLights = MobileForwardParticleLightsIniValue(GMaxRHIShaderPlatform);
 	MobileForwardDecalLightingValue = CVarMobileForwardDecalLightingValue->GetValueOnAnyThread();
 	bMobileAllowDitheredLODTransitionValue = MobileAllowDitheredLODTransitionValueIniValue(GMaxRHIShaderPlatform);
+	bMobileAllowFramebufferFetchValue = MobileAllowFramebufferFetchIniValue(GMaxRHIShaderPlatform);
 
 #ifdef PROJECT_CVAR_ALLOW_STATIC_LIGHTING
 	check(!!PROJECT_CVAR_ALLOW_STATIC_LIGHTING == bAllowStaticLighting);
diff --git a/Engine/Source/Runtime/RenderCore/Private/RenderUtils.cpp b/Engine/Source/Runtime/RenderCore/Private/RenderUtils.cpp
--- a/Engine/Source/Runtime/RenderCore/Private/RenderUtils.cpp
+++ b/Engine/Source/Runtime/RenderCore/Private/RenderUtils.cpp
@@ -465,19 +465,7 @@ RENDERCORE_API bool PlatformGPUSceneUsesUniformBufferView(const FStaticShaderPla
 
 RENDERCORE_API bool MobileAllowFramebufferFetch(const FStaticShaderPlatform Platform)
 {
-	if (Platform == SP_OPENGL_ES3_1_ANDROID)
-	{
-		const static IConsoleVariable* CVarAllowFramebufferFetchOpenGL = IConsoleManager::Get().FindConsoleVariable(TEXT("r.Mobile.AllowFramebufferFetchOpenGL"));
-		return !IsMobileDeferredShadingEnabled(Platform) || CVarAllowFramebufferFetchOpenGL->GetBool();
-	}
-	else if (IsVulkanMobilePlatform(Platform) || IsMetalMobilePlatform(Platform))
-	{
-		static FShaderPlatformCachedIniValue<bool> MobileAllowFramebufferFetchIniValue(TEXT("r.Mobile.AllowFramebufferFetch"));
-		return !IsMobileDeferredShadingEnabled(Platform) || MobileAllowFramebufferFetchIniValue.Get(Platform);
-	}
-	
-	// Only GLES, Vulkan and Metal have FBF
-	return false; 
+	return FReadOnlyCVARCache::MobileAllowFramebufferFetch(Platform);
 }
 
 RENDERCORE_API bool MobileRequiresSceneDepthAux(const FStaticShaderPlatform Platform)
@@ -1054,6 +1042,15 @@ RENDERCORE_API void RenderUtilsInit()
 	{
 		UE_LOG(LogRendererCore, Log, TEXT("Ray tracing is disabled. Reason: disabled through project setting (r.RayTracing=0)."));
 	}
+
+	if (!FReadOnlyCVARCache::MobileAllowFramebufferFetch(GMaxRHIShaderPlatform))
+	{
+		GSupportsShaderFramebufferFetch = false;
+		GSupportsShaderMRTFramebufferFetch = false;
+		GSupportsShaderDepthStencilFetch = false;
+		GSupportsPixelLocalStorage = false;
+		GSupportsShaderFramebufferFetchProgrammableBlending = false;
+	}
 }
 
 class FUnitCubeVertexBuffer : public FVertexBuffer
diff --git a/Engine/Source/Runtime/RenderCore/Public/ReadOnlyCVARCache.h b/Engine/Source/Runtime/RenderCore/Public/ReadOnlyCVARCache.h
--- a/Engine/Source/Runtime/RenderCore/Public/ReadOnlyCVARCache.h
+++ b/Engine/Source/Runtime/RenderCore/Public/ReadOnlyCVARCache.h
@@ -153,6 +153,15 @@ struct FReadOnlyCVARCache
 		#endif
 	}
 
+	static inline bool MobileAllowFramebufferFetch(EShaderPlatform Platform)
+	{
+		#if WITH_EDITOR
+			return MobileAllowFramebufferFetchIniValue(Platform);
+		#else
+			return bMobileAllowFramebufferFetchValue;
+		#endif
+	}
+
 private:
 	RENDERCORE_API static bool bInitialized;
 
@@ -175,6 +184,7 @@ private:
 	RENDERCORE_API static bool bMobileDeferredShadingValue;
 	RENDERCORE_API static bool bMobileEnableMovableSpotlightsShadowValue;
 	RENDERCORE_API static bool bMobileAllowDitheredLODTransitionValue;
+	RENDERCORE_API static bool bMobileAllowFramebufferFetchValue;
 
 private:
 	RENDERCORE_API static int32 MobileEarlyZPassIniValue(EShaderPlatform Platform);
@@ -183,5 +193,6 @@ private:
 	RENDERCORE_API static bool MobileDeferredShadingIniValue(EShaderPlatform Platform);
 	RENDERCORE_API static bool MobileEnableMovableSpotlightsShadowIniValue(EShaderPlatform Platform);
 	RENDERCORE_API static bool MobileAllowDitheredLODTransitionValueIniValue(EShaderPlatform Platform);
+	RENDERCORE_API static bool MobileAllowFramebufferFetchIniValue(EShaderPlatform Platform);
 };
 
