Remove legacy GPU profiler code paths from OpenGL RHI

diff --git a/Engine/Source/Runtime/OpenGLDrv/Private/OpenGLCommands.cpp b/Engine/Source/Runtime/OpenGLDrv/Private/OpenGLCommands.cpp
--- a/Engine/Source/Runtime/OpenGLDrv/Private/OpenGLCommands.cpp
+++ b/Engine/Source/Runtime/OpenGLDrv/Private/OpenGLCommands.cpp
@@ -2223,9 +2223,6 @@ void FOpenGLDynamicRHI::RHIDrawPrimitive(uint32 BaseVertexIndex,uint32 NumPrimit
 	GLsizei NumElements = 0;
 	FindPrimitiveType(PrimitiveType, NumPrimitives, DrawMode, NumElements);
 
-#if (RHI_NEW_GPU_PROFILER == 0)
-	RegisterGPUWork(NumPrimitives * NumInstances, VertexCount * NumInstances);
-#endif
 	if (NumInstances == 1)
 	{
 		SCOPE_CYCLE_COUNTER_DETAILED(STAT_OpenGLDrawPrimitiveDriverTime);
@@ -2255,9 +2252,6 @@ void FOpenGLDynamicRHI::RHIDrawPrimitiveIndirect(FRHIBuffer* ArgumentBufferRHI,
 	VERIFY_GL_SCOPE();
 
 	check(ArgumentBufferRHI);
-#if (RHI_NEW_GPU_PROFILER == 0)
-	RegisterGPUWork(0);
-#endif
 
 	// Zero-stride buffer emulation won't work here, need to use VAB with proper zero strides
 	SetupDraw(nullptr, 0, 1);
@@ -2297,10 +2291,6 @@ void FOpenGLDynamicRHI::RHIDrawIndexedPrimitive(FRHIBuffer* IndexBufferRHI, int3
 	GLenum IndexType = IndexBuffer->GetStride() == sizeof(uint32) ? GL_UNSIGNED_INT : GL_UNSIGNED_SHORT;
 	StartIndex *= IndexBuffer->GetStride() == sizeof(uint32) ? sizeof(uint32) : sizeof(uint16);
 
-#if (RHI_NEW_GPU_PROFILER == 0)
-	RegisterGPUWork(NumPrimitives * NumInstances, NumElements * NumInstances);
-#endif
-
 	if (NumInstances > 1)
 	{
 		SCOPE_CYCLE_COUNTER_DETAILED(STAT_OpenGLDrawPrimitiveDriverTime);
@@ -2335,9 +2325,6 @@ void FOpenGLDynamicRHI::RHIDrawIndexedPrimitiveIndirect(FRHIBuffer* IndexBufferR
 	VERIFY_GL_SCOPE();
 
 	FOpenGLBuffer* IndexBuffer = ResourceCast(IndexBufferRHI);
-#if (RHI_NEW_GPU_PROFILER == 0)
-	RegisterGPUWork(1);
-#endif
 
 	check(ArgumentBufferRHI);
 
@@ -2499,9 +2486,6 @@ void FOpenGLDynamicRHI::RHIClearMRT(const bool* bClearColorArray,int32 NumClearC
 	bool bPrevScissorEnabled = PendingState.bScissorEnabled;
 
 	bool bScissorChanged = false;
-#if (RHI_NEW_GPU_PROFILER == 0)
-	RegisterGPUWork(0);
-#endif
 
 	BindPendingFramebuffer();
 
@@ -2682,10 +2666,6 @@ void FOpenGLDynamicRHI::RHIDispatchComputeShader(uint32 ThreadGroupCountX, uint3
 
 	VERIFY_GL_SCOPE();
 
-#if (RHI_NEW_GPU_PROFILER == 0)
-	RegisterGPUDispatch(FIntVector(ThreadGroupCountX, ThreadGroupCountY, ThreadGroupCountZ));
-#endif
-	
 	SetupDispatch();
 
 	FOpenGL::MemoryBarrier(GL_ALL_BARRIER_BITS);
@@ -2699,10 +2679,6 @@ void FOpenGLDynamicRHI::RHIDispatchIndirectComputeShader(FRHIBuffer* ArgumentBuf
 
 	VERIFY_GL_SCOPE();
 
-#if (RHI_NEW_GPU_PROFILER == 0)
-	RegisterGPUDispatch(FIntVector(1, 1, 1));
-#endif
-
 	SetupDispatch();
 	
 	FOpenGLBuffer* ArgumentBuffer = ResourceCast(ArgumentBufferRHI);
@@ -2756,9 +2732,7 @@ IRHIComputeContext* FOpenGLDynamicRHI::RHIGetCommandContext(ERHIPipeline Pipelin
 
 void FOpenGLDynamicRHI::RHIFinalizeContext(FRHIFinalizeContextArgs&& Args, TRHIPipelineArray<IRHIPlatformCommandList*>& Output)
 {
-#if RHI_NEW_GPU_PROFILER
 	FlushProfilerStats();
-#endif
 
 	for(IRHIComputeContext* Context : Args.Contexts)
 	{
diff --git a/Engine/Source/Runtime/OpenGLDrv/Private/OpenGLDevice.cpp b/Engine/Source/Runtime/OpenGLDrv/Private/OpenGLDevice.cpp
--- a/Engine/Source/Runtime/OpenGLDrv/Private/OpenGLDevice.cpp
+++ b/Engine/Source/Runtime/OpenGLDrv/Private/OpenGLDevice.cpp
@@ -92,12 +92,7 @@ extern void BeginFrame_VertexBufferCleanup();
 
 void FOpenGLDynamicRHI::RHIEndFrame(const FRHIEndFrameArgs& Args)
 {
-#if RHI_NEW_GPU_PROFILER
 	Profiler.EndFrame(Args);
-#else
-	GPUProfilingData->EndFrame();
-#endif
-
 	KickHint.Reset();
 
 #if PLATFORM_ANDROID //adding #if since not sure if this is required for any other platform.
@@ -106,10 +101,6 @@ void FOpenGLDynamicRHI::RHIEndFrame(const FRHIEndFrameArgs& Args)
 
 	BeginFrame_UniformBufferPoolCleanup();
 	BeginFrame_VertexBufferCleanup();
-
-#if (RHI_NEW_GPU_PROFILER == 0)
-	GPUProfilingData->BeginFrame();
-#endif
 }
 
 #if PLATFORM_ANDROID
@@ -1252,10 +1243,6 @@ FOpenGLDynamicRHI::FOpenGLDynamicRHI()
 
 	PrivateOpenGLDevicePtr = this;
 	GlobalUniformBuffers.AddZeroed(FUniformBufferStaticSlotRegistry::Get().GetSlotCount());
-
-#if RHI_NEW_GPU_PROFILER == 0
-	GPUProfilingData.Emplace();
-#endif
 }
 
 extern void DestroyShadersAndPrograms();
@@ -1348,9 +1335,7 @@ void FOpenGLDynamicRHI::Init()
 	// Flush here since we might be switching to a different context/thread for rendering
 	FOpenGL::Flush();
 
-#if RHI_NEW_GPU_PROFILER
 	Profiler.Setup();
-#endif
 
 	FHardwareInfo::RegisterHardwareInfo( NAME_RHI, TEXT( "OpenGL" ) );
 
@@ -1362,9 +1347,7 @@ void FOpenGLDynamicRHI::Shutdown()
 {
 	check(IsInGameThread() && IsInRenderingThread()); // require that the render thread has been shut down
 
-#if RHI_NEW_GPU_PROFILER
 	Profiler.Shutdown();
-#endif
 	Cleanup();
 
 	FOpenGLRenderQuery::Cleanup();
@@ -1384,10 +1367,6 @@ void FOpenGLDynamicRHI::Cleanup()
 		// Reset the RHI initialized flag.
 		GIsRHIInitialized = false;
 
-#if (RHI_NEW_GPU_PROFILER == 0)
-		GPUProfilingData->Cleanup();
-#endif
-
 		// Ask all initialized FRenderResources to release their RHI resources.
 		FRenderResource::ReleaseRHIForAllResources();
 	}
@@ -1454,8 +1433,6 @@ bool FOpenGLDynamicRHIModule::IsSupported()
 	return true;
 }
 
-#if RHI_NEW_GPU_PROFILER
-
 FOpenGLRenderQuery* FOpenGLProfiler::GetQuery()
 {
 	if (QueryPool.IsEmpty())
@@ -1661,5 +1638,3 @@ uint64 FOpenGLProfiler::ResolveQuery(uint64 Value, uint64* Target, bool bDisjoin
 		return CPUDelta + Frame->Calibration.CPUTimestamp;
 	}
 }
-
-#endif // RHI_NEW_GPU_PROFILER
diff --git a/Engine/Source/Runtime/OpenGLDrv/Private/OpenGLDrv.cpp b/Engine/Source/Runtime/OpenGLDrv/Private/OpenGLDrv.cpp
--- a/Engine/Source/Runtime/OpenGLDrv/Private/OpenGLDrv.cpp
+++ b/Engine/Source/Runtime/OpenGLDrv/Private/OpenGLDrv.cpp
@@ -58,11 +58,7 @@ bool FOpenGLDynamicRHI::RHIIsValidTexture(GLuint InTexture) const
 
 void FOpenGLDynamicRHI::RHISetExternalGPUTime(uint64 InExternalGPUTime)
 {
-#if RHI_NEW_GPU_PROFILER
 	Profiler.ExternalGPUTime = InExternalGPUTime;
-#else
-	GPUProfilingData->ExternalGPUTime = InExternalGPUTime;
-#endif
 }
 
 #if PLATFORM_ANDROID
@@ -142,7 +138,6 @@ void FOpenGLDynamicRHI::RHIEGLTerminateContext()
 	#endif
 		}
 
-	#if RHI_NEW_GPU_PROFILER
 		FlushProfilerStats();
 
 		if (Profiler.bEnabled)
@@ -150,17 +145,10 @@ void FOpenGLDynamicRHI::RHIEGLTerminateContext()
 			auto& Event = Profiler.EmplaceEvent<UE::RHI::GPUProfiler::FEvent::FBeginBreadcrumb>(Breadcrumb);
 			Profiler.InsertQuery(&Event.GPUTimestampTOP);
 		}
-	#else
-		if (GPUProfilingData->IsProfilingGPU())
-		{
-			GPUProfilingData->PushEvent(GetNameStr(), FColor::White);
-		}
-	#endif
 	}
 
 	void FOpenGLDynamicRHI::RHIEndBreadcrumbGPU(FRHIBreadcrumbNode* Breadcrumb)
 	{
-	#if RHI_NEW_GPU_PROFILER
 		FlushProfilerStats();
 
 		if (Profiler.bEnabled)
@@ -168,12 +156,6 @@ void FOpenGLDynamicRHI::RHIEGLTerminateContext()
 			auto& Event = Profiler.EmplaceEvent<UE::RHI::GPUProfiler::FEvent::FEndBreadcrumb>(Breadcrumb);
 			Profiler.InsertQuery(&Event.GPUTimestampBOP);
 		}
-	#else
-		if (GPUProfilingData->IsProfilingGPU())
-		{
-			GPUProfilingData->PopEvent();
-		}
-	#endif
 
 		if (ShouldEmitBreadcrumbs())
 		{
@@ -240,317 +222,6 @@ bool FOpenGLDynamicRHI::RHIMatchPrecachePSOInitializers(const FGraphicsPipelineS
 	return true;
 }
 
-#if (RHI_NEW_GPU_PROFILER == 0)
-
-void FOpenGLGPUProfiler::BeginFrame()
-{
-	if (NestedFrameCount++>0)
-	{
-		// guard against nested Begin/EndFrame calls.
-		return;
-	}
-
-	CurrentEventNode = NULL;
-	check(!bTrackingEvents);
-	check(!CurrentEventNodeFrame); // this should have already been cleaned up and the end of the previous frame
-
-	// latch the bools from the game thread into our private copy
-	bLatchedGProfilingGPU = GTriggerGPUProfile;
-	bLatchedGProfilingGPUHitches = GTriggerGPUHitchProfile;
-	if (bLatchedGProfilingGPUHitches)
-	{
-		bLatchedGProfilingGPU = false; // we do NOT permit an ordinary GPU profile during hitch profiles
-	}
-	
-	// if we are starting a hitch profile or this frame is a gpu profile, then save off the state of the draw events
-	if (bLatchedGProfilingGPU || (!bPreviousLatchedGProfilingGPUHitches && bLatchedGProfilingGPUHitches))
-	{
-		bOriginalGEmitDrawEvents = GetEmitDrawEvents();
-	}
-
-	if (bLatchedGProfilingGPU || bLatchedGProfilingGPUHitches)
-	{
-		if (bLatchedGProfilingGPUHitches && GPUHitchDebounce)
-		{
-			// if we are doing hitches and we had a recent hitch, wait to recover
-			// the reasoning is that collecting the hitch report may itself hitch the GPU
-			GPUHitchDebounce--; 
-		}
-		else
-		{
-			SetEmitDrawEvents(true);  // thwart an attempt to turn this off on the game side
-			bTrackingEvents = true;
-			CurrentEventNodeFrame = new FOpenGLEventNodeFrame();
-			CurrentEventNodeFrame->StartFrame();
-		}
-	}
-	else if (bPreviousLatchedGProfilingGPUHitches)
-	{
-		// hitch profiler is turning off, clear history and restore draw events
-		GPUHitchEventNodeFrames.Empty();
-		SetEmitDrawEvents(bOriginalGEmitDrawEvents);
-	}
-	bPreviousLatchedGProfilingGPUHitches = bLatchedGProfilingGPUHitches;
-
-	if (FrameTiming.IsSupported())
-	{
-		FrameTiming.StartTiming();
-	}
-	if (FOpenGLDisjointTimeStampQuery::IsSupported())
-	{
-		CurrentGPUFrameQueryIndex = (CurrentGPUFrameQueryIndex + 1) % MAX_GPUFRAMEQUERIES;
-		DisjointGPUFrameTimeQuery[CurrentGPUFrameQueryIndex].StartTracking();
-	}
-}
-
-void FOpenGLGPUProfiler::EndFrame()
-{
-	if (--NestedFrameCount != 0)
-	{
-		// ignore endframes calls from nested beginframe calls.
-		return;
-	}
-
-	if (FrameTiming.IsSupported())
-	{
-		FrameTiming.EndTiming();
-	}
-	if (FOpenGLDisjointTimeStampQuery::IsSupported())
-	{
-		DisjointGPUFrameTimeQuery[CurrentGPUFrameQueryIndex].EndTracking();
-	}
-
-	if (FrameTiming.IsSupported())
-	{
-		uint64 GPUTiming = FrameTiming.GetTiming();
-		uint64 GPUFreq = FrameTiming.GetTimingFrequency();
-		GRHIGPUFrameTimeHistory.PushFrameCycles(GPUFreq, GPUTiming);
-	}
-	else if (FOpenGLDisjointTimeStampQuery::IsSupported())
-	{
-		static uint64 GLastGPUTiming = 0;
-		static uint64 GLastGPUFreq = 0;
-
-		uint64 GPUTiming = 0;
-		uint64 GPUFreq = FOpenGLDisjointTimeStampQuery::GetTimingFrequency();
-		int OldestQueryIndex = (CurrentGPUFrameQueryIndex + 1) % MAX_GPUFRAMEQUERIES;
-		if (DisjointGPUFrameTimeQuery[OldestQueryIndex].IsResultValid() && DisjointGPUFrameTimeQuery[OldestQueryIndex].GetResult(&GPUTiming))
-		{
-			GRHIGPUFrameTimeHistory.PushFrameCycles(GPUFreq, GPUTiming);
-
-			GLastGPUTiming = GPUTiming;
-			GLastGPUFreq = GPUFreq;
-		}
-		else if (GLastGPUFreq > 0)
-		{
-			// Keep the timing of the last frame if the query was disjoint (e.g. GPU frequency changed and the result is undefined)
-			GRHIGPUFrameTimeHistory.PushFrameCycles(GLastGPUFreq, GLastGPUTiming);
-		}
-		else
-		{
-			// No valid timing data
-			GRHIGPUFrameTimeHistory.PushFrameCycles(1, 0);
-		}
-	}
-	else
-	{
-		GRHIGPUFrameTimeHistory.PushFrameCycles(1.0 / FPlatformTime::GetSecondsPerCycle64(), ExternalGPUTime);
-	}
-
-	// if we have a frame open, close it now.
-	if (CurrentEventNodeFrame)
-	{
-		CurrentEventNodeFrame->EndFrame();
-	}
-
-	check(!bTrackingEvents || bLatchedGProfilingGPU || bLatchedGProfilingGPUHitches);
-	check(!bTrackingEvents || CurrentEventNodeFrame);
-	if (bLatchedGProfilingGPU)
-	{
-		if (bTrackingEvents)
-		{
-			SetEmitDrawEvents(bOriginalGEmitDrawEvents);
-			UE_LOG(LogRHI, Warning, TEXT(""));
-			UE_LOG(LogRHI, Warning, TEXT(""));
-			CurrentEventNodeFrame->DumpEventTree();
-
-			// OPENGL_PERFORMANCE_DATA_INVALID is a compile time constant
-			bool DebugEnabled = false;
-#ifdef GL_ARB_debug_output
-			DebugEnabled = GL_TRUE == glIsEnabled( GL_DEBUG_OUTPUT );
-#endif
-#if !OPENGL_PERFORMANCE_DATA_INVALID
-			if( DebugEnabled )
-#endif
-			{
-				UE_LOG(LogRHI, Warning, TEXT(""));
-				UE_LOG(LogRHI, Warning, TEXT(""));
-				UE_LOG(LogRHI, Warning, TEXT("*********************************************************************************************"));
-				UE_LOG(LogRHI, Warning, TEXT("OpenGL perfomance data is potentially invalid because of the following build/runtime options:"));
-				
-#define LOG_GL_DEBUG_FLAG(a) UE_LOG(LogRHI, Warning, TEXT("   built with %s = %d"), TEXT(#a), a);
-				LOG_GL_DEBUG_FLAG(ENABLE_VERIFY_GL);
-				LOG_GL_DEBUG_FLAG(ENABLE_UNIFORM_BUFFER_LAYOUT_VERIFICATION);
-				LOG_GL_DEBUG_FLAG(ENABLE_UNIFORM_BUFFER_LAYOUT_DUMP);
-				LOG_GL_DEBUG_FLAG(DEBUG_GL_SHADERS);
-				LOG_GL_DEBUG_FLAG(ENABLE_OPENGL_DEBUG_GROUPS);
-				LOG_GL_DEBUG_FLAG(OPENGL_PERFORMANCE_DATA_INVALID);
-#undef LOG_GL_DEBUG_FLAG
-
-				UE_LOG(LogRHI, Warning, TEXT("*********************************************************************************************"));
-				UE_LOG(LogRHI, Warning, TEXT(""));
-				UE_LOG(LogRHI, Warning, TEXT(""));
-			}
-
-
-			GTriggerGPUProfile = false;
-			bLatchedGProfilingGPU = false;
-
-			if (RHIConfig::ShouldSaveScreenshotAfterProfilingGPU()
-				&& GEngine->GameViewport)
-			{
-				GEngine->GameViewport->Exec( NULL, TEXT("SCREENSHOT"), *GLog);
-			}
-		}
-	}
-	else if (bLatchedGProfilingGPUHitches)
-	{
-		//@todo this really detects any hitch, even one on the game thread.
-		// it would be nice to restrict the test to stalls on D3D, but for now...
-		// this needs to be out here because bTrackingEvents is false during the hitch debounce
-		static double LastTime = -1.0;
-		double Now = FPlatformTime::Seconds();
-		if (bTrackingEvents)
-		{
-			/** How long, in seconds a frame much be to be considered a hitch **/
-			static const float HitchThreshold = .1f; //100ms
-			float ThisTime = Now - LastTime;
-			bool bHitched = (ThisTime > HitchThreshold) && LastTime > 0.0 && CurrentEventNodeFrame;
-			if (bHitched)
-			{
-				UE_LOG(LogRHI, Warning, TEXT("*******************************************************************************"));
-				UE_LOG(LogRHI, Warning, TEXT("********** Hitch detected on CPU, frametime = %6.1fms"),ThisTime * 1000.0f);
-				UE_LOG(LogRHI, Warning, TEXT("*******************************************************************************"));
-
-				for (int32 Frame = 0; Frame < GPUHitchEventNodeFrames.Num(); Frame++)
-				{
-					UE_LOG(LogRHI, Warning, TEXT(""));
-					UE_LOG(LogRHI, Warning, TEXT(""));
-					UE_LOG(LogRHI, Warning, TEXT("********** GPU Frame: Current - %d"),GPUHitchEventNodeFrames.Num() - Frame);
-					GPUHitchEventNodeFrames[Frame].DumpEventTree();
-				}
-				UE_LOG(LogRHI, Warning, TEXT(""));
-				UE_LOG(LogRHI, Warning, TEXT(""));
-				UE_LOG(LogRHI, Warning, TEXT("********** GPU Frame: Current"));
-				CurrentEventNodeFrame->DumpEventTree();
-
-				UE_LOG(LogRHI, Warning, TEXT("*******************************************************************************"));
-				UE_LOG(LogRHI, Warning, TEXT("********** End Hitch GPU Profile"));
-				UE_LOG(LogRHI, Warning, TEXT("*******************************************************************************"));
-				if (GEngine->GameViewport)
-				{
-					GEngine->GameViewport->Exec( NULL, TEXT("SCREENSHOT"), *GLog);
-				}
-
-				GPUHitchDebounce = 5; // don't trigger this again for a while
-				GPUHitchEventNodeFrames.Empty(); // clear history
-			}
-			else if (CurrentEventNodeFrame) // this will be null for discarded frames while recovering from a recent hitch
-			{
-				/** How many old frames to buffer for hitch reports **/
-				static const int32 HitchHistorySize = 4;
-
-				if (GPUHitchEventNodeFrames.Num() >= HitchHistorySize)
-				{
-					GPUHitchEventNodeFrames.RemoveAt(0);
-				}
-				GPUHitchEventNodeFrames.Add((FOpenGLEventNodeFrame*)CurrentEventNodeFrame);
-				CurrentEventNodeFrame = NULL;  // prevent deletion of this below; ke kept it in the history
-			}
-		}
-		LastTime = Now;
-	}
-	bTrackingEvents = false;
-	delete CurrentEventNodeFrame;
-	CurrentEventNodeFrame = NULL;
-}
-
-void FOpenGLGPUProfiler::Cleanup()
-{
-	FrameTiming.ReleaseResources();
-	GPUHitchEventNodeFrames.Empty();
-	NestedFrameCount = 0;
-
-	for (FOpenGLDisjointTimeStampQuery& DisjointQuery : DisjointGPUFrameTimeQuery)
-	{
-		DisjointQuery.Cleanup();
-	}
-}
-
-/** Start this frame of per tracking */
-void FOpenGLEventNodeFrame::StartFrame()
-{
-	EventTree.Reset();
-	DisjointQuery.StartTracking();
-	RootEventTiming.StartTiming();
-}
-
-/** End this frame of per tracking, but do not block yet */
-void FOpenGLEventNodeFrame::EndFrame()
-{
-	RootEventTiming.EndTiming();
-	DisjointQuery.EndTracking();
-}
-
-
-float FOpenGLEventNodeFrame::GetRootTimingResults()
-{
-	double RootResult = 0.0f;
-	if (RootEventTiming.IsSupported())
-	{
-		const uint64 GPUTiming = RootEventTiming.GetTiming(true);
-		const uint64 GPUFreq = RootEventTiming.GetTimingFrequency();
-
-		RootResult = double(GPUTiming) / double(GPUFreq);
-	}
-
-	return (float)RootResult;
-}
-
-void FOpenGLEventNodeFrame::LogDisjointQuery()
-{
-	if (DisjointQuery.IsSupported())
-	{
-		UE_LOG(LogRHI, Warning, TEXT("%s"),
-			DisjointQuery.IsResultValid() ?
-			TEXT("Profiled range was continuous.") :
-			TEXT("Profiled range was disjoint! GPU switched to doing something else while profiling.")
-			);
-	}
-	else
-	{
-		UE_LOG(LogRHI, Warning, TEXT("Profiled range \"disjoinness\" could not be determined due to lack of disjoint timer query functionality on this platform."));
-	}
-}
-
-float FOpenGLEventNode::GetTiming()
-{
-	float Result = 0;
-
-	if (Timing.IsSupported())
-	{
-		// Get the timing result and block the CPU until it is ready
-		const uint64 GPUTiming = Timing.GetTiming(true);
-		const uint64 GPUFreq = Timing.GetTimingFrequency();
-
-		Result = double(GPUTiming) / double(GPUFreq);
-	}
-
-	return Result;
-}
-
-#endif // (RHI_NEW_GPU_PROFILER == 0)
-
 void FOpenGLDynamicRHI::InitializeStateResources()
 {
 	VERIFY_GL_SCOPE();
diff --git a/Engine/Source/Runtime/OpenGLDrv/Private/OpenGLDrv.h b/Engine/Source/Runtime/OpenGLDrv/Private/OpenGLDrv.h
--- a/Engine/Source/Runtime/OpenGLDrv/Private/OpenGLDrv.h
+++ b/Engine/Source/Runtime/OpenGLDrv/Private/OpenGLDrv.h
@@ -34,219 +34,6 @@
 
 template<class T> struct TOpenGLResourceTraits;
 
-#if (RHI_NEW_GPU_PROFILER == 0)
-
-// This class has multiple inheritance but really FGPUTiming is a static class
-class FOpenGLBufferedGPUTiming : public FGPUTiming
-{
-public:
-
-	/**
-	 * Constructor.
-	 *
-	 * @param InOpenGLRHI			RHI interface
-	 * @param InBufferSize		Number of buffered measurements
-	 */
-	FOpenGLBufferedGPUTiming(int32 BufferSize);
-
-	void StartTiming();
-
-	/**
-	 * End a GPU timing measurement.
-	 * The timing for this particular measurement will be resolved at a later time by the GPU.
-	 */
-	void EndTiming();
-
-	/**
-	 * Retrieves the most recently resolved timing measurement.
-	 * The unit is the same as for FPlatformTime::Cycles(). Returns 0 if there are no resolved measurements.
-	 *
-	 * @return	Value of the most recently resolved timing, or 0 if no measurements have been resolved by the GPU yet.
-	 */
-	uint64	GetTiming(bool bGetCurrentResultsAndBlock = false);
-
-	void InitResources();
-	void ReleaseResources();
-
-
-private:
-
-	/**
-	 * Initializes the static variables, if necessary.
-	 */
-	static void PlatformStaticInitialize(void* UserData);
-
-	/** Number of timestamps created in 'StartTimestamps' and 'EndTimestamps'. */
-	const int32							BufferSize;
-	/** Current timing being measured on the CPU. */
-	int32								CurrentTimestamp = -1;
-	/** Number of measurements in the buffers (0 - BufferSize). */
-	int32								NumIssuedTimestamps = 0;
-	/** Timestamps for all StartTimings. */
-	TArray<FOpenGLRenderQuery *>		StartTimestamps;
-	/** Timestamps for all EndTimings. */
-	TArray<FOpenGLRenderQuery *>		EndTimestamps;
-	/** Whether we are currently timing the GPU: between StartTiming() and EndTiming(). */
-	bool								bIsTiming = false;
-};
-
-/**
-  * Used to track whether a period was disjoint on the GPU, which means GPU timings are invalid.
-  * OpenGL lacks this concept at present, so the class is just a placeholder
-  * Timings are all assumed to be non-disjoint
-  */
-class FOpenGLDisjointTimeStampQuery
-{
-public:
-	FOpenGLDisjointTimeStampQuery()
-		: DisjointQuery(new FOpenGLRenderQuery { FOpenGLRenderQuery::EType::Disjoint })
-	{}
-
-	void StartTracking();
-	void EndTracking();
-	bool IsResultValid();
-
-	bool GetResult(uint64* OutResult);
-
-	static uint64 GetTimingFrequency()
-	{
-		return 1000000000ull;
-	}
-
-	static bool IsSupported()
-	{
-#if UE_BUILD_SHIPPING
-		return false;
-#else
-		return FOpenGL::SupportsDisjointTimeQueries();
-#endif
-	}
-
-	void Cleanup()
-	{
-		delete DisjointQuery;
-		DisjointQuery = nullptr;
-	}
-
-private:
-	bool	bIsResultValid = false;
-	FOpenGLRenderQuery* DisjointQuery;
-};
-
-/** A single perf event node, which tracks information about a appBeginDrawEvent/appEndDrawEvent range. */
-class FOpenGLEventNode : public FGPUProfilerEventNode
-{
-public:
-
-	FOpenGLEventNode(const TCHAR* InName, FGPUProfilerEventNode* InParent)
-		: FGPUProfilerEventNode(InName, InParent)
-		, Timing(1)
-	{
-		// Initialize Buffered timestamp queries 
-		Timing.InitResources();
-	}
-
-	virtual ~FOpenGLEventNode()
-	{
-		Timing.ReleaseResources();
-	}
-
-	/** 
-	 * Returns the time in ms that the GPU spent in this draw event.  
-	 * This blocks the CPU if necessary, so can cause hitching.
-	 */
-	float GetTiming() override;
-
-	virtual void StartTiming() override
-	{
-		Timing.StartTiming();
-	}
-
-	virtual void StopTiming() override
-	{
-		Timing.EndTiming();
-	}
-
-	FOpenGLBufferedGPUTiming Timing;
-};
-
-/** An entire frame of perf event nodes, including ancillary timers. */
-class FOpenGLEventNodeFrame : public FGPUProfilerEventNodeFrame
-{
-public:
-	FOpenGLEventNodeFrame()
-		: RootEventTiming(1)
-		, DisjointQuery()
-	{
-		RootEventTiming.InitResources();
-	}
-
-	~FOpenGLEventNodeFrame()
-	{
-		RootEventTiming.ReleaseResources();
-	}
-
-	/** Start this frame of per tracking */
-	void StartFrame() override;
-
-	/** End this frame of per tracking, but do not block yet */
-	void EndFrame() override;
-
-	/** Calculates root timing base frequency (if needed by this RHI) */
-	virtual float GetRootTimingResults() override;
-
-	virtual void LogDisjointQuery() override;
-
-	/** Timer tracking inclusive time spent in the root nodes. */
-	FOpenGLBufferedGPUTiming RootEventTiming;
-
-	/** Disjoint query tracking whether the times reported by DumpEventTree are reliable. */
-	FOpenGLDisjointTimeStampQuery DisjointQuery;
-};
-
-/** 
- * Encapsulates GPU profiling logic and data. 
- * There's only one global instance of this struct so it should only contain global data, nothing specific to a frame.
- */
-struct FOpenGLGPUProfiler : public FGPUProfiler
-{
-	/** Used to measure GPU time per frame. */
-	FOpenGLBufferedGPUTiming FrameTiming;
-
-	/** Measuring GPU frame time with a disjoint query. */
-	static const int MAX_GPUFRAMEQUERIES = 4;
-	FOpenGLDisjointTimeStampQuery DisjointGPUFrameTimeQuery[MAX_GPUFRAMEQUERIES];
-	int32 CurrentGPUFrameQueryIndex = 0;
-
-	// count the number of beginframe calls without matching endframe calls.
-	int32 NestedFrameCount = 0;
-
-	uint32 ExternalGPUTime = 0;
-
-	/** GPU hitch profile histories */
-	TIndirectArray<FOpenGLEventNodeFrame> GPUHitchEventNodeFrames;
-
-	FOpenGLGPUProfiler()
-		: FrameTiming(4)
-	{
-		FrameTiming.InitResources();
-		BeginFrame();
-	}
-
-	virtual FGPUProfilerEventNode* CreateEventNode(const TCHAR* InName, FGPUProfilerEventNode* InParent) override
-	{
-		FOpenGLEventNode* EventNode = new FOpenGLEventNode(InName, InParent);
-		return EventNode;
-	}
-
-	void Cleanup();
-
-	void BeginFrame();
-	void EndFrame();
-};
-
-#endif // (RHI_NEW_GPU_PROFILER == 0)
-
 PRAGMA_DISABLE_DEPRECATION_WARNINGS
 
 /** The interface which is implemented by the dynamically bound RHI. */
@@ -698,8 +485,6 @@ private:
 	/** Underlying platform-specific data */
 	struct FPlatformOpenGLDevice* PlatformDevice = nullptr;
 
-#if RHI_NEW_GPU_PROFILER
-
 	friend class FOpenGLRenderQuery;
 
 	FOpenGLProfiler Profiler;
@@ -714,22 +499,6 @@ private:
 		}
 	}
 
-#else
-
-	TOptional<FOpenGLGPUProfiler> GPUProfilingData;
-	friend FOpenGLGPUProfiler;
-
-	void RegisterGPUWork(uint32 NumPrimitives = 0, uint32 NumVertices = 0)
-	{
-		GPUProfilingData->RegisterGPUWork(NumPrimitives, NumVertices);
-	}
-	void RegisterGPUDispatch(FIntVector GroupCount)
-	{
-		GPUProfilingData->RegisterGPUDispatch(GroupCount);
-	}
-
-#endif
-
 	FCriticalSection CustomPresentSection;
 	TRefCountPtr<class FRHICustomPresent> CustomPresent;
 
diff --git a/Engine/Source/Runtime/OpenGLDrv/Private/OpenGLQuery.cpp b/Engine/Source/Runtime/OpenGLDrv/Private/OpenGLQuery.cpp
--- a/Engine/Source/Runtime/OpenGLDrv/Private/OpenGLQuery.cpp
+++ b/Engine/Source/Runtime/OpenGLDrv/Private/OpenGLQuery.cpp
@@ -151,12 +151,6 @@ void FOpenGLRenderQuery::Begin()
 			, Resource
 		);
 		break;
-
-#if RHI_NEW_GPU_PROFILER == 0
-	case EType::Disjoint:
-		FOpenGL::BeginQuery(UGL_TIME_ELAPSED, Resource);
-		break;
-#endif
 	};
 }
 
@@ -181,15 +175,9 @@ void FOpenGLRenderQuery::End(uint64* InTarget)
 		FOpenGL::QueryTimestampCounter(Resource);
 		break;
 
-#if RHI_NEW_GPU_PROFILER
 	case EType::Profiler:
 		FOpenGL::QueryTimestampCounter(Resource);
 		break;
-#else
-	case EType::Disjoint:
-		FOpenGL::EndQuery(UGL_TIME_ELAPSED);
-		break;
-#endif
 	}
 
 	BOPCounter++;
@@ -255,7 +243,6 @@ bool FOpenGLRenderQuery::CacheResult(bool bWait)
 		}
 		break;
 
-#if RHI_NEW_GPU_PROFILER
 	case EType::Profiler:
 		{
 			FOpenGLDynamicRHI& RHI = FOpenGLDynamicRHI::Get();
@@ -287,33 +274,6 @@ bool FOpenGLRenderQuery::CacheResult(bool bWait)
 			}
 		}
 		break;
-#else
-	case EType::Disjoint:
-		{
-			// TimerQueryDisjoint is a one-shot state in the driver, it is not pipelined.
-			// If it returns true, all timers we've submitted after this timer but haven't
-			// yet resolved should be discarded for having invalid data.
-			if (FOpenGL::TimerQueryDisjoint())
-			{
-				for (FOpenGLRenderQuery* Other = this; Other; Other = Other->Next)
-				{
-					if (Other->Type == EType::Disjoint)
-					{
-						Other->SetResult(InvalidDisjointMask);
-					}
-				}
-			}
-			else
-			{
-				GLuint64 Value;
-				FOpenGL::GetQueryObject(Resource, FOpenGL::QM_Result, &Value);
-
-				// Convert to microseconds (GL queries are in nanoseconds)
-				SetResult(Value / 1000);
-			}
-		}
-		break;
-#endif
 	}
 
 	return true;
@@ -541,267 +501,6 @@ FOpenGLEventQuery::~FOpenGLEventQuery()
 	FOpenGL::DeleteSync(Sync);
 }
 
-/*=============================================================================
- * class FOpenGLBufferedGPUTiming
- *=============================================================================*/
-
-#if (RHI_NEW_GPU_PROFILER == 0)
-
-/**
- * Constructor.
- *
- * @param InOpenGLRHI			RHI interface
- * @param InBufferSize		Number of buffered measurements
- */
-FOpenGLBufferedGPUTiming::FOpenGLBufferedGPUTiming(int32 InBufferSize)
-	: BufferSize(InBufferSize)
-{
-}
-
-/**
- * Initializes the static variables, if necessary.
- */
-void FOpenGLBufferedGPUTiming::PlatformStaticInitialize(void* UserData)
-{
-	// Are the static variables initialized?
-	if ( !GAreGlobalsInitialized )
-	{
-		GIsSupported = FOpenGL::SupportsTimestampQueries();
-		SetTimingFrequency(1000 * 1000 * 1000);
-		GAreGlobalsInitialized = true;
-	}
-}
-
-/**
- * Initializes all OpenGL resources and if necessary, the static variables.
- */
-
-static TArray<FOpenGLRenderQuery*> TimerQueryPool;
-
-static FOpenGLRenderQuery* GetTimeQuery()
-{
-	if (TimerQueryPool.Num())
-	{
-		return TimerQueryPool.Pop();
-	}
-	return new FOpenGLRenderQuery(FOpenGLRenderQuery::EType::Timestamp);
-}
-
-void FOpenGLBufferedGPUTiming::InitResources()
-{
-	StaticInitialize(nullptr, PlatformStaticInitialize);
-
-	CurrentTimestamp = 0;
-	NumIssuedTimestamps = 0;
-	bIsTiming = false;
-	GIsSupported = FOpenGL::SupportsTimestampQueries();
-
-	if ( GIsSupported )
-	{
-		StartTimestamps.Reserve(BufferSize);
-		EndTimestamps.Reserve(BufferSize);
-
-		for(int32 BufferIndex = 0; BufferIndex < BufferSize; ++BufferIndex)
-		{
-			StartTimestamps.Add(GetTimeQuery());
-			EndTimestamps.Add(GetTimeQuery());
-		}
-	}
-}
-
-/**
- * Releases all OpenGL resources.
- */
-void FOpenGLBufferedGPUTiming::ReleaseResources()
-{
-	VERIFY_GL_SCOPE();
-
-	for (FOpenGLRenderQuery* Query : StartTimestamps)
-	{
-		TimerQueryPool.Add(Query);
-	}
-
-	for (FOpenGLRenderQuery* Query : EndTimestamps)
-	{
-		TimerQueryPool.Add(Query);
-	}
-
-	StartTimestamps.Reset();
-	EndTimestamps.Reset();
-
-}
-
-/**
- * Start a GPU timing measurement.
- */
-void FOpenGLBufferedGPUTiming::StartTiming()
-{
-	VERIFY_GL_SCOPE();
-	// Issue a timestamp query for the 'start' time.
-	if ( GIsSupported && !bIsTiming )
-	{
-		int32 NewTimestampIndex = (CurrentTimestamp + 1) % BufferSize;
-		StartTimestamps[NewTimestampIndex]->End();
-
-		CurrentTimestamp = NewTimestampIndex;
-		bIsTiming = true;
-	}
-}
-
-/**
- * End a GPU timing measurement.
- * The timing for this particular measurement will be resolved at a later time by the GPU.
- */
-void FOpenGLBufferedGPUTiming::EndTiming()
-{
-	VERIFY_GL_SCOPE();
-	// Issue a timestamp query for the 'end' time.
-	if ( GIsSupported && bIsTiming )
-	{
-		checkSlow( CurrentTimestamp >= 0 && CurrentTimestamp < BufferSize );
-		EndTimestamps[CurrentTimestamp]->End();
-
-		NumIssuedTimestamps = FMath::Min<int32>(NumIssuedTimestamps + 1, BufferSize);
-		bIsTiming = false;
-	}
-}
-
-/**
- * Retrieves the most recently resolved timing measurement.
- * The unit is the same as for FPlatformTime::Cycles(). Returns 0 if there are no resolved measurements.
- *
- * @return	Value of the most recently resolved timing, or 0 if no measurements have been resolved by the GPU yet.
- */
-uint64 FOpenGLBufferedGPUTiming::GetTiming(bool bGetCurrentResultsAndBlock)
-{
-	VERIFY_GL_SCOPE();
-
-	if (GIsSupported)
-	{
-		checkSlow(CurrentTimestamp >= 0 && CurrentTimestamp < BufferSize);
-		int32 TimestampIndex = CurrentTimestamp;
-
-		if (!bGetCurrentResultsAndBlock)
-		{
-			// Quickly check the most recent measurements to see if any of them has been resolved.  Do not flush these queries.
-			for (int32 IssueIndex = 1; IssueIndex < NumIssuedTimestamps; ++IssueIndex)
-			{
-				FOpenGLRenderQuery* StartQuery = StartTimestamps[TimestampIndex];
-				FOpenGLRenderQuery* EndQuery = EndTimestamps[TimestampIndex];
-
-				if (StartQuery->CacheResult(false) && EndQuery->CacheResult(false))
-				{
-					uint64 StartTime = StartQuery->GetResult();
-					uint64 EndTime = EndQuery->GetResult();
-
-					if (EndTime > StartTime)
-					{
-						return EndTime - StartTime;
-					}
-				}
-
-				TimestampIndex = (TimestampIndex + BufferSize - 1) % BufferSize;
-			}
-		}
-
-		if (NumIssuedTimestamps > 0 || bGetCurrentResultsAndBlock)
-		{
-			// None of the (NumIssuedTimestamps - 1) measurements were ready yet,
-			// so check the oldest measurement more thoroughly.
-			// This really only happens if occlusion and frame sync event queries are disabled, otherwise those will block until the GPU catches up to 1 frame behind
-			const bool bBlocking = ( NumIssuedTimestamps == BufferSize ) || bGetCurrentResultsAndBlock;
-
-			FOpenGLRenderQuery* StartQuery = StartTimestamps[TimestampIndex];
-			FOpenGLRenderQuery* EndQuery = EndTimestamps[TimestampIndex];
-
-			bool bHasStart = false, bHasEnd = false;
-
-			{
-				FRenderThreadIdleScope IdleScope(ERenderThreadIdleTypes::WaitingForGPUQuery);
-				SCOPE_CYCLE_COUNTER(STAT_RenderQueryResultTime);
-
-				double StartTimeoutTime = FPlatformTime::Seconds();
-
-				// If we are blocking, retry until the GPU processes the time stamp command
-				while (true)
-				{
-					bHasStart = StartQuery->CacheResult(false);
-					bHasEnd = EndQuery->CacheResult(false);
-					
-					if (bBlocking && !(bHasStart && bHasEnd))
-					{
-						if ((FPlatformTime::Seconds() - StartTimeoutTime) > 0.5)
-						{
-							UE_LOG(LogRHI, Log, TEXT("Timed out while waiting for GPU to catch up. (500 ms)"));
-							return 0;
-						}
-					}
-					else
-					{
-						break;
-					}
-				}
-			}
-
-			if (bHasStart && bHasEnd)
-			{
-				uint64 StartTime = StartQuery->GetResult();
-				uint64 EndTime = EndQuery->GetResult();
-
-				if (EndTime > StartTime)
-				{
-					return EndTime - StartTime;
-				}
-			}
-		}
-	}
-	return 0;
-}
-
-void FOpenGLDisjointTimeStampQuery::StartTracking()
-{
-	VERIFY_GL_SCOPE();
-	if (IsSupported())
-	{
-		DisjointQuery->Begin();
-	}
-}
-
-void FOpenGLDisjointTimeStampQuery::EndTracking()
-{
-	VERIFY_GL_SCOPE();
-
-	if (IsSupported())
-	{
-		DisjointQuery->End();
-	}
-}
-
-bool FOpenGLDisjointTimeStampQuery::IsResultValid()
-{
-	checkSlow(IsSupported());
-	return bIsResultValid;
-}
-
-bool FOpenGLDisjointTimeStampQuery::GetResult(uint64* OutResult)
-{
-	VERIFY_GL_SCOPE();
-
-	if (IsSupported())
-	{
-		DisjointQuery->CacheResult(true);
-
-		uint64 Result = DisjointQuery->GetResult();
-		bIsResultValid = (Result & FOpenGLRenderQuery::InvalidDisjointMask) == 0;
-
-		*OutResult = Result & (~FOpenGLRenderQuery::InvalidDisjointMask);
-	}
-
-	return bIsResultValid;
-}
-
-#endif // (RHI_NEW_GPU_PROFILER == 0)
-
 TQueue<FOpenGLGPUFence::FGLSync, EQueueMode::SingleThreaded> FOpenGLGPUFence::ActiveSyncs;
 
 FOpenGLGPUFence::FOpenGLGPUFence(FName InName)
diff --git a/Engine/Source/Runtime/OpenGLDrv/Private/OpenGLResources.h b/Engine/Source/Runtime/OpenGLDrv/Private/OpenGLResources.h
--- a/Engine/Source/Runtime/OpenGLDrv/Private/OpenGLResources.h
+++ b/Engine/Source/Runtime/OpenGLDrv/Private/OpenGLResources.h
@@ -1115,8 +1115,6 @@ public:
 	uint8 const bMultisampleRenderbuffer : 1;
 };
 
-#if RHI_NEW_GPU_PROFILER
-
 struct FOpenGLProfiler
 {
 	struct FCalibration
@@ -1174,8 +1172,6 @@ struct FOpenGLProfiler
 	}
 };
 
-#endif
-
 class FOpenGLRenderQuery
 {
 public:
@@ -1183,19 +1179,11 @@ public:
 	{
 		Timestamp,
 		Occlusion,
-	#if RHI_NEW_GPU_PROFILER
 		Profiler,
-	#else
-		Disjoint,
-	#endif
 
 		Num
 	};
 
-#if RHI_NEW_GPU_PROFILER == 0
-	static constexpr uint64 InvalidDisjointMask = 0x8000000000000000;
-#endif
-
 private:
 	// Render queries that should be polled by the RHI thread.
 	struct FActiveQueries
diff --git a/Engine/Source/Runtime/OpenGLDrv/Private/OpenGLTexture.cpp b/Engine/Source/Runtime/OpenGLDrv/Private/OpenGLTexture.cpp
--- a/Engine/Source/Runtime/OpenGLDrv/Private/OpenGLTexture.cpp
+++ b/Engine/Source/Runtime/OpenGLDrv/Private/OpenGLTexture.cpp
@@ -1534,9 +1534,6 @@ FTextureRHIRef FOpenGLDynamicRHI::RHIAsyncCreateTexture2D(uint32 SizeX, uint32 S
 void FOpenGLDynamicRHI::RHIGenerateMips(FRHITexture* TextureRHI)
 {
 	VERIFY_GL_SCOPE();
-#if (RHI_NEW_GPU_PROFILER == 0)
-	RegisterGPUWork(0);
-#endif
 
 	FOpenGLTexture* Texture = ResourceCast(TextureRHI);
 	// Setup the texture on a disused unit
