Remove platform RHI implementations of RHIAsyncReallocateTexture2D
- Replaced with a high-level function that makes use of standard RHI commands (RHICreateTexture, RHICopyTexture and RHI transitions).
- Added some additional logic to OpenGL RHI's RHICopyTexture to handle the case of copying the shared mips to an evicted texture, which was previously handled via RHIAsyncReallocateTexture2D.

#Synced-CL 48809907

diff --git a/Engine/Source/Runtime/Apple/MetalRHI/Private/MetalDynamicRHI.h b/Engine/Source/Runtime/Apple/MetalRHI/Private/MetalDynamicRHI.h
--- a/Engine/Source/Runtime/Apple/MetalRHI/Private/MetalDynamicRHI.h
+++ b/Engine/Source/Runtime/Apple/MetalRHI/Private/MetalDynamicRHI.h
@@ -156,7 +156,6 @@ public:
 	[[nodiscard]] virtual FRHITextureInitializer RHICreateTextureInitializer(FRHICommandListBase& RHICmdList, const FRHITextureCreateDesc& CreateDesc) final override;
 	virtual FTextureRHIRef RHIAsyncCreateTexture2D(uint32 SizeX, uint32 SizeY, uint8 Format, uint32 NumMips, ETextureCreateFlags Flags, ERHIAccess InResourceState, void** InitialMipData, uint32 NumInitialMips, const TCHAR* DebugName, FGraphEventRef& OutCompletionEvent) final override;
 	virtual uint32 RHIComputeMemorySize(FRHITexture* TextureRHI) final override;
-	virtual FTextureRHIRef RHIAsyncReallocateTexture2D(FRHITexture* Texture2D, int32 NewMipCount, int32 NewSizeX, int32 NewSizeY, FThreadSafeCounter* RequestStatus) final override;
 
 	virtual FRHILockTextureResult RHILockTexture(FRHICommandListImmediate& RHICmdList, const FRHILockTextureArgs& Arguments) final override;
 	virtual void                  RHIUnlockTexture(FRHICommandListImmediate& RHICmdList, const FRHILockTextureArgs& Arguments) final override;
@@ -202,8 +201,6 @@ public:
 	
 	virtual void RHISubmitCommandLists(FRHISubmitCommandListsArgs&& Args) final override;
 
-	virtual FTextureRHIRef AsyncReallocateTexture2D_RenderThread(class FRHICommandListImmediate& RHICmdList, FRHITexture* Texture2D, int32 NewMipCount, int32 NewSizeX, int32 NewSizeY, FThreadSafeCounter* RequestStatus) final override;
-
 	virtual FUpdateTexture3DData RHIBeginUpdateTexture3D(FRHICommandListBase& RHICmdList, FRHITexture* Texture, uint32 MipIndex, const struct FUpdateTextureRegion3D& UpdateRegion) final override;
 	virtual void RHIEndUpdateTexture3D(FRHICommandListBase& RHICmdList, FUpdateTexture3DData& UpdateData) final override;
 
diff --git a/Engine/Source/Runtime/Apple/MetalRHI/Private/MetalTexture.cpp b/Engine/Source/Runtime/Apple/MetalRHI/Private/MetalTexture.cpp
--- a/Engine/Source/Runtime/Apple/MetalRHI/Private/MetalTexture.cpp
+++ b/Engine/Source/Runtime/Apple/MetalRHI/Private/MetalTexture.cpp
@@ -1554,82 +1554,6 @@ FTextureRHIRef FMetalDynamicRHI::RHIAsyncCreateTexture2D(uint32 SizeX, uint32 Si
 	return FTextureRHIRef();
 }
 
-FTextureRHIRef FMetalDynamicRHI::AsyncReallocateTexture2D_RenderThread(class FRHICommandListImmediate& RHICmdList, FRHITexture* Texture2D, int32 NewMipCount, int32 NewSizeX, int32 NewSizeY, FThreadSafeCounter* RequestStatus)
-{
-	return this->RHIAsyncReallocateTexture2D(Texture2D, NewMipCount, NewSizeX, NewSizeY, RequestStatus);
-}
-
-FTextureRHIRef FMetalDynamicRHI::RHIAsyncReallocateTexture2D(FRHITexture* OldTextureRHI, int32 NewMipCount, int32 NewSizeX, int32 NewSizeY, FThreadSafeCounter* RequestStatus)
-{
-    MTL_SCOPED_AUTORELEASE_POOL;
-
-	FRHICommandListImmediate& RHICmdList = FRHICommandListImmediate::Get();
-
-    FMetalSurface* OldTexture = ResourceCast(OldTextureRHI);
-
-    FRHITextureDesc Desc = OldTexture->GetDesc();
-    Desc.Extent = FIntPoint(NewSizeX, NewSizeY);
-    Desc.NumMips = NewMipCount;
-
-    FRHITextureCreateDesc CreateDesc(
-        Desc,
-        RHIGetDefaultResourceState(Desc.Flags, false),
-        TEXT("RHIAsyncReallocateTexture2D")
-    );
-	CreateDesc.SetOwnerName(OldTexture->GetOwnerName());
-
-	LLM_SCOPE_DYNAMIC_STAT_OBJECTPATH_FNAME(CreateDesc.OwnerName, ELLMTagSet::Assets);
-	LLM_SCOPE_DYNAMIC_STAT_OBJECTPATH_FNAME(CreateDesc.GetTraceClassName(), ELLMTagSet::AssetClasses);
-	UE_TRACE_METADATA_SCOPE_ASSET_FNAME(CreateDesc.DebugName, CreateDesc.GetTraceClassName(), CreateDesc.OwnerName);
-
-	FMetalSurface* NewTexture = new FMetalSurface(*Device, FMetalTextureCreateDesc(*Device, CreateDesc));
-	NewTexture->Initialize(RHICmdList);
-
-    // Copy shared mips
-    RHICmdList.EnqueueLambda([OldTexture, NewSizeX, NewSizeY, NewTexture, RequestStatus](FRHICommandListImmediate& ExecutingCmdList)
-    {
-		FMetalRHICommandContext& Context = FMetalRHICommandContext::Get(ExecutingCmdList);
-
-        // figure out what mips to schedule
-        const uint32 NumSharedMips = FMath::Min(OldTexture->GetNumMips(), NewTexture->GetNumMips());
-        const uint32 SourceMipOffset = OldTexture->GetNumMips() - NumSharedMips;
-        const uint32 DestMipOffset = NewTexture->GetNumMips() - NumSharedMips;
-
-        const uint32 BlockSizeX = GPixelFormats[OldTexture->GetFormat()].BlockSizeX;
-        const uint32 BlockSizeY = GPixelFormats[OldTexture->GetFormat()].BlockSizeY;
-
-        // only handling straight 2D textures here
-        uint32 SliceIndex = 0;
-        MTL::Origin Origin(0, 0, 0);
-
-        MTLTexturePtr Tex = OldTexture->Texture;
-
-        // DXT/BC formats on Mac actually do have mip-tails that are smaller than the block size, they end up being uncompressed.
-        bool const bPixelFormatASTC = IsPixelFormatASTCCompressed(OldTexture->GetFormat());
-
-        for (uint32 MipIndex = 0; MipIndex < NumSharedMips; ++MipIndex)
-        {
-            const uint32 UnalignedMipSizeX = FMath::Max<uint32>(1, NewSizeX >> (MipIndex + DestMipOffset));
-            const uint32 UnalignedMipSizeY = FMath::Max<uint32>(1, NewSizeY >> (MipIndex + DestMipOffset));
-            const uint32 MipSizeX = FMath::Max<uint32>(1, NewSizeX >> (MipIndex + DestMipOffset));
-            const uint32 MipSizeY = FMath::Max<uint32>(1, NewSizeY >> (MipIndex + DestMipOffset));
-
-			Context.CopyFromTextureToTexture(
-				OldTexture->Texture.get(), SliceIndex, MipIndex + SourceMipOffset, Origin, MTL::Size(MipSizeX, MipSizeY, 1),
-				NewTexture->Texture.get(), SliceIndex, MipIndex + DestMipOffset  , Origin
-			);
-        }
-
-        // Like D3D mark this as complete immediately.
-        RequestStatus->Decrement();
-
-        FMetalSurface* Source = GetMetalSurfaceFromRHITexture(OldTexture);
-        Source->MakeAliasable();
-    });
-    
-    return NewTexture;
-}
-
 FRHILockTextureResult FMetalDynamicRHI::RHILockTexture(FRHICommandListImmediate& RHICmdList, const FRHILockTextureArgs& Arguments)
 {
 	MTL_SCOPED_AUTORELEASE_POOL;
diff --git a/Engine/Source/Runtime/D3D12RHI/Private/D3D12RHIPrivate.h b/Engine/Source/Runtime/D3D12RHI/Private/D3D12RHIPrivate.h
--- a/Engine/Source/Runtime/D3D12RHI/Private/D3D12RHIPrivate.h
+++ b/Engine/Source/Runtime/D3D12RHI/Private/D3D12RHIPrivate.h
@@ -369,7 +369,6 @@ public:
 	[[nodiscard]] virtual FRHITextureInitializer RHICreateTextureInitializer(FRHICommandListBase& RHICmdList, const FRHITextureCreateDesc& CreateDesc) override;
 
 	virtual uint32 RHIComputeMemorySize(FRHITexture* TextureRHI) final override;
-	virtual FTextureRHIRef RHIAsyncReallocateTexture2D(FRHITexture* Texture2D, int32 NewMipCount, int32 NewSizeX, int32 NewSizeY, FThreadSafeCounter* RequestStatus) override;
 
 	virtual FRHILockTextureResult RHILockTexture(FRHICommandListImmediate& RHICmdList, const FRHILockTextureArgs& Arguments) final override;
 	virtual void                  RHIUnlockTexture(FRHICommandListImmediate& RHICmdList, const FRHILockTextureArgs& Arguments) final override;
@@ -477,13 +476,6 @@ public:
 	virtual void RHIWaitManualFence(FRHICommandList& RHICmdList, ID3D12Fence* Fence, uint64 Value) final override;
 	virtual void RHIVerifyResult(ID3D12Device* Device, HRESULT Result, const ANSICHAR* Code, const ANSICHAR* Filename, uint32 Line, FString Message = FString()) const final override;
 
-	//
-	// The Following functions are the _RenderThread version of the above functions. They allow the RHI to control the thread synchronization for greater efficiency.
-	// These will be un-commented as they are implemented.
-	//
-
-	virtual FTextureRHIRef AsyncReallocateTexture2D_RenderThread(class FRHICommandListImmediate& RHICmdList, FRHITexture* Texture2D, int32 NewMipCount, int32 NewSizeX, int32 NewSizeY, FThreadSafeCounter* RequestStatus);
-
 #if D3D12_RHI_RAYTRACING
 
 	virtual FRayTracingAccelerationStructureSize RHICalcRayTracingSceneSize(const FRayTracingSceneInitializer& Initializer) final override;
diff --git a/Engine/Source/Runtime/D3D12RHI/Private/D3D12Texture.cpp b/Engine/Source/Runtime/D3D12RHI/Private/D3D12Texture.cpp
--- a/Engine/Source/Runtime/D3D12RHI/Private/D3D12Texture.cpp
+++ b/Engine/Source/Runtime/D3D12RHI/Private/D3D12Texture.cpp
@@ -1616,118 +1616,6 @@ FORCENOINLINE void FD3D12CommandContext::CopyTextureRegionChecked(
 	);
 }
 
-FTextureRHIRef FD3D12DynamicRHI::AsyncReallocateTexture2D_RenderThread(FRHICommandListImmediate& RHICmdList, FRHITexture* Texture2DRHI, int32 NewMipCount, int32 NewSizeX, int32 NewSizeY, FThreadSafeCounter* RequestStatus)
-{
-	FD3D12Texture* OldTexture = FD3D12DynamicRHI::ResourceCast(Texture2DRHI);
-	
-	FRHITextureDesc Desc = OldTexture->GetDesc();
-	Desc.Extent = FIntPoint(NewSizeX, NewSizeY);
-	Desc.NumMips = NewMipCount;
-
-	FRHITextureCreateDesc CreateDesc(
-		Desc,
-		RHIGetDefaultResourceState(Desc.Flags, false),
-		TEXT("AsyncReallocateTexture2D_RenderThread")
-	);
-	CreateDesc.SetOwnerName(OldTexture->GetOwnerName());
-
-	// Allocate a new texture.
-	const FCreateTextureInternalResult CreateResult = CreateTextureInternal(CreateDesc, nullptr);
-	UE::D3D12RHI::TextureUtils::ReconcileInitialState(RHICmdList, CreateResult.Texture, CreateResult.CreateD3D12Access, CreateResult.DesiredD3D12Access, CreateResult.bNeedsDiscard);
-
-	FD3D12Texture* NewTexture = CreateResult.Texture;
-
-	RHICmdList.EnqueueLambda([
-		RootOldTexture = OldTexture,
-		RootNewTexture = NewTexture,
-		RequestStatus
-	](FRHICommandListBase& ExecutingCmdList)
-	{
-		// Use the GPU to asynchronously copy the old mip-maps into the new texture.
-		const uint32 NumSharedMips   = FMath::Min(RootOldTexture->GetNumMips(), RootNewTexture->GetNumMips());
-		const uint32 SourceMipOffset = RootOldTexture->GetNumMips() - NumSharedMips;
-		const uint32 DestMipOffset   = RootNewTexture->GetNumMips() - NumSharedMips;
-
-		for (FD3D12Texture::FDualLinkedObjectIterator It(RootOldTexture, RootNewTexture); It; ++It)
-		{
-			FD3D12Texture* DeviceOldTexture = It.GetFirst();
-			FD3D12Texture* DeviceNewTexture = It.GetSecond();
-			check(DeviceOldTexture->GetParentDevice() == DeviceNewTexture->GetParentDevice());
-
-			FD3D12CommandContext& Context = FD3D12CommandContext::Get(ExecutingCmdList, DeviceOldTexture->GetParentDevice()->GetGPUIndex());
-
-			FScopedResourceBarrier ScopeResourceBarrierDst(
-				Context,
-				DeviceNewTexture->GetResource(),
-				ED3D12Access::CopyDest,
-				D3D12_RESOURCE_BARRIER_ALL_SUBRESOURCES);
-			FScopedResourceBarrier ScopeResourceBarrierSrc(
-				Context,
-				DeviceOldTexture->GetResource(),
-				ED3D12Access::CopySrc,
-				D3D12_RESOURCE_BARRIER_ALL_SUBRESOURCES);
-			Context.FlushResourceBarriers();	// Must flush so the desired state is actually set.
-
-			for (uint32 MipIndex = 0; MipIndex < NumSharedMips; ++MipIndex)
-			{
-				// Use the GPU to copy between mip-maps.
-				// This is serialized with other D3D commands, so it isn't necessary to increment Counter to signal a pending asynchronous copy.
-
-				uint32 SrcSubresource = CalcSubresource(MipIndex + SourceMipOffset, 0, DeviceOldTexture->GetNumMips());
-				uint32 DstSubresource = CalcSubresource(MipIndex + DestMipOffset, 0, DeviceNewTexture->GetNumMips());
-
-				CD3DX12_TEXTURE_COPY_LOCATION DestCopyLocation(DeviceNewTexture->GetResource()->GetResource(), DstSubresource);
-				CD3DX12_TEXTURE_COPY_LOCATION SourceCopyLocation(DeviceOldTexture->GetResource()->GetResource(), SrcSubresource);
-				
-				Context.CopyTextureRegionChecked(
-					&DestCopyLocation,
-					0, 0, 0,
-					DeviceNewTexture->GetFormat(),
-					&SourceCopyLocation,
-					nullptr,
-					DeviceOldTexture->GetFormat(),
-					DeviceNewTexture->GetName()
-				);
-
-				Context.UpdateResidency(DeviceNewTexture->GetResource());
-				Context.UpdateResidency(DeviceOldTexture->GetResource());
-
-				Context.ConditionalSplitCommandList();
-
-				DEBUG_EXECUTE_COMMAND_CONTEXT(Context);
-			}
-		}
-
-		// Decrement the thread-safe counter used to track the completion of the reallocation, since D3D handles sequencing the
-		// async mip copies with other D3D calls.
-		RequestStatus->Decrement();
-	});
-
-	return NewTexture;
-}
-
-
-/**
- * Starts an asynchronous texture reallocation. It may complete immediately if the reallocation
- * could be performed without any reshuffling of texture memory, or if there isn't enough memory.
- * The specified status counter will be decremented by 1 when the reallocation is complete (success or failure).
- *
- * Returns a new reference to the texture, which will represent the new mip count when the reallocation is complete.
- * RHIGetAsyncReallocateTexture2DStatus() can be used to check the status of an ongoing or completed reallocation.
- *
- * @param Texture2D		- Texture to reallocate
- * @param NewMipCount	- New number of mip-levels
- * @param NewSizeX		- New width, in pixels
- * @param NewSizeY		- New height, in pixels
- * @param RequestStatus	- Will be decremented by 1 when the reallocation is complete (success or failure).
- * @return				- New reference to the texture, or an invalid reference upon failure
- */
-FTextureRHIRef FD3D12DynamicRHI::RHIAsyncReallocateTexture2D(FRHITexture* Texture2DRHI, int32 NewMipCount, int32 NewSizeX, int32 NewSizeY, FThreadSafeCounter* RequestStatus)
-{
-	UE_LOG(LogD3D12RHI, Fatal, TEXT("RHIAsyncReallocateTexture2D should not be called. AsyncReallocateTexture2D_RenderThread is where this is implemented."));
-	return nullptr;
-}
-
 ///////////////////////////////////////////////////////////////////////////////////////////
 // FD3D12Texture
 ///////////////////////////////////////////////////////////////////////////////////////////
diff --git a/Engine/Source/Runtime/NullDrv/Public/NullRHI.h b/Engine/Source/Runtime/NullDrv/Public/NullRHI.h
--- a/Engine/Source/Runtime/NullDrv/Public/NullRHI.h
+++ b/Engine/Source/Runtime/NullDrv/Public/NullRHI.h
@@ -292,17 +292,6 @@ public:
 	{ 
 		return 0; 
 	}
-	virtual FTextureRHIRef RHIAsyncReallocateTexture2D(FRHITexture* Texture2D, int32 NewMipCount, int32 NewSizeX, int32 NewSizeY, FThreadSafeCounter* RequestStatus) final override
-	{
-		const FRHITextureCreateDesc Desc =
-			FRHITextureCreateDesc::Create2D(TEXT("FNullDynamicRHI::RHIAsyncReallocateTexture2D"), NewSizeX, NewSizeY, Texture2D->GetFormat())
-			.SetClearValue(Texture2D->GetClearBinding())
-			.SetFlags(Texture2D->GetFlags())
-			.SetNumMips(NewMipCount)
-			.SetNumSamples(Texture2D->GetNumSamples());
-
-		return new FNullTexture(Desc);
-	}
 
 	virtual FRHILockTextureResult RHILockTexture(FRHICommandListImmediate& RHICmdList, const FRHILockTextureArgs& Arguments) final override
 	{
diff --git a/Engine/Source/Runtime/OpenGLDrv/Private/OpenGLDrv.h b/Engine/Source/Runtime/OpenGLDrv/Private/OpenGLDrv.h
--- a/Engine/Source/Runtime/OpenGLDrv/Private/OpenGLDrv.h
+++ b/Engine/Source/Runtime/OpenGLDrv/Private/OpenGLDrv.h
@@ -141,7 +141,6 @@ public:
 	UE_API virtual FTextureRHIRef RHIAsyncCreateTexture2D(uint32 SizeX, uint32 SizeY, uint8 Format, uint32 NumMips, ETextureCreateFlags Flags, ERHIAccess InResourceState, void** InitialMipData, uint32 NumInitialMips, const TCHAR* DebugName, FGraphEventRef& OutCompletionEvent) final override;
 	UE_API virtual void RHIGenerateMips(FRHITexture* Texture) final override;
 	UE_API virtual uint32 RHIComputeMemorySize(FRHITexture* TextureRHI) final override;
-	UE_API virtual FTextureRHIRef RHIAsyncReallocateTexture2D(FRHITexture* Texture2D, int32 NewMipCount, int32 NewSizeX, int32 NewSizeY, FThreadSafeCounter* RequestStatus) final override;
 
 	UE_API virtual FRHILockTextureResult RHILockTexture(FRHICommandListImmediate& RHICmdList, const FRHILockTextureArgs& Arguments) final override;
 	UE_API virtual void RHIUnlockTexture(FRHICommandListImmediate& RHICmdList, const FRHILockTextureArgs& Arguments) final override;
@@ -417,8 +416,6 @@ public:
 		}
 	};
 
-	UE_API virtual FTextureRHIRef AsyncReallocateTexture2D_RenderThread(class FRHICommandListImmediate& RHICmdList, FRHITexture* Texture2D, int32 NewMipCount, int32 NewSizeX, int32 NewSizeY, FThreadSafeCounter* RequestStatus) final override;
-
 	UE_API virtual FSamplerStateRHIRef RHICreateSamplerState(const FSamplerStateInitializerRHI& Initializer) final override;
 
 	virtual FGraphicsPipelineStateRHIRef RHICreateGraphicsPipelineState(const FGraphicsPipelineStateInitializer& Initializer) override
diff --git a/Engine/Source/Runtime/OpenGLDrv/Private/OpenGLResources.h b/Engine/Source/Runtime/OpenGLDrv/Private/OpenGLResources.h
--- a/Engine/Source/Runtime/OpenGLDrv/Private/OpenGLResources.h
+++ b/Engine/Source/Runtime/OpenGLDrv/Private/OpenGLResources.h
@@ -1020,11 +1020,6 @@ public:
 		bAllocatedStorage.Init(bInAllocatedStorage, this->GetNumMips() * (bCubemap ? 6 : 1));
 	}
 
-	/**
-	 * Clone texture from a source using CopyImageSubData
-	 */
-	UE_API void CloneViaCopyImage(FOpenGLTexture* Src, uint32 InNumMips, int32 SrcOffset, int32 DstOffset);
-
 	/**
 	 * Resolved the specified face for a read Lock, for non-renderable, CPU readable surfaces this eliminates the readback inside Lock itself.
 	 */
diff --git a/Engine/Source/Runtime/OpenGLDrv/Private/OpenGLTexture.cpp b/Engine/Source/Runtime/OpenGLDrv/Private/OpenGLTexture.cpp
--- a/Engine/Source/Runtime/OpenGLDrv/Private/OpenGLTexture.cpp
+++ b/Engine/Source/Runtime/OpenGLDrv/Private/OpenGLTexture.cpp
@@ -1026,7 +1026,7 @@ void* FOpenGLTexture::Lock(uint32 InMipIndex,uint32 ArrayIndex,EResourceLockMode
 	if (IsEvicted())
 	{
 		check(ArrayIndex == 0);
-		EvictionParamsPtr->SetMipData(InMipIndex, 0, MipBytes);
+		EvictionParamsPtr->SetMipData(InMipIndex, nullptr, MipBytes);
 		return EvictionParamsPtr->MipImageData[InMipIndex].GetData();
 	}
 	else
@@ -1382,54 +1382,6 @@ void FOpenGLTexture::RemoveEvictionData()
 	}
 }
 
-void FOpenGLTexture::CloneViaCopyImage(FOpenGLTexture* Src, uint32 InNumMips, int32 SrcOffset, int32 DstOffset)
-{
-	VERIFY_GL_SCOPE();
-
-	check(Src->bCanCreateAsEvicted == bCanCreateAsEvicted);
-	if (bCanCreateAsEvicted)
-	{
-		// Copy all mips that are present.
-		if (!(!Src->IsEvicted() || Src->EvictionParamsPtr->AreAllMipsPresent()))
-		{
-			UE_LOG(LogRHI, Warning, TEXT("IsEvicted %d, MipsPresent %d, InNumMips %d, SrcOffset %d, DstOffset %d"), Src->IsEvicted(), Src->EvictionParamsPtr->AreAllMipsPresent(), InNumMips, SrcOffset, DstOffset);
-			int MessageCount = 0;
-			for (const auto& MipData : Src->EvictionParamsPtr->MipImageData)
-			{
-				UE_LOG(LogRHI, Warning, TEXT("SrcMipData[%d].Num() == %d"), MessageCount++, MipData.Num());
-			}	
-		}
-		check(!Src->IsEvicted() || Src->EvictionParamsPtr->AreAllMipsPresent() );
-		EvictionParamsPtr->CloneMipData(*Src->EvictionParamsPtr, InNumMips, SrcOffset, DstOffset);
-
-		// the dest texture can remain evicted if: the src was also evicted or has all of the resident mips available or the dest texture has all mips already evicted.
-		if(IsEvicted() && (Src->IsEvicted() || Src->EvictionParamsPtr->AreAllMipsPresent() || EvictionParamsPtr->AreAllMipsPresent()))
-		{
-			return;
-		}
-	}
-
-	for (uint32 ArrayIndex = 0; ArrayIndex < this->GetEffectiveSizeZ(); ArrayIndex++)
-	{
-		// use the Copy Image functionality to copy mip level by mip level
-		for(uint32 MipIndex = 0;MipIndex < InNumMips;++MipIndex)
-		{
-			// Calculate the dimensions of the mip-map.
-			const uint32 DstMipIndex = MipIndex + DstOffset;
-			const uint32 SrcMipIndex = MipIndex + SrcOffset;
-			const uint32 MipSizeX = FMath::Max(this->GetSizeX() >> DstMipIndex,uint32(1));
-			const uint32 MipSizeY = FMath::Max(this->GetSizeY() >> DstMipIndex,uint32(1));
-			
-			if(FOpenGL::AmdWorkaround() && ((MipSizeX < 4) || (MipSizeY < 4))) break;
-
-			// copy the texture data
-			FOpenGL::CopyImageSubData(Src->GetResource(), Src->Target, SrcMipIndex, 0, 0, ArrayIndex,
-				GetResource(), Target, DstMipIndex, 0, 0, ArrayIndex, MipSizeX, MipSizeY, 1);
-		}
-	}
-	
-}
-
 /*-----------------------------------------------------------------------------
 	2D texture support.
 -----------------------------------------------------------------------------*/
@@ -1561,55 +1513,6 @@ uint32 FOpenGLDynamicRHI::RHIComputeMemorySize(FRHITexture* TextureRHI)
 	return Texture->MemorySize;
 }
 
-FTextureRHIRef FOpenGLDynamicRHI::AsyncReallocateTexture2D_RenderThread(class FRHICommandListImmediate& RHICmdList, FRHITexture* Texture2DRHI, int32 NewMipCount, int32 NewSizeX, int32 NewSizeY, FThreadSafeCounter* RequestStatus)
-{
-	return this->RHIAsyncReallocateTexture2D(Texture2DRHI, NewMipCount, NewSizeX, NewSizeY, RequestStatus);
-}
-
-FTextureRHIRef FOpenGLDynamicRHI::RHIAsyncReallocateTexture2D(FRHITexture* Texture2DRHI, int32 NewMipCount, int32 NewSizeX, int32 NewSizeY, FThreadSafeCounter* RequestStatus)
-{
-	FRHICommandListImmediate& RHICmdList = FRHICommandListImmediate::Get();
-
-	FOpenGLTexture* OldTexture = ResourceCast(Texture2DRHI);
-
-	FRHITextureDesc Desc = OldTexture->GetDesc();
-	int32 SourceMipCount = Desc.NumMips;
-
-	Desc.Extent = FIntPoint(NewSizeX, NewSizeY);
-	Desc.NumMips = NewMipCount;
-
-	FRHITextureCreateDesc CreateDesc(
-		Desc,
-		RHIGetDefaultResourceState(Desc.Flags, false),
-		TEXT("RHIAsyncReallocateTexture2D")
-	);
-	CreateDesc.SetOwnerName(OldTexture->GetOwnerName());
-
-	LLM_SCOPE_DYNAMIC_STAT_OBJECTPATH_FNAME(CreateDesc.OwnerName, ELLMTagSet::Assets);
-	LLM_SCOPE_DYNAMIC_STAT_OBJECTPATH_FNAME(CreateDesc.GetTraceClassName(), ELLMTagSet::AssetClasses);
-	UE_TRACE_METADATA_SCOPE_ASSET_FNAME(CreateDesc.DebugName, CreateDesc.GetTraceClassName(), CreateDesc.OwnerName);
-
-	FOpenGLTexture* NewTexture = new FOpenGLTexture(CreateDesc);
-
-	RHICmdList.EnqueueLambda([OldTexture, NewTexture, SourceMipCount, NewMipCount, RequestStatus](FRHICommandListImmediate& RHICmdList)
-	{
-		VERIFY_GL_SCOPE();
-
-		NewTexture->Initialize(RHICmdList);
-
-		// Use the GPU to asynchronously copy the old mip-maps into the new texture.
-		const uint32 NumSharedMips = FMath::Min(SourceMipCount, NewMipCount);
-		const uint32 SourceMipOffset = SourceMipCount - NumSharedMips;
-		const uint32 DestMipOffset = NewMipCount - NumSharedMips;
-
-		NewTexture->CloneViaCopyImage(OldTexture, NumSharedMips, SourceMipOffset, DestMipOffset);
-
-		RequestStatus->Decrement();
-	});
-
-	return NewTexture;
-}
-
 void FOpenGLDynamicRHI::RHIUpdateTexture2D(FRHICommandListBase& RHICmdList, FRHITexture* TextureRHI, uint32 MipIndex, const FUpdateTextureRegion2D& UpdateRegion, uint32 SourcePitch, const uint8* SourceData)
 {
 	const FPixelFormatInfo& FormatInfo = GPixelFormats[TextureRHI->GetFormat()];
@@ -1848,7 +1751,50 @@ void FOpenGLDynamicRHI::RHICopyTexture(FRHITexture* SourceTextureRHI, FRHITextur
 	VERIFY_GL_SCOPE();
 	FOpenGLTexture* SourceTexture = ResourceCast(SourceTextureRHI);
 	FOpenGLTexture* DestTexture = ResourceCast(DestTextureRHI);
-	
+
+	if (DestTexture->IsEvicted()
+		&& SourceTexture->EvictionParamsPtr.IsValid()
+		&& SourceTexture->GetFormat() == DestTexture->GetFormat()
+		&& CopyInfo.SourceSliceIndex == 0
+		&& CopyInfo.DestSliceIndex == 0
+		&& CopyInfo.NumSlices == 1
+		&& CopyInfo.SourcePosition.IsZero()
+		&& CopyInfo.DestPosition.IsZero())
+	{
+		// Both source and destination textures are candidates for eviction.
+		// We might be able to achieve the copy entirely offline by copying the CPU-side mip data, rather than involving the GPU, leaving the destination evicted.
+
+		// Check the dimensions of the copy are compatible
+		FUintVector3 SrcMipExtent = UE::RHITextureUtils::CalculateMipExtents(SourceTexture->GetDesc(), CopyInfo.SourceMipIndex);
+		FUintVector3 DstMipExtent = UE::RHITextureUtils::CalculateMipExtents(DestTexture->GetDesc(), CopyInfo.DestMipIndex);
+
+		if (SrcMipExtent == DstMipExtent && (CopyInfo.Size.IsZero() || CopyInfo.Size == FIntVector3(SrcMipExtent)))
+		{
+			// Check all the required mips from the source are available to copy from
+			bool bRequiredSrcMipsPresent = true;
+			const TArray<TArray<uint8>>& SrcMipData = SourceTexture->EvictionParamsPtr->MipImageData;
+
+			for (uint32 Mip = CopyInfo.SourceMipIndex, Index = 0; Index < CopyInfo.NumMips; ++Mip, ++Index)
+			{
+				if (!SrcMipData.IsValidIndex(Mip) || SrcMipData[Mip].IsEmpty())
+				{
+					bRequiredSrcMipsPresent = false;
+					break;
+				}
+			}
+
+			if (bRequiredSrcMipsPresent)
+			{
+				// We're replacing entire mips in the destination with copies from the source.
+				// Clone the offline mip data directly from the source.
+				DestTexture->EvictionParamsPtr->CloneMipData(*SourceTexture->EvictionParamsPtr, CopyInfo.NumMips, CopyInfo.SourceMipIndex, CopyInfo.DestMipIndex);
+
+				// Leave the destination texture evicted.
+				return;
+			}
+		}
+	}
+
 	//since the texture will be modified we cannot use the eviction data.
 	DestTexture->RemoveEvictionData();
 
diff --git a/Engine/Source/Runtime/RHI/Private/RHICommandList.cpp b/Engine/Source/Runtime/RHI/Private/RHICommandList.cpp
--- a/Engine/Source/Runtime/RHI/Private/RHICommandList.cpp
+++ b/Engine/Source/Runtime/RHI/Private/RHICommandList.cpp
@@ -8,6 +8,7 @@
 #include "RHI.h"
 #include "ProfilingDebugging/CsvProfiler.h"
 #include "ProfilingDebugging/MemoryTrace.h"
+#include "ProfilingDebugging/AssetMetadataTrace.h"
 #include "GenericPlatform/GenericPlatformCrashContext.h"
 #include "RHIBreadcrumbs.h"
 #include "RHIResourceReplace.h"
@@ -18,6 +19,7 @@
 #include "Stats/StatsTrace.h"
 #include "Stats/ThreadIdleStats.h"
 #include "HAL/PlatformMisc.h"
+#include "HAL/LowLevelMemStats.h"
 
 CSV_DEFINE_CATEGORY_MODULE(RHI_API, RHITStalls, false);
 CSV_DEFINE_CATEGORY_MODULE(RHI_API, RHITFlushes, false);
@@ -1885,6 +1887,69 @@ void FRHICommandListImmediate::EndFrame()
 	GDynamicRHI->RHIEndFrame_RenderThread(*this);
 }
 
+FTextureRHIRef FRHICommandListImmediate::AsyncReallocateTexture2D(FRHITexture* OldTexture, int32 NewMipCount, int32 NewSizeX, int32 NewSizeY, FThreadSafeCounter* RequestStatus)
+{
+	check(!GRHISupportsAsyncTextureCreation);
+
+	LLM_SCOPE(ELLMTag::Textures);
+	RHI_BREADCRUMB_EVENT_F(*this, "AsyncReallocateTexture2D", "AsyncReallocateTexture2D - %s", RHI_BREADCRUMB_FIELD("Texture Name", OldTexture->GetName()));
+
+	FRHIScopedGPUMask MaskScope(*this, FRHIGPUMask::All());
+
+	ERHIPipeline OriginalPipeline = GetPipeline();
+	if (OriginalPipeline != ERHIPipeline::Graphics)
+	{
+		SwitchPipeline(ERHIPipeline::Graphics);
+	}
+
+	FRHITextureDesc NewDesc = OldTexture->GetDesc();
+	NewDesc.Extent = FIntPoint(NewSizeX, NewSizeY);
+	NewDesc.NumMips = NewMipCount;
+
+	FString DebugName = OldTexture->GetName().ToString();
+
+	FRHITextureCreateDesc CreateDesc(
+		NewDesc,
+		ERHIAccess::CopyDest,
+		*DebugName
+	);
+	CreateDesc.SetOwnerName(OldTexture->GetOwnerName());
+
+	LLM_SCOPE_DYNAMIC_STAT_OBJECTPATH_FNAME(CreateDesc.OwnerName, ELLMTagSet::Assets);
+	LLM_SCOPE_DYNAMIC_STAT_OBJECTPATH_FNAME(CreateDesc.GetTraceClassName(), ELLMTagSet::AssetClasses);
+	UE_TRACE_METADATA_SCOPE_ASSET_FNAME(CreateDesc.DebugName, CreateDesc.GetTraceClassName(), CreateDesc.OwnerName);
+
+	FTextureRHIRef NewTexture = CreateTexture(CreateDesc);
+
+	// Ensure the old texture is in CopySrc state (the new texture was already created in CopyDest state).
+	FRHITransitionInfo TransitionsBefore[] =
+	{
+		FRHITransitionInfo(OldTexture, ERHIAccess::Unknown, ERHIAccess::CopySrc)
+	};
+	Transition(MakeArrayView(TransitionsBefore, UE_ARRAY_COUNT(TransitionsBefore)));
+
+	UE::RHI::CopySharedMips(*this, OldTexture, NewTexture);
+
+	FRHITransitionInfo TransitionsAfter[] =
+	{
+		FRHITransitionInfo(OldTexture, ERHIAccess::CopySrc, ERHIAccess::SRVMask),
+		FRHITransitionInfo(NewTexture, ERHIAccess::CopyDest, RHIGetDefaultResourceState(NewDesc.Flags, false))
+	};
+	Transition(MakeArrayView(TransitionsAfter, UE_ARRAY_COUNT(TransitionsAfter)));
+
+	if (OriginalPipeline != ERHIPipeline::Graphics)
+	{
+		SwitchPipeline(OriginalPipeline);
+	}
+
+	EnqueueLambda([RequestStatus](FRHICommandListImmediate&)
+	{
+		RequestStatus->Decrement();
+	});
+
+	return NewTexture;
+}
+
 void FRHIComputeCommandList::Transition(TArrayView<const FRHITransitionInfo> Infos, ERHITransitionCreateFlags CreateFlags)
 {
 	TransitionInternal(Infos, CreateFlags);
@@ -2158,14 +2223,6 @@ void FDynamicRHI::UnlockStagingBuffer_RenderThread(class FRHICommandListImmediat
 	}
 }
 
-FTextureRHIRef FDynamicRHI::AsyncReallocateTexture2D_RenderThread(class FRHICommandListImmediate& RHICmdList, FRHITexture* Texture2D, int32 NewMipCount, int32 NewSizeX, int32 NewSizeY, FThreadSafeCounter* RequestStatus)
-{
-	QUICK_SCOPE_CYCLE_COUNTER(STAT_RHIMETHOD_AsyncReallocateTexture2D_Flush);
-	CSV_SCOPED_TIMING_STAT(RHITFlushes, AsyncReallocateTexture2D_RenderThread);
-	RHICmdList.ImmediateFlush(EImmediateFlushType::FlushRHIThread);  
-	return GDynamicRHI->RHIAsyncReallocateTexture2D(Texture2D, NewMipCount, NewSizeX, NewSizeY, RequestStatus);
-}
-
 FUpdateTexture3DData FDynamicRHI::RHIBeginUpdateTexture3D(FRHICommandListBase& RHICmdList, FRHITexture* Texture, uint32 MipIndex, const struct FUpdateTextureRegion3D& UpdateRegion)
 {
 	check(IsInParallelRenderingThread());
diff --git a/Engine/Source/Runtime/RHI/Public/DynamicRHI.h b/Engine/Source/Runtime/RHI/Public/DynamicRHI.h
--- a/Engine/Source/Runtime/RHI/Public/DynamicRHI.h
+++ b/Engine/Source/Runtime/RHI/Public/DynamicRHI.h
@@ -544,25 +544,6 @@ public:
 	// FlushType: Thread safe
 	virtual uint32 RHIComputeMemorySize(FRHITexture* TextureRHI) = 0;
 
-	/**
-	* Starts an asynchronous texture reallocation. It may complete immediately if the reallocation
-	* could be performed without any reshuffling of texture memory, or if there isn't enough memory.
-	* The specified status counter will be decremented by 1 when the reallocation is complete (success or failure).
-	*
-	* Returns a new reference to the texture, which will represent the new mip count when the reallocation is complete.
-	* RHIFinalizeAsyncReallocateTexture2D() must be called to complete the reallocation.
-	*
-	* @param Texture2D		- Texture to reallocate
-	* @param NewMipCount	- New number of mip-levels
-	* @param NewSizeX		- New width, in pixels
-	* @param NewSizeY		- New height, in pixels
-	* @param RequestStatus	- Will be decremented by 1 when the reallocation is complete (success or failure).
-	* @return				- New reference to the texture, or an invalid reference upon failure
-	*/
-	// FlushType: Flush RHI Thread
-	// NP: Note that no RHI currently implements this as an async call, we should simplify the API.
-	virtual FTextureRHIRef RHIAsyncReallocateTexture2D(FRHITexture* Texture2D, int32 NewMipCount, int32 NewSizeX, int32 NewSizeY, FThreadSafeCounter* RequestStatus) = 0;
-
 	virtual FRHILockTextureResult RHILockTexture  (FRHICommandListImmediate& RHICmdList, const FRHILockTextureArgs& Arguments) = 0;
 	virtual void                  RHIUnlockTexture(FRHICommandListImmediate& RHICmdList, const FRHILockTextureArgs& Arguments) = 0;
 
@@ -931,8 +912,6 @@ public:
 	//
 	virtual void RHIProcessDeleteQueue() {}
 
-	RHI_API virtual FTextureRHIRef AsyncReallocateTexture2D_RenderThread(class FRHICommandListImmediate& RHICmdList, FRHITexture* Texture2D, int32 NewMipCount, int32 NewSizeX, int32 NewSizeY, FThreadSafeCounter* RequestStatus);
-
 	RHI_API virtual FUpdateTexture3DData RHIBeginUpdateTexture3D(FRHICommandListBase& RHICmdList, FRHITexture* Texture, uint32 MipIndex, const struct FUpdateTextureRegion3D& UpdateRegion);
 	RHI_API virtual void RHIEndUpdateTexture3D(FRHICommandListBase& RHICmdList, FUpdateTexture3DData& UpdateData);
 
diff --git a/Engine/Source/Runtime/RHI/Public/RHICommandList.h b/Engine/Source/Runtime/RHI/Public/RHICommandList.h
--- a/Engine/Source/Runtime/RHI/Public/RHICommandList.h
+++ b/Engine/Source/Runtime/RHI/Public/RHICommandList.h
@@ -4518,11 +4518,7 @@ public:
 		return GDynamicRHI->RHIGetTextureMemoryVisualizeData(TextureData,SizeX,SizeY,Pitch,PixelSize);
 	}
 	
-	inline FTextureRHIRef AsyncReallocateTexture2D(FRHITexture* Texture2D, int32 NewMipCount, int32 NewSizeX, int32 NewSizeY, FThreadSafeCounter* RequestStatus)
-	{
-		LLM_SCOPE(ELLMTag::Textures);
-		return GDynamicRHI->AsyncReallocateTexture2D_RenderThread(*this, Texture2D, NewMipCount, NewSizeX, NewSizeY, RequestStatus);
-	}
+	RHI_API FTextureRHIRef AsyncReallocateTexture2D(FRHITexture* Texture2D, int32 NewMipCount, int32 NewSizeX, int32 NewSizeY, FThreadSafeCounter* RequestStatus);
 
 	inline FRHILockTextureResult LockTexture(const FRHILockTextureArgs& Arguments)
 	{
diff --git a/Engine/Source/Runtime/RHI/Public/RHIValidation.h b/Engine/Source/Runtime/RHI/Public/RHIValidation.h
--- a/Engine/Source/Runtime/RHI/Public/RHIValidation.h
+++ b/Engine/Source/Runtime/RHI/Public/RHIValidation.h
@@ -552,33 +552,6 @@ public:
 		return RHI->RHIComputeMemorySize(TextureRHI);
 	}
 
-	/**
-	* Starts an asynchronous texture reallocation. It may complete immediately if the reallocation
-	* could be performed without any reshuffling of texture memory, or if there isn't enough memory.
-	* The specified status counter will be decremented by 1 when the reallocation is complete (success or failure).
-	*
-	* Returns a new reference to the texture, which will represent the new mip count when the reallocation is complete.
-	* RHIFinalizeAsyncReallocateTexture2D() must be called to complete the reallocation.
-	*
-	* @param Texture2D		- Texture to reallocate
-	* @param NewMipCount	- New number of mip-levels
-	* @param NewSizeX		- New width, in pixels
-	* @param NewSizeY		- New height, in pixels
-	* @param RequestStatus	- Will be decremented by 1 when the reallocation is complete (success or failure).
-	* @return				- New reference to the texture, or an invalid reference upon failure
-	*/
-	// FlushType: Flush RHI Thread
-	// NP: Note that no RHI currently implements this as an async call, we should simplify the API.
-	virtual FTextureRHIRef RHIAsyncReallocateTexture2D(FRHITexture* Texture2D, int32 NewMipCount, int32 NewSizeX, int32 NewSizeY, FThreadSafeCounter* RequestStatus) override final
-	{
-		// TODO: find proper state for new texture
-		ERHIAccess ResourceState = ERHIAccess::SRVMask;
-
-		FTextureRHIRef NewTexture2D = RHI->RHIAsyncReallocateTexture2D(Texture2D, NewMipCount, NewSizeX, NewSizeY, RequestStatus);
-		NewTexture2D->InitBarrierTracking(NewMipCount, 1, NewTexture2D->GetFormat(), NewTexture2D->GetFlags(), ResourceState, NewTexture2D->GetTrackerResource()->GetDebugName()); // @todo the threading of GetDebugName() is wrong.
-		return NewTexture2D;
-	}
-
 	virtual FRHILockTextureResult RHILockTexture(FRHICommandListImmediate& RHICmdList, const FRHILockTextureArgs& Arguments) final override;
 	virtual void RHIUnlockTexture(FRHICommandListImmediate& RHICmdList, const FRHILockTextureArgs& Arguments) final override;
 
@@ -964,16 +937,6 @@ public:
 		return RHI->RHIGetMinimumAlignmentForBufferBackedSRV(Format);
 	}
 
-	virtual FTextureRHIRef AsyncReallocateTexture2D_RenderThread(class FRHICommandListImmediate& RHICmdList, FRHITexture* Texture2D, int32 NewMipCount, int32 NewSizeX, int32 NewSizeY, FThreadSafeCounter* RequestStatus) override final
-	{
-		// TODO: find proper state for new texture
-		ERHIAccess ResourceState = ERHIAccess::SRVMask;
-
-		FTextureRHIRef NewTexture2D = RHI->AsyncReallocateTexture2D_RenderThread(RHICmdList, Texture2D, NewMipCount, NewSizeX, NewSizeY, RequestStatus);
-		NewTexture2D->InitBarrierTracking(NewMipCount, 1, NewTexture2D->GetFormat(), NewTexture2D->GetFlags(), ResourceState, NewTexture2D->GetTrackerResource()->GetDebugName()); // @todo the threading of GetDebugName() is wrong.
-		return NewTexture2D;
-	}
-
 	virtual void* LockBuffer_BottomOfPipe(class FRHICommandListBase& RHICmdList, FRHIBuffer* Buffer, uint32 Offset, uint32 SizeRHI, EResourceLockMode LockMode) override final
 	{
 		RHI_VALIDATION_CHECK(LockMode != RLM_WriteOnly_NoOverwrite || GRHISupportsMapWriteNoOverwrite, TEXT("Using RLM_WriteOnly_NoOverwrite when the RHI doesn't support it."));
diff --git a/Engine/Source/Runtime/VulkanRHI/Private/VulkanDynamicRHI.h b/Engine/Source/Runtime/VulkanRHI/Private/VulkanDynamicRHI.h
--- a/Engine/Source/Runtime/VulkanRHI/Private/VulkanDynamicRHI.h
+++ b/Engine/Source/Runtime/VulkanRHI/Private/VulkanDynamicRHI.h
@@ -181,7 +181,6 @@ public:
 
 	virtual FTextureRHIRef RHIAsyncCreateTexture2D(uint32 SizeX, uint32 SizeY, uint8 Format, uint32 NumMips, ETextureCreateFlags Flags, ERHIAccess InResourceState, void** InitialMipData, uint32 NumInitialMips, const TCHAR* DebugName, FGraphEventRef& OutCompletionEvent) final override;
 	virtual uint32 RHIComputeMemorySize(FRHITexture* TextureRHI) final override;
-	virtual FTextureRHIRef RHIAsyncReallocateTexture2D(FRHITexture* Texture2D, int32 NewMipCount, int32 NewSizeX, int32 NewSizeY, FThreadSafeCounter* RequestStatus) final override;
 
 	virtual FRHILockTextureResult RHILockTexture(FRHICommandListImmediate& RHICmdList, const FRHILockTextureArgs& Arguments) final override;
 	virtual void RHIUnlockTexture(FRHICommandListImmediate& RHICmdList, const FRHILockTextureArgs& Arguments) final override;
@@ -245,8 +244,6 @@ public:
 	virtual uint64 RHIComputeStatePrecachePSOHash(const FGraphicsPipelineStateInitializer& Initializer) final override;
 	virtual bool RHIMatchPrecachePSOInitializers(const FGraphicsPipelineStateInitializer& LHS, const FGraphicsPipelineStateInitializer& RHS) final override;
 
-	virtual FTextureRHIRef AsyncReallocateTexture2D_RenderThread(class FRHICommandListImmediate& RHICmdList, FRHITexture* Texture2D, int32 NewMipCount, int32 NewSizeX, int32 NewSizeY, FThreadSafeCounter* RequestStatus) override final;
-
 	virtual FUpdateTexture3DData RHIBeginUpdateTexture3D(FRHICommandListBase& RHICmdList, FRHITexture* Texture, uint32 MipIndex, const struct FUpdateTextureRegion3D& UpdateRegion) override final;
 	virtual void RHIEndUpdateTexture3D(FRHICommandListBase& RHICmdList, FUpdateTexture3DData& UpdateData) override final;
 
diff --git a/Engine/Source/Runtime/VulkanRHI/Private/VulkanTexture.cpp b/Engine/Source/Runtime/VulkanRHI/Private/VulkanTexture.cpp
--- a/Engine/Source/Runtime/VulkanRHI/Private/VulkanTexture.cpp
+++ b/Engine/Source/Runtime/VulkanRHI/Private/VulkanTexture.cpp
@@ -1564,105 +1564,6 @@ FTextureRHIRef FVulkanDynamicRHI::RHIAsyncCreateTexture2D(uint32 SizeX,uint32 Si
 	return FTextureRHIRef();
 }
 
-static void DoAsyncReallocateTexture2D(FVulkanContextCommon& Context, FVulkanTexture* OldTexture, FVulkanTexture* NewTexture, int32 NewMipCount, int32 NewSizeX, int32 NewSizeY, FThreadSafeCounter* RequestStatus)
-{
-	LLM_SCOPE_VULKAN(ELLMTagVulkan::VulkanTextures);
-	//QUICK_SCOPE_CYCLE_COUNTER(STAT_FRHICommandGnmAsyncReallocateTexture2D_Execute);
-
-	// figure out what mips to copy from/to
-	const uint32 NumSharedMips = FMath::Min(OldTexture->GetNumMips(), NewTexture->GetNumMips());
-	const uint32 SourceFirstMip = OldTexture->GetNumMips() - NumSharedMips;
-	const uint32 DestFirstMip = NewTexture->GetNumMips() - NumSharedMips;
-
-	FVulkanCommandBuffer& CommandBuffer = Context.GetCommandBuffer();
-	ensure(CommandBuffer.IsOutsideRenderPass());
-
-	VkCommandBuffer StagingCommandBuffer = CommandBuffer.GetHandle();
-
-	VkImageCopy Regions[MAX_TEXTURE_MIP_COUNT];
-	check(NumSharedMips <= MAX_TEXTURE_MIP_COUNT);
-	FMemory::Memzero(&Regions[0], sizeof(VkImageCopy) * NumSharedMips);
-	for (uint32 Index = 0; Index < NumSharedMips; ++Index)
-	{
-		uint32 MipWidth = FMath::Max<uint32>(NewSizeX >> (DestFirstMip + Index), 1u);
-		uint32 MipHeight = FMath::Max<uint32>(NewSizeY >> (DestFirstMip + Index), 1u);
-
-		VkImageCopy& Region = Regions[Index];
-		Region.srcSubresource.aspectMask = VK_IMAGE_ASPECT_COLOR_BIT;
-		Region.srcSubresource.mipLevel = SourceFirstMip + Index;
-		Region.srcSubresource.baseArrayLayer = 0;
-		Region.srcSubresource.layerCount = 1;
-		//Region.srcOffset
-		Region.dstSubresource.aspectMask = VK_IMAGE_ASPECT_COLOR_BIT;
-		Region.dstSubresource.mipLevel = DestFirstMip + Index;
-		Region.dstSubresource.baseArrayLayer = 0;
-		Region.dstSubresource.layerCount = 1;
-		//Region.dstOffset
-		Region.extent.width = MipWidth;
-		Region.extent.height = MipHeight;
-		Region.extent.depth = 1;
-	}
-
-	const VkImageSubresourceRange SourceSubResourceRange = FVulkanPipelineBarrier::MakeSubresourceRange(VK_IMAGE_ASPECT_COLOR_BIT, SourceFirstMip, NumSharedMips);
-	const VkImageSubresourceRange DestSubResourceRange = FVulkanPipelineBarrier::MakeSubresourceRange(VK_IMAGE_ASPECT_COLOR_BIT, DestFirstMip, NumSharedMips);
-
-	{
-		// Pre-copy barriers
-		FVulkanPipelineBarrier Barrier;
-		Barrier.AddImageLayoutTransition(OldTexture->Image, OldTexture->GetDefaultLayout(), VK_IMAGE_LAYOUT_TRANSFER_SRC_OPTIMAL, SourceSubResourceRange);
-		Barrier.AddImageLayoutTransition(NewTexture->Image, VK_IMAGE_LAYOUT_UNDEFINED, VK_IMAGE_LAYOUT_TRANSFER_DST_OPTIMAL, DestSubResourceRange);
-		Barrier.Execute(&CommandBuffer);
-	}
-
-	VulkanRHI::vkCmdCopyImage(StagingCommandBuffer, OldTexture->Image, VK_IMAGE_LAYOUT_TRANSFER_SRC_OPTIMAL, NewTexture->Image, VK_IMAGE_LAYOUT_TRANSFER_DST_OPTIMAL, NumSharedMips, Regions);
-
-	{
-		// Post-copy barriers
-		FVulkanPipelineBarrier Barrier;
-		Barrier.AddImageLayoutTransition(OldTexture->Image, VK_IMAGE_LAYOUT_TRANSFER_SRC_OPTIMAL, OldTexture->GetDefaultLayout(), SourceSubResourceRange);
-		Barrier.AddImageLayoutTransition(NewTexture->Image, VK_IMAGE_LAYOUT_TRANSFER_DST_OPTIMAL, NewTexture->GetDefaultLayout(), DestSubResourceRange);
-		Barrier.Execute(&CommandBuffer);
-	}
-
-	// request is now complete
-	RequestStatus->Decrement();
-
-	// the next unlock for this texture can't block the GPU (it's during runtime)
-	//NewTexture->bSkipBlockOnUnlock = true;
-}
-
-FTextureRHIRef FVulkanDynamicRHI::AsyncReallocateTexture2D_RenderThread(FRHICommandListImmediate& RHICmdList, FRHITexture* OldTextureRHI, int32 NewMipCount, int32 NewSizeX, int32 NewSizeY, FThreadSafeCounter* RequestStatus)
-{
-	LLM_SCOPE_VULKAN(ELLMTagVulkan::VulkanTextures);
-
-	FVulkanTexture* OldTexture = ResourceCast(OldTextureRHI);
-	const FRHITextureDesc& OldDesc = OldTexture->GetDesc();
-
-	const FRHITextureCreateDesc CreateDesc =
-		FRHITextureCreateDesc::Create2D(TEXT("AsyncReallocateTexture2D_RenderThread"), NewSizeX, NewSizeY, OldDesc.Format)
-		.SetClearValue(OldDesc.ClearValue)
-		.SetFlags(OldDesc.Flags)
-		.SetNumMips(NewMipCount)
-		.SetNumSamples(OldDesc.NumSamples)
-		.DetermineInititialState()
-		.SetOwnerName(OldTexture->GetOwnerName());
-
-	FVulkanTexture* NewTexture = CreateTextureInternal(RHICmdList, CreateDesc);
-
-	RHICmdList.EnqueueLambda(TEXT("AsyncReallocateTexture2D"), [OldTexture, NewTexture, NewMipCount, NewSizeX, NewSizeY, RequestStatus](FRHICommandListImmediate& ImmCmdList)
-	{
-		FVulkanUploadContext& UploadContext = FVulkanUploadContext::Get(ImmCmdList);
-		DoAsyncReallocateTexture2D(UploadContext, OldTexture, NewTexture, NewMipCount, NewSizeX, NewSizeY, RequestStatus);
-	});
-
-	return NewTexture;
-}
-
-FTextureRHIRef FVulkanDynamicRHI::RHIAsyncReallocateTexture2D(FRHITexture* OldTextureRHI, int32 NewMipCount, int32 NewSizeX, int32 NewSizeY, FThreadSafeCounter* RequestStatus)
-{
-	return AsyncReallocateTexture2D_RenderThread(FRHICommandListImmediate::Get(), OldTextureRHI, NewMipCount, NewSizeX, NewSizeY, RequestStatus);
-}
-
 static FCriticalSection GTextureMapLock;
 
 FRHILockTextureResult FVulkanDynamicRHI::RHILockTexture(FRHICommandListImmediate& RHICmdList, const FRHILockTextureArgs& Arguments)
diff --git a/Engine/Source/Runtime/Windows/D3D11RHI/Private/D3D11RHIPrivate.h b/Engine/Source/Runtime/Windows/D3D11RHI/Private/D3D11RHIPrivate.h
--- a/Engine/Source/Runtime/Windows/D3D11RHI/Private/D3D11RHIPrivate.h
+++ b/Engine/Source/Runtime/Windows/D3D11RHI/Private/D3D11RHIPrivate.h
@@ -270,8 +270,6 @@ public:
 	[[nodiscard]] D3D11RHI_API virtual FRHITextureInitializer RHICreateTextureInitializer(FRHICommandListBase& RHICmdList, const FRHITextureCreateDesc& CreateDesc) final override;
 	D3D11RHI_API virtual FTextureRHIRef RHIAsyncCreateTexture2D(uint32 SizeX, uint32 SizeY, uint8 Format, uint32 NumMips, ETextureCreateFlags Flags, ERHIAccess InResourceState, void** InitialMipData, uint32 NumInitialMips, const TCHAR* DebugName, FGraphEventRef& OutCompletionEvent) final override;
 	D3D11RHI_API virtual uint32 RHIComputeMemorySize(FRHITexture* TextureRHI) final override;
-	D3D11RHI_API virtual void RHIAsyncCopyTexture2DCopy(FRHITexture* NewTexture2DRHI, FRHITexture* Texture2DRHI, int32 NewMipCount, int32 NewSizeX, int32 NewSizeY, FThreadSafeCounter* RequestStatus);
-	D3D11RHI_API virtual FTextureRHIRef RHIAsyncReallocateTexture2D(FRHITexture* Texture2D, int32 NewMipCount, int32 NewSizeX, int32 NewSizeY, FThreadSafeCounter* RequestStatus) final override;
 
 	D3D11RHI_API virtual FRHILockTextureResult RHILockTexture(FRHICommandListImmediate& RHICmdList, const FRHILockTextureArgs& Arguments) final override;
 	D3D11RHI_API virtual void                  RHIUnlockTexture(FRHICommandListImmediate& RHICmdList, const FRHILockTextureArgs& Arguments) final override;
@@ -379,10 +377,6 @@ public:
 	D3D11RHI_API void PollQueryResults();
 	D3D11RHI_API bool PollQueryResultsForEndFrame(FD3D11RenderQuery* Query);
 
-	// *_RenderThread functions. Command lists call these functions on RT. You can implement your own behavior inside these functions.
-	// For example, deferring the actual creation to RHI thread by sending an RHI command.
-	// For D3D11, these functions mainly just remove RHIT stalls because ID3D11Device is thread safe.
-	D3D11RHI_API virtual FTextureRHIRef AsyncReallocateTexture2D_RenderThread(class FRHICommandListImmediate& RHICmdList, FRHITexture* Texture2D, int32 NewMipCount, int32 NewSizeX, int32 NewSizeY, FThreadSafeCounter* RequestStatus) final override;
 	D3D11RHI_API virtual void RHIEndUpdateTexture3D(FRHICommandListBase& RHICmdList, FUpdateTexture3DData& UpdateData) final override;
 
 	D3D11RHI_API virtual void RHIBeginRenderPass(const FRHIRenderPassInfo& InInfo, const TCHAR* InName) final override;
diff --git a/Engine/Source/Runtime/Windows/D3D11RHI/Private/D3D11Texture.cpp b/Engine/Source/Runtime/Windows/D3D11RHI/Private/D3D11Texture.cpp
--- a/Engine/Source/Runtime/Windows/D3D11RHI/Private/D3D11Texture.cpp
+++ b/Engine/Source/Runtime/Windows/D3D11RHI/Private/D3D11Texture.cpp
@@ -945,93 +945,6 @@ uint32 FD3D11DynamicRHI::RHIComputeMemorySize(FRHITexture* TextureRHI)
 	return Texture->GetMemorySize();
 }
 
-void FD3D11DynamicRHI::RHIAsyncCopyTexture2DCopy(FRHITexture* NewTexture2DRHI, FRHITexture* Texture2DRHI, int32 NewMipCount, int32 NewSizeX, int32 NewSizeY, FThreadSafeCounter* RequestStatus)
-{
-	FD3D11Texture* Texture2D = ResourceCast(Texture2DRHI);
-	FD3D11Texture* NewTexture2D = ResourceCast(NewTexture2DRHI);
-
-	// Use the GPU to asynchronously copy the old mip-maps into the new texture.
-	const uint32 NumSharedMips = FMath::Min(Texture2D->GetNumMips(), NewTexture2D->GetNumMips());
-	const uint32 SourceMipOffset = Texture2D->GetNumMips() - NumSharedMips;
-	const uint32 DestMipOffset = NewTexture2D->GetNumMips() - NumSharedMips;
-	for (uint32 MipIndex = 0; MipIndex < NumSharedMips; ++MipIndex)
-	{
-		// Use the GPU to copy between mip-maps.
-		// This is serialized with other D3D commands, so it isn't necessary to increment Counter to signal a pending asynchronous copy.
-		Direct3DDeviceIMContext->CopySubresourceRegion(
-			NewTexture2D->GetResource(),
-			D3D11CalcSubresource(MipIndex + DestMipOffset, 0, NewTexture2D->GetNumMips()),
-			0,
-			0,
-			0,
-			Texture2D->GetResource(),
-			D3D11CalcSubresource(MipIndex + SourceMipOffset, 0, Texture2D->GetNumMips()),
-			NULL
-		);
-	}
-
-	// Decrement the thread-safe counter used to track the completion of the reallocation, since D3D handles sequencing the
-	// async mip copies with other D3D calls.
-	RequestStatus->Decrement();
-}
-
-FTextureRHIRef FD3D11DynamicRHI::RHIAsyncReallocateTexture2D(FRHITexture* Texture2DRHI, int32 NewMipCount, int32 NewSizeX, int32 NewSizeY, FThreadSafeCounter* RequestStatus)
-{
-	FD3D11Texture* Texture2D = ResourceCast(Texture2DRHI);
-
-	FRHITextureCreateDesc CreateDesc(
-		Texture2D->GetDesc(),
-		RHIGetDefaultResourceState(Texture2D->GetDesc().Flags, false),
-		TEXT("RHIAsyncReallocateTexture2D")
-	);
-	CreateDesc.Extent = FIntPoint(NewSizeX, NewSizeY);
-	CreateDesc.NumMips = NewMipCount;
-	CreateDesc.SetOwnerName(Texture2D->GetOwnerName());
-
-	// Allocate a new texture.
-	FD3D11Texture* NewTexture2D = CreateTextureInternal(CreateDesc, {});
-
-	RHIAsyncCopyTexture2DCopy(NewTexture2D, Texture2D, NewMipCount, NewSizeX, NewSizeY, RequestStatus);
-
-	return NewTexture2D;
-}
-
-FTextureRHIRef FD3D11DynamicRHI::AsyncReallocateTexture2D_RenderThread(
-	class FRHICommandListImmediate& RHICmdList,
-	FRHITexture* Texture2D,
-	int32 NewMipCount,
-	int32 NewSizeX,
-	int32 NewSizeY,
-	FThreadSafeCounter* RequestStatus)
-{
-	FTextureRHIRef NewTexture2D;
-
-	if (ShouldNotEnqueueRHICommand())
-	{
-		NewTexture2D = RHIAsyncReallocateTexture2D(Texture2D, NewMipCount, NewSizeX, NewSizeY, RequestStatus);
-	}
-	else
-	{
-		// Allocate a new texture.
-		FRHITextureCreateDesc CreateDesc(
-			Texture2D->GetDesc(),
-			RHIGetDefaultResourceState(Texture2D->GetDesc().Flags, false),
-			TEXT("AsyncReallocateTexture2D_RenderThread")
-		);
-		CreateDesc.Extent = FIntPoint(NewSizeX, NewSizeY);
-		CreateDesc.NumMips = NewMipCount;
-		CreateDesc.SetOwnerName(Texture2D->GetOwnerName());
-
-		NewTexture2D = CreateTextureInternal(CreateDesc, {});
-
-		RunOnRHIThread([this, NewTexture2D, Texture2D, NewMipCount, NewSizeX, NewSizeY, RequestStatus]()
-		{
-			RHIAsyncCopyTexture2DCopy(NewTexture2D, Texture2D, NewMipCount, NewSizeX, NewSizeY, RequestStatus);
-		});
-	}
-	return NewTexture2D;
-}
-
 FRHILockTextureResult FD3D11Texture::Lock(FD3D11DynamicRHI* D3DRHI, const FRHILockTextureArgs& Arguments, bool bForceLockDeferred)
 {
 	check(!IsTexture3D()); // Only 2D texture locks are implemented
