Android GLES - Add an option to inject benign code flow change into all fragment shaders to coax changes to code gen with some vendors.
device capabilities are checked when loading the program binary cache, this will invalidate the cache should any capabilities change.
fix bug in OpenGL RHI's ReplaceCString that causes problems when the replacement string contains a match for the source.

#tests Tested local client, tested updating from old gles cache, tested updated cache works on repeat runs, tested changing cvars correctly rejects the old cache.

diff --git a/Engine/Source/Runtime/OpenGLDrv/Private/OpenGLProgramBinaryFileCache.cpp b/Engine/Source/Runtime/OpenGLDrv/Private/OpenGLProgramBinaryFileCache.cpp
--- a/Engine/Source/Runtime/OpenGLDrv/Private/OpenGLProgramBinaryFileCache.cpp
+++ b/Engine/Source/Runtime/OpenGLDrv/Private/OpenGLProgramBinaryFileCache.cpp
@@ -107,19 +107,29 @@ namespace UE
 			return StoreCompressedBinariesCVar->GetInt() != 0;
 		}
 
-		static const uint32 GBinaryProgramFileVersion = 6;
+		static const uint32 GBinaryProgramFileVersion = 7;
 		struct FBinaryCacheFileHeader
 		{
 			uint32 Version = 0xFFFFFFFF;
+			uint32 CapabilitiesCRC = 0;
 			FGuid BinaryCacheGuid;
 			bool bCacheUsesCompressedBinaries;
 			uint32 ProgramCount = 0;
 			uint32 ValidSize = 0;
 
+			static uint32 GetCapabilitiesCRC()
+			{
+				FOpenGLShaderDeviceCapabilities Capabilities;
+				GetCurrentOpenGLShaderDeviceCapabilities(Capabilities);
+				uint32 crc = FCrc::MemCrc32(&Capabilities, sizeof(Capabilities));
+				return crc;
+			}
+
 			static FBinaryCacheFileHeader CreateHeader(const FGuid& BinaryCacheGuidIn, uint32 NumPrograms, uint32 ValidSize)
 			{
 				FBinaryCacheFileHeader NewHeader;
 				NewHeader.Version = GBinaryProgramFileVersion;
+				NewHeader.CapabilitiesCRC = GetCapabilitiesCRC();
 				NewHeader.BinaryCacheGuid = BinaryCacheGuidIn;
 				NewHeader.bCacheUsesCompressedBinaries = UE::OpenGL::AreBinaryProgramsCompressed();
 				NewHeader.ProgramCount = NumPrograms;
@@ -133,6 +143,7 @@ namespace UE
 				check(Ar.IsLoading() || Header.IsValidVersion()); // This should always be correct when saving.
 				if (Header.IsValidVersion())
 				{
+					Ar << Header.CapabilitiesCRC;
 					Ar << Header.BinaryCacheGuid;
 					Ar << Header.bCacheUsesCompressedBinaries;
 					Ar << Header.ProgramCount;
@@ -150,6 +161,7 @@ namespace UE
 			bool IsValid(const FGuid* OptionalWantedGuid) const
 			{
 				return IsValidVersion()
+					&& CapabilitiesCRC == GetCapabilitiesCRC()
 					&& (UE::OpenGL::AreBinaryProgramsCompressed() == bCacheUsesCompressedBinaries)
 					&& (OptionalWantedGuid == nullptr || (*OptionalWantedGuid == BinaryCacheGuid))
 					&& (ProgramCount > 0);
@@ -157,7 +169,7 @@ namespace UE
 
 			FString ToString() const
 			{
-				return FString::Format(TEXT("{0}, {1}, {2}, {3}"), { Version, bCacheUsesCompressedBinaries, *BinaryCacheGuid.ToString(), ProgramCount });
+				return FString::Format(TEXT("{0}, {1}, {2}, {3}, {4}"), { Version, bCacheUsesCompressedBinaries, *BinaryCacheGuid.ToString(), CapabilitiesCRC, ProgramCount });
 			}
 		};
 	}
diff --git a/Engine/Source/Runtime/OpenGLDrv/Private/OpenGLShaders.cpp b/Engine/Source/Runtime/OpenGLDrv/Private/OpenGLShaders.cpp
--- a/Engine/Source/Runtime/OpenGLDrv/Private/OpenGLShaders.cpp
+++ b/Engine/Source/Runtime/OpenGLDrv/Private/OpenGLShaders.cpp
@@ -112,6 +112,17 @@ static FAutoConsoleVariableRef CVarGCacheAllProgramBinaries(
 );
 
 #if PLATFORM_ANDROID
+
+static bool GbInjectFragmentControlFlowHint = false;
+static FAutoConsoleVariableRef CVarInjectFragmentControlFlowHint(
+	TEXT("android.OpenGL.InjectFragmentControlFlowHint"),
+	GbInjectFragmentControlFlowHint,
+	TEXT("true - Injects a benign control flow branch into the fragment shader to trigger alternative code generation with some compilers.\n")
+	TEXT("false - fragment shaders are unchanged. (default)")
+	,
+	ECVF_ReadOnly
+);
+
 bool GOpenGLShaderHackLastCompileSuccess = false;
 #endif
 
@@ -557,7 +568,7 @@ namespace
 		int32 FoundIndex = 0;
 		for (const ANSICHAR * FoundPointer = FCStringAnsi::Strstr(Dest.GetData(), Source);
 			nullptr != FoundPointer;
-			FoundPointer = FCStringAnsi::Strstr(Dest.GetData()+FoundIndex, Source))
+			FoundPointer = FCStringAnsi::Strstr(Dest.GetData()+FoundIndex+ReplacementLen, Source))
 		{
 			FoundIndex = FoundPointer - Dest.GetData();
 			Dest.RemoveAt(FoundIndex, SourceLen);
@@ -775,7 +786,6 @@ void OPENGLDRV_API GetCurrentOpenGLShaderDeviceCapabilities(FOpenGLShaderDeviceC
 	if (FOpenGL::IsAndroidGLESCompatibilityModeEnabled())
 	{
 		Capabilities.TargetPlatform = EOpenGLShaderTargetPlatform::OGLSTP_Android;
-		Capabilities.bSupportsShaderFramebufferFetch = GSupportsShaderFramebufferFetch;
 		Capabilities.bRequiresARMShaderFramebufferFetchDepthStencilUndef = false;
 		Capabilities.bRequiresReadOnlyBuffersWorkaround = false;
 		Capabilities.MaxVaryingVectors = FOpenGL::GetMaxVaryingVectors();
@@ -784,12 +794,12 @@ void OPENGLDRV_API GetCurrentOpenGLShaderDeviceCapabilities(FOpenGLShaderDeviceC
 
 #elif PLATFORM_ANDROID
 		Capabilities.TargetPlatform = EOpenGLShaderTargetPlatform::OGLSTP_Android;
-		Capabilities.bSupportsShaderFramebufferFetch = GSupportsShaderFramebufferFetch;
 		Capabilities.bRequiresReadOnlyBuffersWorkaround = FOpenGL::RequiresReadOnlyBuffersWorkaround();
 		Capabilities.bRequiresARMShaderFramebufferFetchDepthStencilUndef = FOpenGL::RequiresARMShaderFramebufferFetchDepthStencilUndef();
 		Capabilities.MaxVaryingVectors = FOpenGL::GetMaxVaryingVectors();
 		Capabilities.bRequiresDisabledEarlyFragmentTests = FOpenGL::RequiresDisabledEarlyFragmentTests();
 		Capabilities.bRequiresPreciseQualifierWorkaround = FOpenGL::RequiresPreciseQualifierWorkaround();
+		Capabilities.bRequiresControlFlowOptimizerWorkaround = GbInjectFragmentControlFlowHint;
 #elif PLATFORM_IOS
 	Capabilities.TargetPlatform = EOpenGLShaderTargetPlatform::OGLSTP_iOS;
 #else
@@ -863,6 +873,11 @@ void OPENGLDRV_API GLSLToDeviceCompatibleGLSL(FAnsiCharArray& GlslCodeOriginal,
 		{
 			AppendCString(GlslCode, "#define UE_MRT_FRAMEBUFFER_FETCH 1\n");
 		}
+
+		if (Capabilities.bRequiresControlFlowOptimizerWorkaround)
+		{
+			ReplaceCString(GlslCodeOriginal, "void main()\n{\n", "void main()\n{\nif (gl_FragCoord.x < 0.0) return;\n");
+		}
 	}
 
 	if (bEmitTextureExternal)
diff --git a/Engine/Source/Runtime/OpenGLDrv/Private/OpenGLShaders.h b/Engine/Source/Runtime/OpenGLDrv/Private/OpenGLShaders.h
--- a/Engine/Source/Runtime/OpenGLDrv/Private/OpenGLShaders.h
+++ b/Engine/Source/Runtime/OpenGLDrv/Private/OpenGLShaders.h
@@ -24,12 +24,12 @@ struct FOpenGLShaderDeviceCapabilities
 {
 	EOpenGLShaderTargetPlatform TargetPlatform;
 	EShaderPlatform MaxRHIShaderPlatform;
-	bool bSupportsShaderFramebufferFetch;
 	bool bRequiresARMShaderFramebufferFetchDepthStencilUndef;
 	GLint MaxVaryingVectors;
 	bool bRequiresDisabledEarlyFragmentTests;
 	bool bRequiresReadOnlyBuffersWorkaround;
 	bool bRequiresPreciseQualifierWorkaround;
+	bool bRequiresControlFlowOptimizerWorkaround;
 };
 
 /**
