Ensure gles binary cache archive is flushed before writing the progress header.
- logging improvement
- add some handling code for unlikely createfilereader failures.

#Synced-CL 51487278

diff --git a/Engine/Source/Runtime/OpenGLDrv/Private/OpenGLProgramBinaryFileCache.cpp b/Engine/Source/Runtime/OpenGLDrv/Private/OpenGLProgramBinaryFileCache.cpp
--- a/Engine/Source/Runtime/OpenGLDrv/Private/OpenGLProgramBinaryFileCache.cpp
+++ b/Engine/Source/Runtime/OpenGLDrv/Private/OpenGLProgramBinaryFileCache.cpp
@@ -522,8 +522,13 @@ void FOpenGLProgramBinaryCache::UpdatePrecacheMapping()
 	FScopeLock Lock(&GProgramBinaryFileCacheCS);
 	FArchive& DestAr = *BinaryCacheWriteFileHandle;
 	int64 DestPos = DestAr.Tell();
+	bool bBinaryCacheFileInitiallyValid = !(DestAr.IsError() || DestAr.IsCriticalError());
 
-	if (DestPos >= CurrentShaderPipelineProperties.LastMappedPosition + (GBinaryCacheMMapAfterEveryMB * 1024 * 1024))
+	static bool bOneTime = true;
+	UE_CLOGF(!bBinaryCacheFileInitiallyValid && bOneTime, LogRHI, Error, "UpdatePrecacheMapping : archive failed (%d, %d)", DestAr.IsError(), DestAr.IsCriticalError());
+	bOneTime = bOneTime && bBinaryCacheFileInitiallyValid;
+
+	if (bBinaryCacheFileInitiallyValid && DestPos >= CurrentShaderPipelineProperties.LastMappedPosition + (GBinaryCacheMMapAfterEveryMB * 1024 * 1024))
 	{
 		int32 ProgramCount = ProgramsInCurrentCache.Num();
 		MarkValidContent(ProgramCount);
@@ -531,10 +536,26 @@ void FOpenGLProgramBinaryCache::UpdatePrecacheMapping()
 		UE_LOG(LogRHI, Log, TEXT("UpdatePrecacheMapping : Mapping program cache from %" INT64_FMT " to %" INT64_FMT " "), CurrentShaderPipelineProperties.LastMappedPosition, DestPos);
 
 		TUniquePtr<FArchive> BinaryProgramReader = TUniquePtr<FArchive>(IFileManager::Get().CreateFileReader(*GetProgramBinaryCacheFilePath()));
- 		BinaryProgramReader->Seek(CurrentShaderPipelineProperties.LastMappedPosition);
-		int32 ProgramsToRead = ProgramCount - CurrentShaderPipelineProperties.MappedPrograms;
- 		bool bSuccess = ReadProgramFile_Internal(ProgramsToRead, DestPos, GetProgramBinaryCacheFilePath(), *BinaryProgramReader);
-		CurrentShaderPipelineProperties.MappedPrograms = ProgramCount;
+		bool bSuccess = BinaryProgramReader.IsValid();
+		if(bSuccess)
+		{
+			BinaryProgramReader->Seek(CurrentShaderPipelineProperties.LastMappedPosition);
+			int32 ProgramsToRead = ProgramCount - CurrentShaderPipelineProperties.MappedPrograms;
+			bSuccess = ReadProgramFile_Internal(ProgramsToRead, DestPos, GetProgramBinaryCacheFilePath(), *BinaryProgramReader);
+			if(bSuccess)
+			{
+				CurrentShaderPipelineProperties.MappedPrograms = ProgramCount;
+			}
+			else
+			{
+				UE_LOGF(LogRHI, Log, "UpdatePrecacheMapping : failed to re-read from program cache.");
+			}
+		}
+		else
+		{
+			UE_LOGF(LogRHI, Log, "UpdatePrecacheMapping : Mapping failed to open %ls", *GetProgramBinaryCacheFilePath());
+		}
+
 		check(bSuccess);
 	}
 }
@@ -862,6 +883,7 @@ bool FOpenGLProgramBinaryCache::MarkValidContent(int32 NumPrograms)
 	if (FailCode == 0)
 	{
 		FArchive& Ar = *BinaryCacheWriteFileHandle;
+		Ar.Flush();
 		int64 CurrPos = Ar.Tell();
 		Ar.Seek(0);
 		UE::OpenGL::FBinaryCacheFileHeader OutHeader = UE::OpenGL::FBinaryCacheFileHeader::CreateHeader(CurrentShaderPipelineProperties.CacheVersionGuid, NumPrograms, (int32)CurrPos);
@@ -947,7 +969,7 @@ void FOpenGLProgramBinaryCache::AddProgramBinaryDataToBinaryCache(const FOpenGLP
 	check(IsBuildingCache_internal());
 	check(BinaryProgramData.IsValid());
 	FArchive& Ar = *BinaryCacheWriteFileHandle;
-	bool bInitiallyValid = !(Ar.IsError() || Ar.IsCriticalError());
+	const bool bBinaryCacheFileInitiallyValid = !(Ar.IsError() || Ar.IsCriticalError());
 	// Serialize to output file:
 	FOpenGLProgramKey SerializedProgramKey = ProgramKey;
 	const TArrayView<const uint8> BinaryProgramDataView = BinaryProgramData.GetDataView();
@@ -962,8 +984,9 @@ void FOpenGLProgramBinaryCache::AddProgramBinaryDataToBinaryCache(const FOpenGLP
 	Ar.Serialize(const_cast<uint8*>(ProgramBinaryBytes), ProgramBinarySize);
 	uint32 End = Ar.Tell();
 
-	UE_CLOG(Ar.IsError() || Ar.IsCriticalError(), LogRHI, Error, TEXT("AddProgramBinaryDataToBinaryCache : archive failed (%d, %d, %d, %d, %d, %d)"), Ar.IsError(), Ar.IsCriticalError(), Start, End, ProgramBinarySize, bInitiallyValid);
-	UE_LOG(LogRHI, VeryVerbose, TEXT("AddProgramBinaryDataToBinaryCache : added %d bytes to cache"), End-Start);
+	const bool bBinaryCacheValidAfterwards = !(Ar.IsError() || Ar.IsCriticalError());
+	UE_CLOGF(!bBinaryCacheValidAfterwards, LogRHI, Error, "AddProgramBinaryDataToBinaryCache : archive failed (%d, %d, %d, %d, %d, %d, %d)", Ar.IsError(), Ar.IsCriticalError(), Start, End, ProgramBinarySize, bBinaryCacheFileInitiallyValid, FPlatformMisc::GetLastError());
+	UE_LOGF(LogRHI, VeryVerbose, "AddProgramBinaryDataToBinaryCache : added %d bytes to cache", End-Start);
 
 	if (UE::OpenGL::AreBinaryProgramsCompressed())
 	{
@@ -974,13 +997,13 @@ void FOpenGLProgramBinaryCache::AddProgramBinaryDataToBinaryCache(const FOpenGLP
 		TotalUncompressed += Header->UncompressedSize;
 		TotalCompressed += ProgramBinarySize;
 
-		UE_LOG(LogRHI, Verbose, TEXT("AddProgramBinaryDataToBinaryCache: total Uncompressed: %d, total Compressed %d, Total saved so far: %d"), TotalUncompressed, TotalCompressed, TotalUncompressed - TotalCompressed);
+		UE_LOGF(LogRHI, Verbose, "AddProgramBinaryDataToBinaryCache: total Uncompressed: %d, total Compressed %d, Total saved so far: %d", TotalUncompressed, TotalCompressed, TotalUncompressed - TotalCompressed);
 	}
-	UE_LOG(LogRHI, VeryVerbose, TEXT("AddProgramBinaryDataToBinaryCache: written Program %s to cache (%d bytes)"), *ProgramKey.ToString(), BinaryProgramDataView.Num());
+	UE_LOGF(LogRHI, VeryVerbose, "AddProgramBinaryDataToBinaryCache: written Program %ls to cache (%d bytes)", *ProgramKey.ToString(), BinaryProgramDataView.Num());
 
 	ProgramsInCurrentCache.Add(ProgramKey);
 
- 	if (IsPrecachingEnabled())
+ 	if (IsPrecachingEnabled() && bBinaryCacheValidAfterwards)
  	{
 		int32 NumProgramsStored = ProgramsInCurrentCache.Num();
 		if (NumProgramsStored >= CurrentShaderPipelineProperties.NumProgramsFlushed + GBinaryCachePeriodicFlushProgramCount)
