Update Android GLES to handle more than one surface

#rnx

#Synced-CL 0

diff --git a/Engine/Source/Runtime/ApplicationCore/Public/Android/AndroidPlatformFramePacer.h b/Engine/Source/Runtime/ApplicationCore/Public/Android/AndroidPlatformFramePacer.h
--- a/Engine/Source/Runtime/ApplicationCore/Public/Android/AndroidPlatformFramePacer.h
+++ b/Engine/Source/Runtime/ApplicationCore/Public/Android/AndroidPlatformFramePacer.h
@@ -16,7 +16,7 @@ struct IAndroidFramePacer
 {
 	virtual void Init() = 0;
 	virtual ~IAndroidFramePacer() = default;
-	virtual bool SwapBuffers(bool bLockToVsync) { return true; }
+	virtual bool SwapBuffers(struct FPlatformOpenGLSurfaceContext* SurfaceContext, bool bLockToVsync) { return true; }
 	virtual bool SupportsFramePace(int32 QueryFramePace) = 0;
 };
 
@@ -24,7 +24,7 @@ struct FAndroidOpenGLFramePacer : public IAndroidFramePacer
 {
 	virtual void Init() override;
 	virtual ~FAndroidOpenGLFramePacer();
-	virtual bool SwapBuffers(bool bLockToVsync) override;
+	virtual bool SwapBuffers(struct FPlatformOpenGLSurfaceContext* SurfaceContext, bool bLockToVsync) override;
 	virtual bool SupportsFramePace(int32 QueryFramePace) override;
 
 private:
@@ -70,7 +70,7 @@ struct FAndroidPlatformRHIFramePacer : public FGenericPlatformRHIFramePacer
 	// FAndroidPlatformRHIFramePacer interface
 	static void Init(IAndroidFramePacer* InFramePacer);
 	static int32 GetLegacySyncInterval();
-	static void SwapBuffers(bool bLockToVsync) { if (ensure(FramePacer)) { FramePacer->SwapBuffers(bLockToVsync); } }
+	static void SwapBuffers(struct FPlatformOpenGLSurfaceContext* SurfaceContext, bool bLockToVsync) { if (ensure(FramePacer)) { FramePacer->SwapBuffers(SurfaceContext, bLockToVsync); } }
 
 	static TAutoConsoleVariable<int32> CVarUseSwappyForFramePacing;
 	static TAutoConsoleVariable<int32> CVarSwappyGPUFrameTimeFence;
diff --git a/Engine/Source/Runtime/Launch/Private/Android/AndroidEventManager.cpp b/Engine/Source/Runtime/Launch/Private/Android/AndroidEventManager.cpp
--- a/Engine/Source/Runtime/Launch/Private/Android/AndroidEventManager.cpp
+++ b/Engine/Source/Runtime/Launch/Private/Android/AndroidEventManager.cpp
@@ -6,12 +6,13 @@
 #include "Android/AndroidApplication.h"
 #include "AudioDevice.h"
 #include "Misc/CallbackDevice.h"
-#include <android/native_window.h> 
-#include <android/native_window_jni.h> 
+#include <android/native_window.h>
+#include <android/native_window_jni.h>
 #include "IHeadMountedDisplay.h"
 #include "IXRTrackingSystem.h"
 #include "RenderingThread.h"
 #include "UnrealEngine.h"
+#include "Android/AndroidDynamicRHI.h"
 
 DEFINE_LOG_CATEGORY(LogAndroidEvents);
 
@@ -345,6 +346,16 @@ void FAppEventManager::ExecWindowCreated()
 	TOptional<FAndroidWindow::FNativeAccessor> MainWindowAccessor = FAndroidWindowManager::Get().GetMainWindowAsNativeAccessor(false);
 	FAndroidAppEntry::ReInitWindow(MainWindowAccessor);
 	FAndroidApplication::OnWindowSizeChanged();
+
+	if (FAndroidMisc::UseNewWindowBehavior())
+	{
+		// Unblock rendering now that viewport has been resized to match the new surface
+		const auto& OnWindowCreatedCallback = FPlatformDynamicRHI::GetRHIOnWindowCreatedCallback();
+		if (OnWindowCreatedCallback)
+		{
+			OnWindowCreatedCallback();
+		}
+	}
 }
 
 void FAppEventManager::ExecWindowResized()
diff --git a/Engine/Source/Runtime/Launch/Private/Android/AndroidJNI.cpp b/Engine/Source/Runtime/Launch/Private/Android/AndroidJNI.cpp
--- a/Engine/Source/Runtime/Launch/Private/Android/AndroidJNI.cpp
+++ b/Engine/Source/Runtime/Launch/Private/Android/AndroidJNI.cpp
@@ -1041,6 +1041,11 @@ bool AndroidThunkCpp_IsOpenXRImmersiveHMDApplication()
 	return IsOpenXRImmersiveHMDApplication == 1;
 }
 
+bool AndroidThunkCpp_IsStandaloneApplication()
+{
+	return USE_ANDROID_STANDALONE != 0;
+}
+
 // call out to JNI to see if the application was packaged for Oculus Mobile
 bool AndroidThunkCpp_IsOculusMobileApplication()
 {
diff --git a/Engine/Source/Runtime/OpenGLDrv/Private/Android/AndroidEGL.cpp b/Engine/Source/Runtime/OpenGLDrv/Private/Android/AndroidEGL.cpp
--- a/Engine/Source/Runtime/OpenGLDrv/Private/Android/AndroidEGL.cpp
+++ b/Engine/Source/Runtime/OpenGLDrv/Private/Android/AndroidEGL.cpp
@@ -50,70 +50,106 @@ const int EGLMinStencilBits	  = 8; // This is required for UMG clipping
 const int EGLMinSampleBuffers = 0;
 const int EGLMinSampleSamples = 0;
 
-struct EGLConfigParms
+
+FPlatformOpenGLSurfaceContext::FPlatformOpenGLSurfaceContext(void* InWindowHandle) : FPlatformOpenGLSurfaceContext()
+{
+	// the context should be locked at the beginning of the update process, this overload is used for GT initiated events.
+	check(IsInGameThread());
+
+	AndroidWindow = (FAndroidWindow*)InWindowHandle;
+	if (FAndroidMisc::UseNewWindowBehavior())
+	{		
+		if (AndroidWindow) // AndroidWindow can be null in the cases where offscreen rendering is expected.
+		{
+			FAndroidWindow::FNativeAccessor WindowContainer = AndroidWindow->GetANativeAccessor(false);
+			// Note: the window handle is only locked within this scope.
+			Window = WindowContainer.GetANativeWindow();
+			// We continue to use the native handle as it will only be changed via the android event thread and enqueued to the RHIT
+			// Only the android event thread and rhit will use the handle from this point.
+		}
+	}
+}
+
+struct AndroidESPImpl
 {
-	/** Whether this is a valid configuration or not */
-	int validConfig = 0;
-	/** The number of bits requested for the red component */
-	int redSize = 8;
-	/** The number of bits requested for the green component */
-	int greenSize = 8;
-	/** The number of bits requested for the blue component */
-	int blueSize = 8;
-	/** The number of bits requested for the alpha component */
-	int alphaSize = 0;
-	/** The number of bits requested for the depth component */
-	int depthSize = 24;
-	/** The number of bits requested for the stencil component */
-	int stencilSize = 0;
-	/** The number of multisample buffers requested */
-	int sampleBuffers = 0;
-	/** The number of samples requested */
-	int sampleSamples = 0;
-
-	EGLConfigParms()
+	bool Initalized = false;
+
+	// A single GL context is used for all surfaces.
+	FPlatformOpenGLContext RenderingContext;
+	FAndroidEGLDisplayInfo EGLDisplay;
+
+	// SurfacesMutex protects GLSurfaces, but not the contained instances.
+	UE::FMutex SurfacesMutex;
+	TArray<TUniquePtr<FPlatformOpenGLSurfaceContext>> GLSurfaces;
+
+	FPlatformOpenGLSurfaceContext* GetSurfaceFromAndroidWindow(const FAndroidWindow* FindMe)
+	{
+		UE::TUniqueLock ScopeLock(SurfacesMutex);
+		check(!GLSurfaces.IsEmpty());
+		TUniquePtr<FPlatformOpenGLSurfaceContext>* Found = GLSurfaces.FindByPredicate([FindMe](const TUniquePtr<FPlatformOpenGLSurfaceContext>& Item) -> bool
+			{
+				return Item->GetAndroidWindow() == FindMe;
+			});
+		return Found ? Found->Get() : nullptr;
+	}
+
+	FPlatformOpenGLSurfaceContext* GetGLSurfaceFromWindowContainer(const TOptional<FAndroidWindow::FNativeAccessor>& WindowContainer)
 	{
-		// If not default, set the preference
-		int DepthBufferPreference = (int)FAndroidWindow::GetDepthBufferPreference();
-		if (DepthBufferPreference > 0)
+		VERIFY_GL_SCOPE();
+		FPlatformOpenGLSurfaceContext* Surface = nullptr;
+		if (!FAndroidMisc::UseNewWindowBehavior())
+		{
+			Surface = GetSurfaceFromAndroidWindow(nullptr);
+		}
+		else if (WindowContainer.IsSet())
 		{
-			depthSize = DepthBufferPreference;
+			Surface = GetSurfaceFromAndroidWindow(WindowContainer->Get());
 		}
 
-		if (FAndroidMisc::GetMobilePropagateAlphaSetting() > 0)
+		if (!Surface)
 		{
-			alphaSize = 8;
+			UE_LOG(LogRHI, Warning, TEXT("GetGLSurfaceFromWindowContainer: surface context not found for the given window"));
 		}
+		return Surface;
 	}
 };
 
+FPlatformOpenGLSurfaceContext* AndroidEGL::FindSurfaceContextFromWindow(FAndroidWindow* InWindowHandle)
+{
+	UE::TUniqueLock ScopeLock(PImplData->SurfacesMutex);
 
-struct AndroidESPImpl
+	FAndroidWindow* FindMeAndroidWindow = (FAndroidWindow*)InWindowHandle;
+	for (TUniquePtr<FPlatformOpenGLSurfaceContext>& Surface : PImplData->GLSurfaces)
+	{
+		if (Surface->GetAndroidWindow() == InWindowHandle)
+		{
+			return Surface.Get();
+		}
+	}
+	return nullptr;
+}
+
+FPlatformOpenGLSurfaceContext* AndroidEGL::CreateSurfaceContextFromWindow(FAndroidWindow* InWindowHandle)
 {
-	FPlatformOpenGLContext	RenderingContext;
-
-	EGLDisplay     eglDisplay         = EGL_NO_DISPLAY;
-	EGLint         eglNumConfigs      = 0;
-	EGLint         eglFormat          = -1;
-	EGLConfig      eglConfigParam     = nullptr;
-	EGLSurface     eglSurface         = EGL_NO_SURFACE;
-	EGLSurface     auxSurface         = EGL_NO_SURFACE;
-	EGLint         eglWidth           = 8; // required for Gear VR apps with internal win surf mgmt
-	EGLint         eglHeight          = 8; // required for Gear VR apps with internal win surf mgmt
-	EGLint         NativeVisualID     = 0;
-	EGLConfigParms Parms;
-	float          eglRatio           = 0;
-	int            DepthSize          = 0;
-	ANativeWindow* Window             = nullptr;
-	GLuint         ResolveFrameBuffer = 0;
-	GLuint         DummyFrameBuffer   = 0;
-	FPlatformRect  CachedWindowRect     {};
-	bool           Initalized         = false;
-	bool           bIsDebug           = false;
-	bool           bIsWndSurface      = false; // True when the surface is attached to a HW window.
-
-	AndroidESPImpl() = default;
-};
+	check(FindSurfaceContextFromWindow(InWindowHandle) == nullptr);
+	{
+		UE::TUniqueLock ScopeLock(PImplData->SurfacesMutex);
+		PImplData->GLSurfaces.Add(MakeUnique<FPlatformOpenGLSurfaceContext>(InWindowHandle));
+	}
+	return FindSurfaceContextFromWindow(InWindowHandle);
+}
+
+void AndroidEGL::DestroySurfaceContext(FPlatformOpenGLSurfaceContext* SurfaceContext)
+{
+	UE::TUniqueLock ScopeLock(PImplData->SurfacesMutex);
+	DestroyBackBuffer(SurfaceContext);
+	DestroyRenderSurface(SurfaceContext);
+
+	PImplData->GLSurfaces.RemoveAll([SurfaceContext](const TUniquePtr<FPlatformOpenGLSurfaceContext>& FindSurface)
+		{
+			return FindSurface.Get() == SurfaceContext;
+		});
+}
 
 const EGLint Attributes[] = {
 	EGL_RED_SIZE,       EGLMinRedBits,
@@ -140,86 +176,120 @@ AndroidEGL::AndroidEGL()
 		ANativeWindow_setBuffersTransform_API = reinterpret_cast<PFN_ANativeWindow_setBuffersTransform>(dlsym(LibNativeWindow, "ANativeWindow_setBuffersTransform"));
 	}
 	FPlatformMisc::LowLevelOutputDebugStringf(TEXT("ANativeWindow_setBuffersTransform is %s on this device"), ANativeWindow_setBuffersTransform_API == nullptr ? TEXT("not supported") : TEXT("supported"));
+
+	// Create a special case null surface context.
+	// this is used for all windows when !FAndroidMisc::UseNewWindowBehavior()
+	// or off screen rendering (e.g. when FAndroidMisc::UseNewWindowBehavior() but no window exists).
+	AndroidEGL::CreateSurfaceContextFromWindow(nullptr);
 }
 
 void AndroidEGL::ResetDisplay()
 {
 	VERIFY_EGL_SCOPE();
-	if(PImplData->eglDisplay != EGL_NO_DISPLAY)
+	if(PImplData->EGLDisplay.eglDisplay != EGL_NO_DISPLAY)
 	{
-		FPlatformMisc::LowLevelOutputDebugStringf( TEXT("AndroidEGL::ResetDisplay()" ));
-		eglMakeCurrent(PImplData->eglDisplay, EGL_NO_SURFACE, EGL_NO_SURFACE, EGL_NO_CONTEXT);
+		FPlatformMisc::LowLevelOutputDebugStringf( TEXT("AndroidEGL::ResetDisplay(%p) " ), PImplData->EGLDisplay.eglDisplay);
+
+#if USE_ANDROID_STANDALONE
+		if (FAndroidMisc::UseNewWindowBehavior())
+		{
+			// This prevents flickering when attaching/detaching extra views in multi-viewport scenarios
+			VERIFY_GL_SCOPE();
+			glBindFramebuffer(GL_FRAMEBUFFER, 0);
+			glBindFramebuffer(GL_READ_FRAMEBUFFER, 0);
+			glBindFramebuffer(GL_DRAW_FRAMEBUFFER, 0);
+		}
+#endif
+
+		if (eglMakeCurrent(PImplData->EGLDisplay.eglDisplay, EGL_NO_SURFACE, EGL_NO_SURFACE, EGL_NO_CONTEXT) == EGL_FALSE)
+		{
+			UE_LOG(LogRHI, Warning, TEXT("eglMakeCurrent failed in ResetDisplay: 0x%x"), eglGetError());
+		}
 	}
 }
 
-void AndroidEGL::DestroyRenderSurface()
+void AndroidEGL::DestroyRenderSurface(FPlatformOpenGLSurfaceContext* Surface)
 {
 	VERIFY_EGL_SCOPE();
-	FPlatformMisc::LowLevelOutputDebugStringf( TEXT("AndroidEGL::DestroyRenderSurface()" ));
-	if( PImplData->eglSurface != EGL_NO_SURFACE )
+	FPlatformMisc::LowLevelOutputDebugStringf(TEXT("AndroidEGL::DestroyRenderSurface(%p)"), Surface);
+	if (Surface->eglSurface != EGL_NO_SURFACE)
 	{
-		eglDestroySurface(PImplData->eglDisplay, PImplData->eglSurface);
-		PImplData->eglSurface = EGL_NO_SURFACE;
+		eglDestroySurface(PImplData->EGLDisplay.eglDisplay, Surface->eglSurface);
+		Surface->eglSurface = EGL_NO_SURFACE;
 	}
-
-	PImplData->RenderingContext.eglSurface = EGL_NO_SURFACE;
 }
 
 void AndroidEGL::TerminateEGL()
 {
 	VERIFY_EGL_SCOPE();
 
-	eglTerminate(PImplData->eglDisplay);
-	PImplData->eglDisplay = EGL_NO_DISPLAY;
+	if (PImplData->EGLDisplay.auxSurface != EGL_NO_SURFACE)
+	{
+		eglDestroySurface(PImplData->EGLDisplay.eglDisplay, PImplData->EGLDisplay.auxSurface);
+		PImplData->EGLDisplay.auxSurface = EGL_NO_SURFACE;
+	}
+
+	if (PImplData->EGLDisplay.eglDisplay != EGL_NO_DISPLAY)
+	{
+		eglTerminate(PImplData->EGLDisplay.eglDisplay);
+		PImplData->EGLDisplay.eglDisplay = EGL_NO_DISPLAY;
+	}
+
 	PImplData->Initalized = false;
 }
 
 /* Can be called from any thread */
-EGLBoolean AndroidEGL::SetCurrentContext(EGLContext InContext, EGLSurface InSurface)
+EGLBoolean AndroidEGL::SetCurrentContext(EGLContext InContext, FPlatformOpenGLSurfaceContext* InSurfaceInfo)
 {
 	VERIFY_EGL_SCOPE();
 	//context can be null.so can surface from PlatformNULLContextSetup
-	EGLBoolean Result = EGL_FALSE;
+	EGLBoolean Result = EGL_TRUE;
 	EGLContext CurrentContext = GetCurrentContext();
+	EGLSurface CurrentSurface = GetCurrentSurface(EGL_DRAW);
+	check(GetCurrentSurface(EGL_DRAW) == GetCurrentSurface(EGL_READ));
 
+	EGLSurface WantedSurface = InSurfaceInfo ? InSurfaceInfo->eglSurface : EGL_NO_SURFACE;
 	// activate the context
-	if (CurrentContext != InContext)
+	if (CurrentContext != InContext || CurrentSurface != WantedSurface)
 	{
 		if (CurrentContext !=EGL_NO_CONTEXT)
 		{
 			glFlush();
 		}
 
-		if (InContext == EGL_NO_CONTEXT && InSurface == EGL_NO_SURFACE)
+		if (InContext == EGL_NO_CONTEXT && InSurfaceInfo == nullptr)
 		{
-			ResetDisplay();
+			// disconnect context from current thread
+			Result = eglMakeCurrent(PImplData->EGLDisplay.eglDisplay, EGL_NO_SURFACE, EGL_NO_SURFACE, EGL_NO_CONTEXT);
+			checkf(Result == EGL_TRUE, TEXT("ERROR: SetCurrentContext eglMakeCurrent EGL_NO_CONTEXT failed : 0x%x"), eglGetError());
 		}
 		else
 		{
 			//if we have a valid context, and no surface then create a tiny pbuffer and use that temporarily
-			EGLSurface Surface = InSurface;
-			if (!bSupportsKHRSurfacelessContext && InContext != EGL_NO_CONTEXT && InSurface == EGL_NO_SURFACE)
+			EGLSurface Surface = WantedSurface;
+			if (!bSupportsKHRSurfacelessContext && InContext != EGL_NO_CONTEXT && Surface == EGL_NO_SURFACE)
 			{
-				checkf(PImplData->auxSurface == EGL_NO_SURFACE, TEXT("ERROR: PImplData->auxSurface already in use. PBuffer surface leak!"));
-				EGLint PBufferAttribs[] =
-				{
-					EGL_WIDTH, 1,
-					EGL_HEIGHT, 1,
-					EGL_TEXTURE_TARGET, EGL_NO_TEXTURE,
-					EGL_TEXTURE_FORMAT, EGL_NO_TEXTURE,
-					EGL_NONE
-				};
-
-				PImplData->auxSurface = eglCreatePbufferSurface(PImplData->eglDisplay, PImplData->eglConfigParam, PBufferAttribs);
-
-				if (PImplData->auxSurface == EGL_NO_SURFACE)
+				// Reuse existing auxSurface if available, only create if needed
+				if (PImplData->EGLDisplay.auxSurface == EGL_NO_SURFACE)
 				{
-					checkf(PImplData->auxSurface != EGL_NO_SURFACE, TEXT("eglCreatePbufferSurface error : 0x%x"), eglGetError());
+					EGLint PBufferAttribs[] =
+					{
+						EGL_WIDTH, 1,
+						EGL_HEIGHT, 1,
+						EGL_TEXTURE_TARGET, EGL_NO_TEXTURE,
+						EGL_TEXTURE_FORMAT, EGL_NO_TEXTURE,
+						EGL_NONE
+					};
+
+					PImplData->EGLDisplay.auxSurface = eglCreatePbufferSurface(PImplData->EGLDisplay.eglDisplay, PImplData->EGLDisplay.eglConfigParam, PBufferAttribs);
+
+					checkf(PImplData->EGLDisplay.auxSurface != EGL_NO_SURFACE, TEXT("eglCreatePbufferSurface error : 0x%x"), eglGetError());
 				}
-				Surface = PImplData->auxSurface;
+				Surface = PImplData->EGLDisplay.auxSurface;
 			}
 
-			Result = eglMakeCurrent(PImplData->eglDisplay, Surface, Surface, InContext);
+		//	FPlatformMisc::LowLevelOutputDebugStringf(TEXT("AndroidEGL::SetCurrentContext() %p, %p, %p, %p"), PImplData->EGLDisplay.eglDisplay, Surface, Surface, InContext);
+			Result = eglMakeCurrent(PImplData->EGLDisplay.eglDisplay, Surface, Surface, InContext);
 			checkf(Result == EGL_TRUE, TEXT("ERROR: SetCurrentContext eglMakeCurrent failed : 0x%x"), eglGetError());
 		}
 	}
@@ -231,42 +301,44 @@ void AndroidEGL::ResetInternal()
 	Terminate();
 }
 
-void AndroidEGL::CreateEGLRenderSurface(ANativeWindow* InWindow, bool bCreateWndSurface)
+void AndroidEGL::CreateEGLRenderSurface(FPlatformOpenGLSurfaceContext* Surface, bool bCreateWndSurface)
 {
 	VERIFY_EGL_SCOPE();
+	check(Surface);
 
 	// due to possible early initialization, don't redo this
-	if (PImplData->eglSurface != EGL_NO_SURFACE)
+	if (Surface->eglSurface != EGL_NO_SURFACE)
 	{
-		FPlatformMisc::LowLevelOutputDebugStringf(TEXT("AndroidEGL::CreateEGLRenderSurface() Already initialized: %p"), PImplData->eglSurface);
+		FPlatformMisc::LowLevelOutputDebugStringf(TEXT("AndroidEGL::CreateEGLRenderSurface() Already initialized: window=%p, %p"), Surface->Window, Surface->eglSurface);
 		return;
 	}
 
-	PImplData->bIsWndSurface = bCreateWndSurface;
+	Surface->bIsWndSurface = bCreateWndSurface;
 
 	if (bCreateWndSurface)
 	{
-		check(InWindow);
+		check(Surface->Window);
 		//need ANativeWindow
-		PImplData->eglSurface = eglCreateWindowSurface(PImplData->eglDisplay, PImplData->eglConfigParam,InWindow, NULL);
+		Surface->eglSurface = eglCreateWindowSurface(PImplData->EGLDisplay.eglDisplay, PImplData->EGLDisplay.eglConfigParam, Surface->Window, NULL);
+		FPlatformMisc::LowLevelOutputDebugStringf(TEXT("AndroidEGL::CreateEGLRenderSurface() eglCreateWindowSurface: window=%p, %p"), Surface->Window, Surface->eglSurface);
 
 		if (FAndroidPlatformRHIFramePacer::CVarAllowFrameTimestamps.GetValueOnAnyThread())
 		{
-			STANDALONE_DEBUG_LOGf(LogAndroid, TEXT("AndroidEGL::CreateEGLRenderSurface(InWindow = %p) using a.allowFrameTimestamps enable EGL_TIMESTAMPS_ANDROID on %p"), InWindow, PImplData->eglSurface);
-			eglSurfaceAttrib(PImplData->eglDisplay, PImplData->eglSurface, EGL_TIMESTAMPS_ANDROID, EGL_TRUE);
+			STANDALONE_DEBUG_LOGf(LogAndroid, TEXT("AndroidEGL::CreateEGLRenderSurface(native window = %p) using a.allowFrameTimestamps enable EGL_TIMESTAMPS_ANDROID on %p"), Surface->Window, Surface->eglSurface);
+			eglSurfaceAttrib(PImplData->EGLDisplay.eglDisplay, Surface->eglSurface, EGL_TIMESTAMPS_ANDROID, EGL_TRUE);
 		}
 		else
 		{
 			// HAD to add the false condition so that android attributes reflect current state of CVar.
-			STANDALONE_DEBUG_LOGf(LogAndroid, TEXT("AndroidEGL::CreateEGLRenderSurface(InWindow = %p) using a.allowFrameTimestamps disable EGL_TIMESTAMPS_ANDROID on %p"), InWindow, PImplData->eglSurface);
-			eglSurfaceAttrib(PImplData->eglDisplay, PImplData->eglSurface, EGL_TIMESTAMPS_ANDROID, EGL_FALSE);
+			STANDALONE_DEBUG_LOGf(LogAndroid, TEXT("AndroidEGL::CreateEGLRenderSurface(native Window = %p) using a.allowFrameTimestamps disable EGL_TIMESTAMPS_ANDROID on %p"), Surface->Window, Surface->eglSurface);
+			eglSurfaceAttrib(PImplData->EGLDisplay.eglDisplay, Surface->eglSurface, EGL_TIMESTAMPS_ANDROID, EGL_FALSE);
 		}
 
-		STANDALONE_DEBUG_LOGf(LogAndroid, TEXT("AndroidEGL::CreateEGLRenderSurface() %p" ), PImplData->eglSurface);
+		STANDALONE_DEBUG_LOGf(LogAndroid, TEXT("AndroidEGL::CreateEGLRenderSurface() window=%p %p" ), Surface->Window, Surface->eglSurface);
 
-		if (PImplData->eglSurface == EGL_NO_SURFACE)
+		if (Surface->eglSurface == EGL_NO_SURFACE)
 		{
-			checkf(PImplData->eglSurface != EGL_NO_SURFACE, TEXT("eglCreateWindowSurface error : 0x%x"), eglGetError());
+			checkf(Surface->eglSurface != EGL_NO_SURFACE, TEXT("eglCreateWindowSurface error : 0x%x"), eglGetError());
 			ResetInternal();
 		}
 
@@ -286,12 +358,12 @@ void AndroidEGL::CreateEGLRenderSurface(ANativeWindow* InWindow, bool bCreateWnd
 		initEGL(eglAttemptedParams);
 
 		// try again
-		eglSurface = eglCreateWindowSurface(PImplData->eglDisplay, eglConfig, surface, null);
+		eglSurface = eglCreateWindowSurface(PImplData->EGLDisplay.eglDisplay, eglConfig, surface, null);
 		}
 
 		*/
 		EGLBoolean result = EGL_FALSE;
-		if (!( result =  ( eglQuerySurface(PImplData->eglDisplay, PImplData->eglSurface, EGL_WIDTH, &PImplData->eglWidth) && eglQuerySurface(PImplData->eglDisplay, PImplData->eglSurface, EGL_HEIGHT, &PImplData->eglHeight) ) ) )
+		if (!( result =  ( eglQuerySurface(PImplData->EGLDisplay.eglDisplay, Surface->eglSurface, EGL_WIDTH, &Surface->eglWidth) && eglQuerySurface(PImplData->EGLDisplay.eglDisplay, Surface->eglSurface, EGL_HEIGHT, &Surface->eglHeight) ) ) )
 		{
 			ResetInternal();
 		}
@@ -310,17 +382,17 @@ void AndroidEGL::CreateEGLRenderSurface(ANativeWindow* InWindow, bool bCreateWnd
 			EGL_NONE
 		};
 
-		checkf(PImplData->eglWidth != 0, TEXT("eglWidth is ZERO; could be a problem!"));
-		checkf(PImplData->eglHeight != 0, TEXT("eglHeight is ZERO; could be a problem!"));
-		pbufferAttribs[1] = PImplData->eglWidth;
-		pbufferAttribs[3] = PImplData->eglHeight;
+		checkf(Surface->eglWidth != 0, TEXT("eglWidth is ZERO; could be a problem!"));
+		checkf(Surface->eglHeight != 0, TEXT("eglHeight is ZERO; could be a problem!"));
+		pbufferAttribs[1] = Surface->eglWidth;
+		pbufferAttribs[3] = Surface->eglHeight;
 
 		FPlatformMisc::LowLevelOutputDebugStringf( TEXT("AndroidEGL::CreateEGLRenderSurface(%d), eglSurface = eglCreatePbufferSurface(), %dx%d" ),
 			int(bCreateWndSurface), pbufferAttribs[1], pbufferAttribs[3]);
-		PImplData->eglSurface = eglCreatePbufferSurface(PImplData->eglDisplay, PImplData->eglConfigParam, pbufferAttribs);
-		if(PImplData->eglSurface== EGL_NO_SURFACE )
+		Surface->eglSurface = eglCreatePbufferSurface(PImplData->EGLDisplay.eglDisplay, PImplData->EGLDisplay.eglConfigParam, pbufferAttribs);
+		if(Surface->eglSurface== EGL_NO_SURFACE )
 		{
-			checkf(PImplData->eglSurface != EGL_NO_SURFACE, TEXT("eglCreatePbufferSurface error : 0x%x"), eglGetError());
+			checkf(Surface->eglSurface != EGL_NO_SURFACE, TEXT("eglCreatePbufferSurface error : 0x%x"), eglGetError());
 			ResetInternal();
 		}
 	}
@@ -337,15 +409,15 @@ void AndroidEGL::InitEGL(APIVariant API)
 	}
 //	bAlreadyInitialized = true;
 
-	check(PImplData->eglDisplay == EGL_NO_DISPLAY)
-	PImplData->eglDisplay = eglGetDisplay(EGL_DEFAULT_DISPLAY);
-	checkf(PImplData->eglDisplay, TEXT(" eglGetDisplay error : 0x%x "), eglGetError());
+	check(PImplData->EGLDisplay.eglDisplay == EGL_NO_DISPLAY)
+	PImplData->EGLDisplay.eglDisplay = eglGetDisplay(EGL_DEFAULT_DISPLAY);
+	checkf(PImplData->EGLDisplay.eglDisplay, TEXT(" eglGetDisplay error : 0x%x "), eglGetError());
 	
-	EGLBoolean  result = 	eglInitialize(PImplData->eglDisplay, 0 , 0);
+	EGLBoolean  result = 	eglInitialize(PImplData->EGLDisplay.eglDisplay, 0 , 0);
 	checkf( result == EGL_TRUE, TEXT("elgInitialize error: 0x%x "), eglGetError());
 
 	// Get the EGL Extension list to determine what is supported
-	FString Extensions = ANSI_TO_TCHAR( eglQueryString( PImplData->eglDisplay, EGL_EXTENSIONS));
+	FString Extensions = ANSI_TO_TCHAR( eglQueryString( PImplData->EGLDisplay.eglDisplay, EGL_EXTENSIONS));
 
 	UE_LOG(LogAndroid, Log, TEXT("EGL Extensions: \n%s"), *Extensions);
 
@@ -372,12 +444,12 @@ void AndroidEGL::InitEGL(APIVariant API)
 #if ENABLE_CONFIG_FILTER
 
 	EGLConfig* EGLConfigList = NULL;
-	result = eglChooseConfig(PImplData->eglDisplay, Attributes, NULL, 0, &PImplData->eglNumConfigs);
+	result = eglChooseConfig(PImplData->EGLDisplay.eglDisplay, Attributes, NULL, 0, &PImplData->EGLDisplay.eglNumConfigs);
 	if (result)
 	{
-		int NumConfigs = PImplData->eglNumConfigs;
+		int NumConfigs = PImplData->EGLDisplay.eglNumConfigs;
 		EGLConfigList = new EGLConfig[NumConfigs];
-		result = eglChooseConfig(PImplData->eglDisplay, Attributes, EGLConfigList, NumConfigs, &PImplData->eglNumConfigs);
+		result = eglChooseConfig(PImplData->EGLDisplay.eglDisplay, Attributes, EGLConfigList, NumConfigs, &PImplData->EGLDisplay.eglNumConfigs);
 	}
 	if (!result)
 	{
@@ -386,28 +458,28 @@ void AndroidEGL::InitEGL(APIVariant API)
 
 	checkf(result == EGL_TRUE , TEXT(" eglChooseConfig error: 0x%x"), eglGetError());
 
-	checkf(PImplData->eglNumConfigs != 0  ,TEXT(" eglChooseConfig num EGLConfigLists is 0 . error: 0x%x"), eglGetError());
+	checkf(PImplData->EGLDisplay.eglNumConfigs != 0  ,TEXT(" eglChooseConfig num EGLConfigLists is 0 . error: 0x%x"), eglGetError());
 
 	int ResultValue = 0 ;
 	bool haveConfig = false;
 	int64 score = LONG_MAX;
-	for (uint32_t i = 0; i < PImplData->eglNumConfigs; i++)
+	for (uint32_t i = 0; i < PImplData->EGLDisplay.eglNumConfigs; i++)
 	{
 		int64 currScore = 0;
 		int r, g, b, a, d, s, sb, sc, nvi;
-		eglGetConfigAttrib(PImplData->eglDisplay, EGLConfigList[i], EGL_RED_SIZE, &ResultValue); r = ResultValue;
-		eglGetConfigAttrib(PImplData->eglDisplay, EGLConfigList[i], EGL_GREEN_SIZE, &ResultValue); g = ResultValue;
-		eglGetConfigAttrib(PImplData->eglDisplay, EGLConfigList[i], EGL_BLUE_SIZE, &ResultValue); b = ResultValue;
-		eglGetConfigAttrib(PImplData->eglDisplay, EGLConfigList[i], EGL_ALPHA_SIZE, &ResultValue); a = ResultValue;
-		eglGetConfigAttrib(PImplData->eglDisplay, EGLConfigList[i], EGL_DEPTH_SIZE, &ResultValue); d = ResultValue;
-		eglGetConfigAttrib(PImplData->eglDisplay, EGLConfigList[i], EGL_STENCIL_SIZE, &ResultValue); s = ResultValue;
-		eglGetConfigAttrib(PImplData->eglDisplay, EGLConfigList[i], EGL_SAMPLE_BUFFERS, &ResultValue); sb = ResultValue;
-		eglGetConfigAttrib(PImplData->eglDisplay, EGLConfigList[i], EGL_SAMPLES, &ResultValue); sc = ResultValue;
+		eglGetConfigAttrib(PImplData->EGLDisplay.eglDisplay, EGLConfigList[i], EGL_RED_SIZE, &ResultValue); r = ResultValue;
+		eglGetConfigAttrib(PImplData->EGLDisplay.eglDisplay, EGLConfigList[i], EGL_GREEN_SIZE, &ResultValue); g = ResultValue;
+		eglGetConfigAttrib(PImplData->EGLDisplay.eglDisplay, EGLConfigList[i], EGL_BLUE_SIZE, &ResultValue); b = ResultValue;
+		eglGetConfigAttrib(PImplData->EGLDisplay.eglDisplay, EGLConfigList[i], EGL_ALPHA_SIZE, &ResultValue); a = ResultValue;
+		eglGetConfigAttrib(PImplData->EGLDisplay.eglDisplay, EGLConfigList[i], EGL_DEPTH_SIZE, &ResultValue); d = ResultValue;
+		eglGetConfigAttrib(PImplData->EGLDisplay.eglDisplay, EGLConfigList[i], EGL_STENCIL_SIZE, &ResultValue); s = ResultValue;
+		eglGetConfigAttrib(PImplData->EGLDisplay.eglDisplay, EGLConfigList[i], EGL_SAMPLE_BUFFERS, &ResultValue); sb = ResultValue;
+		eglGetConfigAttrib(PImplData->EGLDisplay.eglDisplay, EGLConfigList[i], EGL_SAMPLES, &ResultValue); sc = ResultValue;
 
 		// Optional, Tegra-specific non-linear depth buffer, which allows for much better
 		// effective depth range in relatively limited bit-depths (e.g. 16-bit)
 		int bNonLinearDepth = 0;
-		if (eglGetConfigAttrib(PImplData->eglDisplay, EGLConfigList[i], EGL_DEPTH_ENCODING_NV, &ResultValue))
+		if (eglGetConfigAttrib(PImplData->EGLDisplay.eglDisplay, EGLConfigList[i], EGL_DEPTH_ENCODING_NV, &ResultValue))
 		{
 			bNonLinearDepth = (ResultValue == EGL_DEPTH_ENCODING_NONLINEAR_NV) ? 1 : 0;
 		}
@@ -419,16 +491,16 @@ void AndroidEGL::InitEGL(APIVariant API)
 
 		// Favor EGLConfigLists by RGB, then Depth, then Non-linear Depth, then Stencil, then Alpha
 		currScore = 0;
-		currScore |= ((int64)FPlatformMath::Min(FPlatformMath::Abs(sb - PImplData->Parms.sampleBuffers), 15)) << 29;
-		currScore |= ((int64)FPlatformMath::Min(FPlatformMath::Abs(sc - PImplData->Parms.sampleSamples), 31)) << 24;
+		currScore |= ((int64)FPlatformMath::Min(FPlatformMath::Abs(sb - PImplData->EGLDisplay.Parms.sampleBuffers), 15)) << 29;
+		currScore |= ((int64)FPlatformMath::Min(FPlatformMath::Abs(sc - PImplData->EGLDisplay.Parms.sampleSamples), 31)) << 24;
 		currScore |= FPlatformMath::Min(
-						FPlatformMath::Abs(r - PImplData->Parms.redSize) +
-						FPlatformMath::Abs(g - PImplData->Parms.greenSize) +
-						FPlatformMath::Abs(b - PImplData->Parms.blueSize), 127) << 17;
-		currScore |= FPlatformMath::Min(FPlatformMath::Abs(d - PImplData->Parms.depthSize), 63) << 11;
+						FPlatformMath::Abs(r - PImplData->EGLDisplay.Parms.redSize) +
+						FPlatformMath::Abs(g - PImplData->EGLDisplay.Parms.greenSize) +
+						FPlatformMath::Abs(b - PImplData->EGLDisplay.Parms.blueSize), 127) << 17;
+		currScore |= FPlatformMath::Min(FPlatformMath::Abs(d - PImplData->EGLDisplay.Parms.depthSize), 63) << 11;
 		currScore |= FPlatformMath::Min(FPlatformMath::Abs(1 - bNonLinearDepth), 1) << 10;
-		currScore |= FPlatformMath::Min(FPlatformMath::Abs(s - PImplData->Parms.stencilSize), 31) << 6;
-		currScore |= FPlatformMath::Min(FPlatformMath::Abs(a - PImplData->Parms.alphaSize), 31) << 0;
+		currScore |= FPlatformMath::Min(FPlatformMath::Abs(s - PImplData->EGLDisplay.Parms.stencilSize), 31) << 6;
+		currScore |= FPlatformMath::Min(FPlatformMath::Abs(a - PImplData->EGLDisplay.Parms.alphaSize), 31) << 0;
 
 #if ENABLE_EGL_DEBUG
 		LogConfigInfo(EGLConfigList[i]);
@@ -436,11 +508,12 @@ void AndroidEGL::InitEGL(APIVariant API)
 
 		if (currScore < score || !haveConfig)
 		{
-			PImplData->eglConfigParam	= EGLConfigList[i];
-			PImplData->DepthSize = d;		// store depth/stencil sizes
+			PImplData->EGLDisplay.eglConfigParam	= EGLConfigList[i];
+			PImplData->EGLDisplay.DepthSize = d;		// store depth/stencil sizes
 			haveConfig	= true;
 			score		= currScore;
-			eglGetConfigAttrib(PImplData->eglDisplay, EGLConfigList[i], EGL_NATIVE_VISUAL_ID, &ResultValue);PImplData->NativeVisualID = ResultValue;
+			eglGetConfigAttrib(PImplData->EGLDisplay.eglDisplay, EGLConfigList[i], EGL_NATIVE_VISUAL_ID, &ResultValue);
+			PImplData->EGLDisplay.NativeVisualID = ResultValue;
 		}
 	}
 	check(haveConfig);
@@ -448,18 +521,18 @@ void AndroidEGL::InitEGL(APIVariant API)
 #else
 
 	EGLConfig EGLConfigList[1];
-	if(!( result = eglChooseConfig(PImplData->eglDisplay, Attributes, EGLConfigList, 1,  &PImplData->eglNumConfigs)))
+	if(!( result = eglChooseConfig(PImplData->EGLDisplay.eglDisplay, Attributes, EGLConfigList, 1,  &PImplData->EGLDisplay.eglNumConfigs)))
 	{
 		ResetInternal();
 	}
 
 	checkf(result == EGL_TRUE , TEXT(" eglChooseConfig error: 0x%x"), eglGetError());
 
-	checkf(PImplData->eglNumConfigs != 0  ,TEXT(" eglChooseConfig num EGLConfigLists is 0 . error: 0x%x"), eglGetError());
-	PImplData->eglConfigParam	= EGLConfigList[0];
+	checkf(PImplData->EGLDisplay.eglNumConfigs != 0  ,TEXT(" eglChooseConfig num EGLConfigLists is 0 . error: 0x%x"), eglGetError());
+	PImplData->EGLDisplay.eglConfigParam = EGLConfigList[0];
 	int ResultValue = 0 ;
-	eglGetConfigAttrib(PImplData->eglDisplay, EGLConfigList[0], EGL_DEPTH_SIZE, &ResultValue); PImplData->DepthSize = ResultValue;
-	eglGetConfigAttrib(PImplData->eglDisplay, EGLConfigList[0], EGL_NATIVE_VISUAL_ID, &ResultValue);PImplData->NativeVisualID = ResultValue;
+	eglGetConfigAttrib(PImplData->EGLDisplay.eglDisplay, EGLConfigList[0], EGL_DEPTH_SIZE, &ResultValue); PImplData->EGLDisplay.DepthSize = ResultValue;
+	eglGetConfigAttrib(PImplData->EGLDisplay.eglDisplay, EGLConfigList[0], EGL_NATIVE_VISUAL_ID, &ResultValue); PImplData->EGLDisplay.NativeVisualID = ResultValue;
 #endif
 }
 
@@ -472,52 +545,58 @@ void AndroidEGL::SetRenderContextWindowSurface(const TOptional<FAndroidWindow::F
 
 	FPlatformMisc::LowLevelOutputDebugStringf(TEXT("AndroidEGL::SetRenderContextWindowSurface  recreating context! tid: %d"),  FPlatformTLS::GetCurrentThreadId());
 
-	UnBindRender();
+	FPlatformOpenGLSurfaceContext* Surface = PImplData->GetGLSurfaceFromWindowContainer(WindowContainer);
+	if (!Surface)
+	{
+		UE_LOG(LogRHI, Warning, TEXT("AndroidEGL::SetRenderContextWindowSurface: no surface context available, aborting"));
+		return;
+	}
+
+	UnBindRender(Surface);
 
-	SetCurrentContext(EGL_NO_CONTEXT, EGL_NO_SURFACE);
+	SetCurrentContext(EGL_NO_CONTEXT, nullptr);
 	bool bCreateSurface = !AndroidThunkCpp_IsXRImmersiveHMDApplication();
 
 	if (!FAndroidMisc::UseNewWindowBehavior())
 	{
 		// SetRenderContextWindowSurface is called only when the window lock is successful.
-		PImplData->Window = (ANativeWindow*)FAndroidWindow::GetHardwareWindow_EventThread();
-		check(PImplData->Window);
+		Surface->Window = (ANativeWindow*)FAndroidWindow::GetHardwareWindow_EventThread();
+	}
+	else
+	{
+		check(WindowContainer.IsSet());
+		Surface->Window = WindowContainer->GetANativeWindow();
 	}
+	check(Surface->Window);
 
-	InitRenderSurface(false, bCreateSurface, WindowContainer);
-	SetCurrentRenderingContext();
+	InitRenderSurface(false, bCreateSurface, Surface);
+	SetCurrentRenderingContext(Surface);
 
 	FPlatformMisc::LowLevelOutputDebugStringf(TEXT("AndroidEGL::SetRenderContextWindowSurface  DONE! tid: %d"),  FPlatformTLS::GetCurrentThreadId());
 }
 
-void AndroidEGL::ResizeRenderContextSurface(const TOptional<FAndroidWindow::FNativeAccessor>& WindowContainer)
+void AndroidEGL::ResizeRenderContextSurface(FPlatformOpenGLSurfaceContext* Surface)
 {
 	VERIFY_GL_SCOPE();
 
-	check(!FAndroidMisc::UseNewWindowBehavior() || !WindowContainer.IsSet() || WindowContainer.GetValue().GetANativeWindow() == PImplData->Window);
 	// Resize render originates from the gamethread, we cant use Window_Event here.
-	if (PImplData->Window)
+	if (Surface->eglWidth != (Surface->CachedWindowRect.Right - Surface->CachedWindowRect.Left) ||
+		Surface->eglHeight != (Surface->CachedWindowRect.Bottom - Surface->CachedWindowRect.Top))
 	{
-		if (PImplData->eglWidth != (PImplData->CachedWindowRect.Right - PImplData->CachedWindowRect.Left) || 
-			PImplData->eglHeight != (PImplData->CachedWindowRect.Bottom - PImplData->CachedWindowRect.Top))
+		UE_LOG(LogAndroid, Log, TEXT("AndroidEGL::ResizeRenderContextSurface, Surface->Window=%p, Surface->eglWidth=%d, Surface->eglHeight=%d!, CachedWidth=%d, CachedHeight=%d, tid: %d"), 
+			Surface->Window,
+			Surface->eglWidth,
+			Surface->eglHeight,
+			(Surface->CachedWindowRect.Right - Surface->CachedWindowRect.Left),
+			(Surface->CachedWindowRect.Bottom - Surface->CachedWindowRect.Top),
+			FPlatformTLS::GetCurrentThreadId()
+		);
+
+		UnBindRender(Surface);
+
 		{
-			UE_LOG(LogAndroid, Log, TEXT("AndroidEGL::ResizeRenderContextSurface, PImplData->Window=%p, PImplData->eglWidth=%d, PImplData->eglHeight=%d!, CachedWidth=%d, CachedHeight=%d, tid: %d"), 
-				PImplData->Window,
-				PImplData->eglWidth,
-				PImplData->eglHeight,
-				(PImplData->CachedWindowRect.Right - PImplData->CachedWindowRect.Left),
-				(PImplData->CachedWindowRect.Bottom - PImplData->CachedWindowRect.Top),
-				FPlatformTLS::GetCurrentThreadId()
-			);
-
-			UnBindRender();
-
-			SetCurrentContext(EGL_NO_CONTEXT, EGL_NO_SURFACE);
-			{
-				bool bCreateSurface = !AndroidThunkCpp_IsXRImmersiveHMDApplication();
-				InitRenderSurface(false, bCreateSurface, WindowContainer);
-			}
-			SetCurrentRenderingContext();
+			bool bCreateSurface = !AndroidThunkCpp_IsXRImmersiveHMDApplication();
+			InitRenderSurface(false, bCreateSurface, Surface);
 		}
 	}
 }
@@ -532,63 +611,71 @@ AndroidEGL* AndroidEGL::GetInstance()
 	return Singleton;
 }
 
-void AndroidEGL::DestroyBackBuffer()
+void AndroidEGL::DestroyBackBuffer(FPlatformOpenGLSurfaceContext* Surface)
 {
-	if (PImplData->ResolveFrameBuffer)
+	if (ensure(Surface->ViewportFramebuffer))
 	{
 		VERIFY_GL_SCOPE();
-		glDeleteFramebuffers(1, &PImplData->ResolveFrameBuffer);
-		PImplData->ResolveFrameBuffer = 0 ;
-	}
+		glDeleteFramebuffers(1, &Surface->ViewportFramebuffer);
 
-	if (PImplData->DummyFrameBuffer)
-	{
-		VERIFY_GL_SCOPE();
-		glDeleteFramebuffers(1, &PImplData->DummyFrameBuffer);
-		PImplData->DummyFrameBuffer = 0;
+		Surface->ViewportFramebuffer = 0 ;
 	}
 }
 
-void AndroidEGL::InitBackBuffer()
+// this is expected to happen once, during RHIInit
+void AndroidEGL::InitBackBuffer(struct FPlatformOpenGLSurfaceContext* Surface)
 {
-	PImplData->RenderingContext.ViewportFramebuffer = GetResolveFrameBuffer();
+	if (Surface == nullptr)
+	{
+		if (!FAndroidMisc::UseNewWindowBehavior())
+		{
+			Surface = FindSurfaceContextFromWindow(nullptr);
+		}
+		else
+		{
+			return; //to satisfy lint warning
+		}
+	} 
+
+	check(Surface);
+	check(Surface->ViewportFramebuffer == 0);
+	VERIFY_GL_SCOPE();
+	glGenFramebuffers(1, &Surface->ViewportFramebuffer);
 }
 
 extern void AndroidThunkCpp_SetDesiredViewSize(int32 Width, int32 Height);
 
-void AndroidEGL::InitRenderSurface(bool bUseSmallSurface, bool bCreateWndSurface, const TOptional<FAndroidWindow::FNativeAccessor>& WindowContainer)
+void AndroidEGL::InitRenderSurface(bool bUseSmallSurface, bool bCreateWndSurface, FPlatformOpenGLSurfaceContext* Surface)
 {
 	FPlatformMisc::LowLevelOutputDebugStringf(TEXT("AndroidEGL::InitRenderSurface %d, %d"), int(bUseSmallSurface), int(bCreateWndSurface));
+	if (Surface == nullptr)
+	{
+		// todo: validate, only needed 
+		Surface = FindSurfaceContextFromWindow(nullptr);
+	}
 
 	if (FAndroidMisc::UseNewWindowBehavior())
 	{
 		bUseSmallSurface = true;
 		bCreateWndSurface = false;
-		PImplData->Window = nullptr;
 
-		if (WindowContainer.IsSet())
 		{
-			const FAndroidWindow::FNativeAccessor& Accessor = WindowContainer.GetValue();
-			if(Accessor.GetANativeWindow())
+			if(Surface->Window)
 			{
 				bUseSmallSurface = false;
 				bCreateWndSurface = true;
-				PImplData->Window = (ANativeWindow*) Accessor.GetANativeWindow();
-				FPlatformMisc::LowLevelOutputDebugStringf(TEXT("AndroidEGL::InitRenderSurface window %p was supplied %d, %d"), PImplData->Window, int(bUseSmallSurface), int(bCreateWndSurface));
+				FPlatformMisc::LowLevelOutputDebugStringf(TEXT("AndroidEGL::InitRenderSurface window %p was supplied %d, %d"), Surface->Window, int(bUseSmallSurface), int(bCreateWndSurface));
 			}
 			else
 			{
 				FPlatformMisc::LowLevelOutputDebugStringf(TEXT("AndroidEGL::InitRenderSurface null native window was supplied %d, %d"), int(bUseSmallSurface), int(bCreateWndSurface));
 			}
 		}
-		else
-		{
-			FPlatformMisc::LowLevelOutputDebugStringf(TEXT("AndroidEGL::InitRenderSurface No window was supplied %d, %d"), int(bUseSmallSurface), int(bCreateWndSurface));
-		}
 	}
 	else
 	{
-		check(PImplData->Window);
+		// the previous window behaviour should not see a null native window..
+		check(Surface->Window);
 	}
 
 	int32 Width = 8, Height = 8;
@@ -596,11 +683,11 @@ void AndroidEGL::InitRenderSurface(bool bUseSmallSurface, bool bCreateWndSurface
 	{
 		FPlatformRect WindowSize = FAndroidWindow::GetScreenRect();
 
-		if (PImplData->CachedWindowRect.Right > 0 && PImplData->CachedWindowRect.Bottom > 0)
+		if (Surface->CachedWindowRect.Right > 0 && Surface->CachedWindowRect.Bottom > 0)
 		{
 			// If we resumed from a lost window reuse the window size, the game thread will update the window dimensions.
-			FPlatformMisc::LowLevelOutputDebugStringf(TEXT("AndroidEGL::InitRenderSurface, Using CachedWindowRect, left: %d, top: %d, right: %d, bottom: %d "), PImplData->CachedWindowRect.Left, PImplData->CachedWindowRect.Top, PImplData->CachedWindowRect.Right, PImplData->CachedWindowRect.Bottom);
-			WindowSize = PImplData->CachedWindowRect;
+			FPlatformMisc::LowLevelOutputDebugStringf(TEXT("AndroidEGL::InitRenderSurface, Using CachedWindowRect, left: %d, top: %d, right: %d, bottom: %d "), Surface->CachedWindowRect.Left, Surface->CachedWindowRect.Top, Surface->CachedWindowRect.Right, Surface->CachedWindowRect.Bottom);
+			WindowSize = Surface->CachedWindowRect;
 		}
 #if USE_ANDROID_STANDALONE
 		if (WindowSize.Left != 0 || WindowSize.Top != 0)
@@ -609,26 +696,32 @@ void AndroidEGL::InitRenderSurface(bool bUseSmallSurface, bool bCreateWndSurface
 		}
 		Width = WindowSize.Right - WindowSize.Left;
 		Height = WindowSize.Bottom - WindowSize.Top;
+
+		if (Surface->Window)
+		{
+			Width = ANativeWindow_getWidth(Surface->Window);
+			Height = ANativeWindow_getHeight(Surface->Window);
+		}
 #else
 
 		Width = WindowSize.Right;
 		Height = WindowSize.Bottom;
+		
+		AndroidThunkCpp_SetDesiredViewSize(Width, Height);
 #endif
 
 		FPlatformMisc::LowLevelOutputDebugStringf(TEXT("AndroidEGL::InitRenderSurface, Using width: %d, height %d "), Width, Height);
-		AndroidThunkCpp_SetDesiredViewSize(Width, Height);
 	}
 
-	if(PImplData->Window)
+	if(Surface->Window)
 	{
-		FIntVector2 OriginalWindowSize(ANativeWindow_getWidth(PImplData->Window), ANativeWindow_getHeight(PImplData->Window));
+		FIntVector2 OriginalWindowSize(ANativeWindow_getWidth(Surface->Window), ANativeWindow_getHeight(Surface->Window));
 
-		FPlatformMisc::LowLevelOutputDebugStringf(TEXT("AndroidEGL::InitRenderSurface, setting wnd: %p, width: %d->%d, height %d->%d "), PImplData->Window, OriginalWindowSize.X, Width, OriginalWindowSize.Y, Height);
-		ANativeWindow_setBuffersGeometry(PImplData->Window, Width, Height, PImplData->NativeVisualID);
+		FPlatformMisc::LowLevelOutputDebugStringf(TEXT("AndroidEGL::InitRenderSurface, setting wnd: %p, width: %d->%d, height %d->%d "), Surface->Window, OriginalWindowSize.X, Width, OriginalWindowSize.Y, Height);
+		ANativeWindow_setBuffersGeometry(Surface->Window, Width, Height, PImplData->EGLDisplay.NativeVisualID);
 	}
-	CreateEGLRenderSurface(PImplData->Window, bCreateWndSurface);
 
-	PImplData->RenderingContext.eglSurface = PImplData->eglSurface;
+	CreateEGLRenderSurface(Surface, bCreateWndSurface);
 }
 
 // EGL_CONTEXT_OPENGL_ROBUST_ACCESS_EXT is enabled if configrules asks for it or the command line specifies it.
@@ -648,18 +741,18 @@ void AndroidEGL::Init(APIVariant API, uint32 MajorVersion, uint32 MinorVersion)
 	
 	if (PImplData->Initalized)
 	{
-		ensure(bDebug == PImplData->bIsDebug); // if this fires you would need to tear down the previous context and recreate to honour the debug change.
+		ensure(bDebug == PImplData->RenderingContext.bIsDebug); // if this fires you would need to tear down the previous context and recreate to honour the debug change.
 		return;
 	}
 
 	InitEGL(API);
-	PImplData->bIsDebug = bDebug;
+	PImplData->RenderingContext.bIsDebug = bDebug;
 	if (bSupportsKHRCreateContext)
 	{
 		const uint32 MaxElements = 16;
 		uint32 Flags = 0;
 
-		Flags |= PImplData->bIsDebug ? EGL_CONTEXT_OPENGL_DEBUG_BIT_KHR : 0;
+		Flags |= PImplData->RenderingContext.bIsDebug ? EGL_CONTEXT_OPENGL_DEBUG_BIT_KHR : 0;
 
 		ContextAttributes = new int[MaxElements];
 		uint32 Element = 0;
@@ -725,11 +818,6 @@ void AndroidEGL::Init(APIVariant API, uint32 MajorVersion, uint32 MinorVersion)
 		}
 	}
 
-	if (!FAndroidMisc::UseNewWindowBehavior())
-	{
-		// Getting the hardware window is valid during preinit as we have GAndroidWindowLock held.
-		PImplData->Window = (ANativeWindow*)FAndroidWindow::GetHardwareWindow_EventThread();
-	}
 	PImplData->Initalized   = true;
 }
 
@@ -739,25 +827,19 @@ AndroidEGL::~AndroidEGL()
 	delete []ContextAttributes;
 }
 
-void AndroidEGL::GetDimensions(uint32& OutWidth, uint32& OutHeight)
-{
-	OutWidth = PImplData->eglWidth;
-	OutHeight = PImplData->eglHeight;
-}
-
 void AndroidEGL::DestroyContext(EGLContext InContext)
 {
 	VERIFY_EGL_SCOPE();
 	if(InContext != EGL_NO_CONTEXT) //soft fail
 	{
-		eglDestroyContext(PImplData->eglDisplay, InContext);
+		eglDestroyContext(PImplData->EGLDisplay.eglDisplay, InContext);
 	}
 }
 
 EGLContext AndroidEGL::CreateContext(EGLContext InParentContext)
 {
 	VERIFY_EGL_SCOPE();
-	return eglCreateContext(PImplData->eglDisplay, PImplData->eglConfigParam, InParentContext, ContextAttributes);
+	return eglCreateContext(PImplData->EGLDisplay.eglDisplay, PImplData->EGLDisplay.eglConfigParam, InParentContext, ContextAttributes);
 }
 
 int32 AndroidEGL::GetError()
@@ -770,47 +852,51 @@ bool AndroidEGL::IsInitialized()
 	return PImplData->Initalized;
 }
 
-GLuint AndroidEGL::GetResolveFrameBuffer()
+EGLContext AndroidEGL::GetCurrentContext()
 {
-	return PImplData->ResolveFrameBuffer;
+	VERIFY_EGL_SCOPE();
+	return eglGetCurrentContext();
 }
 
-EGLContext AndroidEGL::GetCurrentContext()
+EGLSurface AndroidEGL::GetCurrentSurface(EGLint readdraw)
 {
 	VERIFY_EGL_SCOPE();
-	return eglGetCurrentContext();
+	return eglGetCurrentSurface(readdraw);
 }
 
 EGLDisplay AndroidEGL::GetDisplay() const
 {
-	return PImplData->eglDisplay;
+	return PImplData->EGLDisplay.eglDisplay;
 }
 
 EGLSurface AndroidEGL::GetSurface() const
 {
-	return PImplData->eglSurface;
+	FPlatformOpenGLSurfaceContext* Surface = PImplData->GetSurfaceFromAndroidWindow(nullptr);
+	return Surface->eglSurface;
 }
 
 EGLConfig AndroidEGL::GetConfig() const
 {
-	return PImplData->eglConfigParam;
+	return PImplData->EGLDisplay.eglConfigParam;
 }
 
 bool AndroidEGL::IsUsingWindowedSurface() const
 {
-	return PImplData->bIsWndSurface;
+	FPlatformOpenGLSurfaceContext* Surface = PImplData->GetSurfaceFromAndroidWindow(nullptr);
+	return Surface->bIsWndSurface;
 }
 
 void AndroidEGL::GetSwapIntervalRange(EGLint& OutMinSwapInterval, EGLint& OutMaxSwapInterval) const
 {
-	eglGetConfigAttrib(PImplData->eglDisplay, PImplData->eglConfigParam, EGL_MIN_SWAP_INTERVAL, &OutMinSwapInterval);
-	eglGetConfigAttrib(PImplData->eglDisplay, PImplData->eglConfigParam, EGL_MAX_SWAP_INTERVAL, &OutMaxSwapInterval);
+	eglGetConfigAttrib(PImplData->EGLDisplay.eglDisplay, PImplData->EGLDisplay.eglConfigParam, EGL_MIN_SWAP_INTERVAL, &OutMinSwapInterval);
+	eglGetConfigAttrib(PImplData->EGLDisplay.eglDisplay, PImplData->EGLDisplay.eglConfigParam, EGL_MAX_SWAP_INTERVAL, &OutMaxSwapInterval);
 }
 
 
 ANativeWindow* AndroidEGL::GetNativeWindow() const
 {
-	return PImplData->Window;
+	FPlatformOpenGLSurfaceContext* Surface = PImplData->GetSurfaceFromAndroidWindow(nullptr);
+	return Surface->Window;
 }
 
 bool AndroidEGL::InitContexts()
@@ -821,33 +907,13 @@ bool AndroidEGL::InitContexts()
 
 void AndroidEGL::AcquireCurrentRenderingContext()
 {
-	SetCurrentRenderingContext();
-
-	if (!PImplData->DummyFrameBuffer)
-	{
-		// Dummy FBO we bind right after SwapBuffers to tell driver that backbuffer is no longer in use by the App
-		glGenFramebuffers(1, &PImplData->DummyFrameBuffer);
-		PImplData->RenderingContext.DummyFrameBuffer = PImplData->DummyFrameBuffer;
-	}
-
-	if (IsOfflineSurfaceRequired())
-	{
-		// Needs to be generated on rendering context
-		if (!PImplData->ResolveFrameBuffer)
-		{
-			glGenFramebuffers(1, &PImplData->ResolveFrameBuffer);
-		}
-	}
-	else
-	{
-		PImplData->ResolveFrameBuffer = 0;
-	}
-
+	// we can acquire a context without a windowed surface.
+	SetCurrentRenderingContext(nullptr);
 }
 
-void AndroidEGL::SetCurrentRenderingContext()
+void AndroidEGL::SetCurrentRenderingContext(FPlatformOpenGLSurfaceContext* InSurface)
 {
-	SetCurrentContext(PImplData->RenderingContext.eglContext, PImplData->RenderingContext.eglSurface);
+	SetCurrentContext(PImplData->RenderingContext.eglContext, InSurface);
 }
 
 void AndroidEGL::ReleaseContextOwnership()
@@ -855,7 +921,7 @@ void AndroidEGL::ReleaseContextOwnership()
 	if (PlatformOpenGLThreadHasRenderingContext())
 	{
 		FPlatformMisc::LowLevelOutputDebugStringf(TEXT("AndroidEGL::  ReleaseContextOwnership,  thread releasing rendering context tid: %d"), FPlatformTLS::GetCurrentThreadId());
-		SetCurrentContext(EGL_NO_CONTEXT, EGL_NO_SURFACE);
+		SetCurrentContext(EGL_NO_CONTEXT, nullptr);
 	}
 	else
 	{
@@ -872,7 +938,16 @@ void AndroidEGL::Terminate()
 	DestroyContext(PImplData->RenderingContext.eglContext);
 	PImplData->RenderingContext.Reset();
 
-	DestroyRenderSurface();
+	{
+		UE::TUniqueLock ScopeLock(PImplData->SurfacesMutex);
+		for (TUniquePtr<FPlatformOpenGLSurfaceContext>& Surface : PImplData->GLSurfaces)
+		{
+			DestroyBackBuffer(Surface.Get());
+			DestroyRenderSurface(Surface.Get());
+		}
+		PImplData->GLSurfaces.Empty();
+	}
+
 	TerminateEGL();
 }
 
@@ -891,13 +966,29 @@ bool AndroidEGL::GetSupportsNoErrorContext()
 	return bSupportsKHRNoErrorContext;
 }
 
-void AndroidEGL::UnBindRender()
+void AndroidEGL::UnBindRender(FPlatformOpenGLSurfaceContext* Surface)
 {
-	FPlatformMisc::LowLevelOutputDebugString(TEXT("AndroidEGL::UnBindRender()"));
+	check(Surface);
+	FPlatformMisc::LowLevelOutputDebugStringf(TEXT("AndroidEGL::UnBindRender(%p)"), Surface->Window);
 	ResetDisplay();
-	DestroyRenderSurface();
+	DestroyRenderSurface(Surface);
+}
+
+// Called by new window behaviour when the android event thread wants to destroy a window.
+// The event thread waits for this to complete.
+void AndroidEGL::OnLostNativeWindow(FPlatformOpenGLSurfaceContext* Surface)
+{
+	check(FAndroidMisc::UseNewWindowBehavior());
+	// remove a surface from the context, free the surface and rebind the context with a null windows. (this creates a pbuffer surface and binds the context to it.)
+	UnBindRender(Surface);
+	FPlatformMisc::LowLevelOutputDebugStringf(TEXT("AndroidEGL::OnLostNativeWindow(native window = null (was %p) )"), Surface->Window);
+	Surface->Window = nullptr;
+	// recreate the surface and make current
+	InitRenderSurface(true, false, Surface);
+	SetCurrentRenderingContext(Surface);
 }
 
+
 void FAndroidAppEntry::ReInitWindow(const TOptional<FAndroidWindow::FNativeAccessor>& WindowContainer)
 {
 	check(IsInGameThread());
@@ -928,15 +1019,28 @@ void AndroidEGL::RefreshWindowSize(const TOptional<FAndroidWindow::FNativeAccess
 	check(IsInGameThread());
 	check(!FAndroidMisc::ShouldUseVulkan());
 	FPlatformRect WindowRect = FAndroidWindow::GetScreenRect();
-	UE_LOG(LogAndroid, Log, TEXT("AndroidEGL::RefreshWindowSize updating window size = %d, %d, cached size : %d, %d tid : %d"), WindowRect.Right, WindowRect.Bottom, PImplData->CachedWindowRect.Right, PImplData->CachedWindowRect.Bottom, FPlatformTLS::GetCurrentThreadId());
-	PImplData->CachedWindowRect = WindowRect;
+
+	// Extract the FAndroidWindow* before lambda capture (FNativeAccessor is non-copyable)
+	const FAndroidWindow* AndroidWindowPtr = nullptr;
+	if (FAndroidMisc::UseNewWindowBehavior() && WindowContainer.IsSet())
+	{
+		AndroidWindowPtr = WindowContainer->Get();
+	}
 
 	ENQUEUE_RENDER_COMMAND(EGLResizeRenderContextSurface)(
-		[&WindowContainer](FRHICommandListImmediate& RHICmdList) mutable
+		[this, WindowRect, AndroidWindowPtr](FRHICommandListImmediate& RHICmdList) mutable
 	{
-		RHICmdList.EnqueueLambda([&WindowContainer](FRHICommandListImmediate&) mutable
-		{
-			AndroidEGL::GetInstance()->ResizeRenderContextSurface(WindowContainer);
+			RHICmdList.EnqueueLambda([this, WindowRect, AndroidWindowPtr](FRHICommandListImmediate&) mutable
+				{
+					FPlatformOpenGLSurfaceContext* Surface = PImplData->GetSurfaceFromAndroidWindow(AndroidWindowPtr);
+					if (!Surface)
+					{
+						UE_LOG(LogAndroid, Warning, TEXT("AndroidEGL::RefreshWindowSize surface context not found, skipping resize"));
+						return;
+					}
+					UE_LOG(LogAndroid, Log, TEXT("AndroidEGL::RefreshWindowSize updating window size = %d, %d, cached size : %d, %d tid : %d"), WindowRect.Right, WindowRect.Bottom, Surface->CachedWindowRect.Right, Surface->CachedWindowRect.Bottom, FPlatformTLS::GetCurrentThreadId());
+					Surface->CachedWindowRect = WindowRect;
+					AndroidEGL::GetInstance()->ResizeRenderContextSurface(Surface);
 		});
 	});
 
@@ -956,35 +1060,36 @@ void AndroidEGL::LogConfigInfo(EGLConfig  EGLConfigInfo)
 {
 	VERIFY_EGL_SCOPE();
 	EGLint ResultValue = 0 ;
-	eglGetConfigAttrib(PImplData->eglDisplay, EGLConfigInfo, EGL_RED_SIZE, &ResultValue); FPlatformMisc::LowLevelOutputDebugStringf( TEXT("EGLConfigInfo : EGL_RED_SIZE :	%u" ), ResultValue );
-	eglGetConfigAttrib(PImplData->eglDisplay, EGLConfigInfo, EGL_GREEN_SIZE, &ResultValue);  FPlatformMisc::LowLevelOutputDebugStringf( TEXT("EGLConfigInfo :EGL_GREEN_SIZE :	%u" ), ResultValue );
-	eglGetConfigAttrib(PImplData->eglDisplay, EGLConfigInfo, EGL_BLUE_SIZE, &ResultValue);  FPlatformMisc::LowLevelOutputDebugStringf( TEXT("EGLConfigInfo :EGL_BLUE_SIZE :	%u" ), ResultValue );
-	eglGetConfigAttrib(PImplData->eglDisplay, EGLConfigInfo, EGL_ALPHA_SIZE, &ResultValue); FPlatformMisc::LowLevelOutputDebugStringf( TEXT("EGLConfigInfo :EGL_ALPHA_SIZE :	%u" ), ResultValue );
-	eglGetConfigAttrib(PImplData->eglDisplay, EGLConfigInfo, EGL_DEPTH_SIZE, &ResultValue);  FPlatformMisc::LowLevelOutputDebugStringf( TEXT("EGLConfigInfo :EGL_DEPTH_SIZE :	%u" ), ResultValue );
-	eglGetConfigAttrib(PImplData->eglDisplay, EGLConfigInfo, EGL_STENCIL_SIZE, &ResultValue);  FPlatformMisc::LowLevelOutputDebugStringf( TEXT("EGLConfigInfo :EGL_STENCIL_SIZE :	%u" ), ResultValue );
-	eglGetConfigAttrib(PImplData->eglDisplay, EGLConfigInfo, EGL_SAMPLE_BUFFERS, &ResultValue);  FPlatformMisc::LowLevelOutputDebugStringf( TEXT("EGLConfigInfo :EGL_SAMPLE_BUFFERS :	%u" ), ResultValue );
-	eglGetConfigAttrib(PImplData->eglDisplay,EGLConfigInfo, EGL_BIND_TO_TEXTURE_RGB, &ResultValue);  FPlatformMisc::LowLevelOutputDebugStringf( TEXT("EGLConfigInfo :EGL_BIND_TO_TEXTURE_RGB :	%u" ), ResultValue );
-	eglGetConfigAttrib(PImplData->eglDisplay,EGLConfigInfo, EGL_SAMPLES, &ResultValue);  FPlatformMisc::LowLevelOutputDebugStringf( TEXT("EGLConfigInfo :EGL_SAMPLES :	%u" ), ResultValue );
-	eglGetConfigAttrib(PImplData->eglDisplay,EGLConfigInfo, EGL_COLOR_BUFFER_TYPE, &ResultValue);  FPlatformMisc::LowLevelOutputDebugStringf( TEXT("EGLConfigInfo :EGL_COLOR_BUFFER_TYPE :	%u" ), ResultValue );
-	eglGetConfigAttrib(PImplData->eglDisplay,EGLConfigInfo, EGL_CONFIG_CAVEAT, &ResultValue);  FPlatformMisc::LowLevelOutputDebugStringf( TEXT("EGLConfigInfo :EGL_CONFIG_CAVEAT :	%u" ), ResultValue );
-	eglGetConfigAttrib(PImplData->eglDisplay,EGLConfigInfo, EGL_CONFIG_ID, &ResultValue); FPlatformMisc::LowLevelOutputDebugStringf( TEXT("EGLConfigInfo :EGL_CONFIG_ID :	%u" ), ResultValue );
-	eglGetConfigAttrib(PImplData->eglDisplay,EGLConfigInfo, EGL_CONFORMANT, &ResultValue);  FPlatformMisc::LowLevelOutputDebugStringf( TEXT("EGLConfigInfo :EGL_CONFORMANT :	%u" ), ResultValue );
-	eglGetConfigAttrib(PImplData->eglDisplay,EGLConfigInfo, EGL_LEVEL, &ResultValue);  FPlatformMisc::LowLevelOutputDebugStringf( TEXT("EGLConfigInfo :EGL_LEVEL :	%u" ), ResultValue );
-	eglGetConfigAttrib(PImplData->eglDisplay,EGLConfigInfo, EGL_LUMINANCE_SIZE, &ResultValue);  FPlatformMisc::LowLevelOutputDebugStringf( TEXT("EGLConfigInfo :EGL_LUMINANCE_SIZE :	%u" ), ResultValue );
-	eglGetConfigAttrib(PImplData->eglDisplay,EGLConfigInfo, EGL_MAX_PBUFFER_WIDTH, &ResultValue);  FPlatformMisc::LowLevelOutputDebugStringf( TEXT("EGLConfigInfo :EGL_MAX_PBUFFER_WIDTH :	%u" ), ResultValue );
-	eglGetConfigAttrib(PImplData->eglDisplay,EGLConfigInfo, EGL_MAX_PBUFFER_HEIGHT, &ResultValue);  FPlatformMisc::LowLevelOutputDebugStringf( TEXT("EGLConfigInfo :EGL_MAX_PBUFFER_HEIGHT :	%u" ), ResultValue );
-	eglGetConfigAttrib(PImplData->eglDisplay,EGLConfigInfo, EGL_MAX_PBUFFER_PIXELS, &ResultValue);  FPlatformMisc::LowLevelOutputDebugStringf( TEXT("EGLConfigInfo :EGL_MAX_PBUFFER_PIXELS :	%u" ), ResultValue );
-	eglGetConfigAttrib(PImplData->eglDisplay,EGLConfigInfo, EGL_MAX_SWAP_INTERVAL, &ResultValue);  FPlatformMisc::LowLevelOutputDebugStringf( TEXT("EGLConfigInfo :EGL_MAX_SWAP_INTERVAL :	%u" ), ResultValue );
-	eglGetConfigAttrib(PImplData->eglDisplay,EGLConfigInfo, EGL_MIN_SWAP_INTERVAL, &ResultValue);  FPlatformMisc::LowLevelOutputDebugStringf( TEXT("EGLConfigInfo :EGL_MIN_SWAP_INTERVAL :	%u" ), ResultValue );
-	eglGetConfigAttrib(PImplData->eglDisplay,EGLConfigInfo, EGL_NATIVE_RENDERABLE, &ResultValue);  FPlatformMisc::LowLevelOutputDebugStringf( TEXT("EGLConfigInfo :EGL_NATIVE_RENDERABLE :	%u" ), ResultValue );
-	eglGetConfigAttrib(PImplData->eglDisplay,EGLConfigInfo, EGL_NATIVE_VISUAL_TYPE, &ResultValue);  FPlatformMisc::LowLevelOutputDebugStringf( TEXT("EGLConfigInfo :EGL_NATIVE_VISUAL_TYPE :	%u" ), ResultValue );
-	eglGetConfigAttrib(PImplData->eglDisplay,EGLConfigInfo, EGL_NATIVE_VISUAL_ID, &ResultValue);  FPlatformMisc::LowLevelOutputDebugStringf( TEXT("EGLConfigInfo :EGL_NATIVE_VISUAL_ID :	%u" ), ResultValue );
-	eglGetConfigAttrib(PImplData->eglDisplay,EGLConfigInfo, EGL_RENDERABLE_TYPE, &ResultValue);  FPlatformMisc::LowLevelOutputDebugStringf( TEXT("EGLConfigInfo :EGL_RENDERABLE_TYPE :	%u" ), ResultValue );
-	eglGetConfigAttrib(PImplData->eglDisplay,EGLConfigInfo, EGL_SURFACE_TYPE, &ResultValue);  FPlatformMisc::LowLevelOutputDebugStringf( TEXT("EGLConfigInfo :EGL_SURFACE_TYPE :	%u" ), ResultValue );
-	eglGetConfigAttrib(PImplData->eglDisplay,EGLConfigInfo, EGL_TRANSPARENT_TYPE, &ResultValue);  FPlatformMisc::LowLevelOutputDebugStringf( TEXT("EGLConfigInfo :EGL_TRANSPARENT_TYPE :	%u" ), ResultValue );
-	eglGetConfigAttrib(PImplData->eglDisplay,EGLConfigInfo, EGL_TRANSPARENT_RED_VALUE, &ResultValue);  FPlatformMisc::LowLevelOutputDebugStringf( TEXT("EGLConfigInfo :EGL_TRANSPARENT_RED_VALUE :	%u" ), ResultValue );
-	eglGetConfigAttrib(PImplData->eglDisplay,EGLConfigInfo, EGL_TRANSPARENT_GREEN_VALUE, &ResultValue);  FPlatformMisc::LowLevelOutputDebugStringf( TEXT("EGLConfigInfo :EGL_TRANSPARENT_GREEN_VALUE :	%u" ), ResultValue );
-	eglGetConfigAttrib(PImplData->eglDisplay,EGLConfigInfo, EGL_TRANSPARENT_BLUE_VALUE, &ResultValue);  FPlatformMisc::LowLevelOutputDebugStringf( TEXT("EGLConfigInfo :EGL_TRANSPARENT_BLUE_VALUE :	%u" ), ResultValue );
+	EGLDisplay eglDisplay = PImplData->EGLDisplay.eglDisplay;
+	eglGetConfigAttrib(eglDisplay, EGLConfigInfo, EGL_RED_SIZE, &ResultValue); FPlatformMisc::LowLevelOutputDebugStringf( TEXT("EGLConfigInfo : EGL_RED_SIZE :	%u" ), ResultValue );
+	eglGetConfigAttrib(eglDisplay, EGLConfigInfo, EGL_GREEN_SIZE, &ResultValue);  FPlatformMisc::LowLevelOutputDebugStringf( TEXT("EGLConfigInfo :EGL_GREEN_SIZE :	%u" ), ResultValue );
+	eglGetConfigAttrib(eglDisplay, EGLConfigInfo, EGL_BLUE_SIZE, &ResultValue);  FPlatformMisc::LowLevelOutputDebugStringf( TEXT("EGLConfigInfo :EGL_BLUE_SIZE :	%u" ), ResultValue );
+	eglGetConfigAttrib(eglDisplay, EGLConfigInfo, EGL_ALPHA_SIZE, &ResultValue); FPlatformMisc::LowLevelOutputDebugStringf( TEXT("EGLConfigInfo :EGL_ALPHA_SIZE :	%u" ), ResultValue );
+	eglGetConfigAttrib(eglDisplay, EGLConfigInfo, EGL_DEPTH_SIZE, &ResultValue);  FPlatformMisc::LowLevelOutputDebugStringf( TEXT("EGLConfigInfo :EGL_DEPTH_SIZE :	%u" ), ResultValue );
+	eglGetConfigAttrib(eglDisplay, EGLConfigInfo, EGL_STENCIL_SIZE, &ResultValue);  FPlatformMisc::LowLevelOutputDebugStringf( TEXT("EGLConfigInfo :EGL_STENCIL_SIZE :	%u" ), ResultValue );
+	eglGetConfigAttrib(eglDisplay, EGLConfigInfo, EGL_SAMPLE_BUFFERS, &ResultValue);  FPlatformMisc::LowLevelOutputDebugStringf( TEXT("EGLConfigInfo :EGL_SAMPLE_BUFFERS :	%u" ), ResultValue );
+	eglGetConfigAttrib(eglDisplay, EGLConfigInfo, EGL_BIND_TO_TEXTURE_RGB, &ResultValue);  FPlatformMisc::LowLevelOutputDebugStringf( TEXT("EGLConfigInfo :EGL_BIND_TO_TEXTURE_RGB :	%u" ), ResultValue );
+	eglGetConfigAttrib(eglDisplay, EGLConfigInfo, EGL_SAMPLES, &ResultValue);  FPlatformMisc::LowLevelOutputDebugStringf( TEXT("EGLConfigInfo :EGL_SAMPLES :	%u" ), ResultValue );
+	eglGetConfigAttrib(eglDisplay, EGLConfigInfo, EGL_COLOR_BUFFER_TYPE, &ResultValue);  FPlatformMisc::LowLevelOutputDebugStringf( TEXT("EGLConfigInfo :EGL_COLOR_BUFFER_TYPE :	%u" ), ResultValue );
+	eglGetConfigAttrib(eglDisplay, EGLConfigInfo, EGL_CONFIG_CAVEAT, &ResultValue);  FPlatformMisc::LowLevelOutputDebugStringf( TEXT("EGLConfigInfo :EGL_CONFIG_CAVEAT :	%u" ), ResultValue );
+	eglGetConfigAttrib(eglDisplay, EGLConfigInfo, EGL_CONFIG_ID, &ResultValue); FPlatformMisc::LowLevelOutputDebugStringf( TEXT("EGLConfigInfo :EGL_CONFIG_ID :	%u" ), ResultValue );
+	eglGetConfigAttrib(eglDisplay, EGLConfigInfo, EGL_CONFORMANT, &ResultValue);  FPlatformMisc::LowLevelOutputDebugStringf( TEXT("EGLConfigInfo :EGL_CONFORMANT :	%u" ), ResultValue );
+	eglGetConfigAttrib(eglDisplay, EGLConfigInfo, EGL_LEVEL, &ResultValue);  FPlatformMisc::LowLevelOutputDebugStringf( TEXT("EGLConfigInfo :EGL_LEVEL :	%u" ), ResultValue );
+	eglGetConfigAttrib(eglDisplay, EGLConfigInfo, EGL_LUMINANCE_SIZE, &ResultValue);  FPlatformMisc::LowLevelOutputDebugStringf( TEXT("EGLConfigInfo :EGL_LUMINANCE_SIZE :	%u" ), ResultValue );
+	eglGetConfigAttrib(eglDisplay, EGLConfigInfo, EGL_MAX_PBUFFER_WIDTH, &ResultValue);  FPlatformMisc::LowLevelOutputDebugStringf( TEXT("EGLConfigInfo :EGL_MAX_PBUFFER_WIDTH :	%u" ), ResultValue );
+	eglGetConfigAttrib(eglDisplay, EGLConfigInfo, EGL_MAX_PBUFFER_HEIGHT, &ResultValue);  FPlatformMisc::LowLevelOutputDebugStringf( TEXT("EGLConfigInfo :EGL_MAX_PBUFFER_HEIGHT :	%u" ), ResultValue );
+	eglGetConfigAttrib(eglDisplay, EGLConfigInfo, EGL_MAX_PBUFFER_PIXELS, &ResultValue);  FPlatformMisc::LowLevelOutputDebugStringf( TEXT("EGLConfigInfo :EGL_MAX_PBUFFER_PIXELS :	%u" ), ResultValue );
+	eglGetConfigAttrib(eglDisplay, EGLConfigInfo, EGL_MAX_SWAP_INTERVAL, &ResultValue);  FPlatformMisc::LowLevelOutputDebugStringf( TEXT("EGLConfigInfo :EGL_MAX_SWAP_INTERVAL :	%u" ), ResultValue );
+	eglGetConfigAttrib(eglDisplay, EGLConfigInfo, EGL_MIN_SWAP_INTERVAL, &ResultValue);  FPlatformMisc::LowLevelOutputDebugStringf( TEXT("EGLConfigInfo :EGL_MIN_SWAP_INTERVAL :	%u" ), ResultValue );
+	eglGetConfigAttrib(eglDisplay, EGLConfigInfo, EGL_NATIVE_RENDERABLE, &ResultValue);  FPlatformMisc::LowLevelOutputDebugStringf( TEXT("EGLConfigInfo :EGL_NATIVE_RENDERABLE :	%u" ), ResultValue );
+	eglGetConfigAttrib(eglDisplay, EGLConfigInfo, EGL_NATIVE_VISUAL_TYPE, &ResultValue);  FPlatformMisc::LowLevelOutputDebugStringf( TEXT("EGLConfigInfo :EGL_NATIVE_VISUAL_TYPE :	%u" ), ResultValue );
+	eglGetConfigAttrib(eglDisplay, EGLConfigInfo, EGL_NATIVE_VISUAL_ID, &ResultValue);  FPlatformMisc::LowLevelOutputDebugStringf( TEXT("EGLConfigInfo :EGL_NATIVE_VISUAL_ID :	%u" ), ResultValue );
+	eglGetConfigAttrib(eglDisplay, EGLConfigInfo, EGL_RENDERABLE_TYPE, &ResultValue);  FPlatformMisc::LowLevelOutputDebugStringf( TEXT("EGLConfigInfo :EGL_RENDERABLE_TYPE :	%u" ), ResultValue );
+	eglGetConfigAttrib(eglDisplay, EGLConfigInfo, EGL_SURFACE_TYPE, &ResultValue);  FPlatformMisc::LowLevelOutputDebugStringf( TEXT("EGLConfigInfo :EGL_SURFACE_TYPE :	%u" ), ResultValue );
+	eglGetConfigAttrib(eglDisplay, EGLConfigInfo, EGL_TRANSPARENT_TYPE, &ResultValue);  FPlatformMisc::LowLevelOutputDebugStringf( TEXT("EGLConfigInfo :EGL_TRANSPARENT_TYPE :	%u" ), ResultValue );
+	eglGetConfigAttrib(eglDisplay, EGLConfigInfo, EGL_TRANSPARENT_RED_VALUE, &ResultValue);  FPlatformMisc::LowLevelOutputDebugStringf( TEXT("EGLConfigInfo :EGL_TRANSPARENT_RED_VALUE :	%u" ), ResultValue );
+	eglGetConfigAttrib(eglDisplay, EGLConfigInfo, EGL_TRANSPARENT_GREEN_VALUE, &ResultValue);  FPlatformMisc::LowLevelOutputDebugStringf( TEXT("EGLConfigInfo :EGL_TRANSPARENT_GREEN_VALUE :	%u" ), ResultValue );
+	eglGetConfigAttrib(eglDisplay, EGLConfigInfo, EGL_TRANSPARENT_BLUE_VALUE, &ResultValue);  FPlatformMisc::LowLevelOutputDebugStringf( TEXT("EGLConfigInfo :EGL_TRANSPARENT_BLUE_VALUE :	%u" ), ResultValue );
 }
 
 void AndroidEGL::UpdateBuffersTransform()
@@ -1101,7 +1206,7 @@ void SetSharedContextGameCommand(TSharedPtr<FEvent, ESPMode::ThreadSafe> GTBlock
 {
 	check(IsInGameThread());
 	AndroidEGL* EGL = AndroidEGL::GetInstance();
-	EGL->SetCurrentContext(EGL_NO_CONTEXT, EGL_NO_SURFACE);
+	EGL->SetCurrentContext(EGL_NO_CONTEXT, nullptr);
 
 	GTBlockedTrigger->Trigger();
 }
diff --git a/Engine/Source/Runtime/OpenGLDrv/Private/Android/AndroidEGL.h b/Engine/Source/Runtime/OpenGLDrv/Private/Android/AndroidEGL.h
--- a/Engine/Source/Runtime/OpenGLDrv/Private/Android/AndroidEGL.h
+++ b/Engine/Source/Runtime/OpenGLDrv/Private/Android/AndroidEGL.h
@@ -15,6 +15,7 @@
 #include <EGL/eglext.h>
 #include <GLES3/gl31.h>
 #include "Android/AndroidWindow.h"
+#include <atomic>
 
 struct AndroidESPImpl;
 struct ANativeWindow;
@@ -29,34 +30,135 @@ struct ANativeWindow;
 
 DECLARE_LOG_CATEGORY_EXTERN(LogEGL, Log, All);
 
+
+struct EGLConfigParms
+{
+	/** Whether this is a valid configuration or not */
+	int validConfig = 0;
+	/** The number of bits requested for the red component */
+	int redSize = 8;
+	/** The number of bits requested for the green component */
+	int greenSize = 8;
+	/** The number of bits requested for the blue component */
+	int blueSize = 8;
+	/** The number of bits requested for the alpha component */
+	int alphaSize = 0;
+	/** The number of bits requested for the depth component */
+	int depthSize = 24;
+	/** The number of bits requested for the stencil component */
+	int stencilSize = 0;
+	/** The number of multisample buffers requested */
+	int sampleBuffers = 0;
+	/** The number of samples requested */
+	int sampleSamples = 0;
+
+	EGLConfigParms()
+	{
+		// If not default, set the preference
+		int DepthBufferPreference = (int)FAndroidWindow::GetDepthBufferPreference();
+		if (DepthBufferPreference > 0)
+		{
+			depthSize = DepthBufferPreference;
+		}
+
+		if (FAndroidMisc::GetMobilePropagateAlphaSetting() > 0)
+		{
+			alphaSize = 8;
+		}
+	}
+};
+
+struct FAndroidEGLDisplayInfo
+{
+	EGLDisplay  eglDisplay = EGL_NO_DISPLAY;
+	EGLConfig	eglConfigParam = nullptr;
+
+	EGLConfigParms Parms;
+	EGLint         eglNumConfigs = 0;
+
+	EGLint		NativeVisualID = 0;
+	int			DepthSize = 0;
+
+	// this is an offscreen temp surface used when no HW window is supplied.
+	EGLSurface  auxSurface = EGL_NO_SURFACE;
+};
+
 struct FPlatformOpenGLContext
 {
-	EGLContext	eglContext;
+	FPlatformOpenGLContext()
+	{
+		Reset();
+	}
+
+	void Reset()
+	{
+		eglContext = EGL_NO_CONTEXT;
+		DefaultVertexArrayObject = 0;
+		bIsDebug = false;
+	}
+
+	GLuint					DefaultVertexArrayObject;
+	bool					bIsDebug;
+	EGLContext				eglContext;
+};
+
+struct FPlatformOpenGLSurfaceContext
+{
+public:
+	FPlatformOpenGLSurfaceContext(void* InWindowHandle);
+	FPlatformOpenGLSurfaceContext(TOptional<FAndroidWindow::FNativeAccessor> WindowContainerIn);
+
 	GLuint		ViewportFramebuffer;
-	EGLSurface	eglSurface;
-	GLuint		DefaultVertexArrayObject;
+
 	GLuint		BackBufferResource;
 	GLenum		BackBufferTarget;
-	GLuint		DummyFrameBuffer;
+	FPlatformRect  CachedWindowRect{};
 
-	FPlatformOpenGLContext()
+	bool TargetDirty = false;
+
+	// EGLProperties
+	EGLint         eglFormat = -1;
+	EGLSurface     eglSurface = EGL_NO_SURFACE;
+	EGLint         eglWidth = 8; // required for Gear VR apps with internal win surf mgmt
+	EGLint         eglHeight = 8; // required for Gear VR apps with internal win surf mgmt
+	float          eglRatio = 0;
+	ANativeWindow* Window = nullptr;
+	bool           bIsWndSurface = false; // True when the surface is attached to a HW window.
+
+	EGLSurface GetSurface() const { return eglSurface; }
+
+private:
+	// android window properties:
+	FAndroidWindow* AndroidWindow = nullptr;
+
+public:
+	FAndroidWindow* GetAndroidWindow() { return AndroidWindow; }
+
+	bool IsUsingWindowedSurface() const { return bIsWndSurface; }
+
+	FPlatformOpenGLSurfaceContext()
 	{
 		Reset();
 	}
 
 	void Reset()
 	{
-		eglContext = EGL_NO_CONTEXT;
-		eglSurface = EGL_NO_SURFACE;
 		ViewportFramebuffer = 0;
-		DefaultVertexArrayObject = 0;
 		BackBufferResource = 0;
 		BackBufferTarget = 0;
-		DummyFrameBuffer = 0;
+		CachedWindowRect = FPlatformRect();
+		TargetDirty = false;
+		eglFormat = -1;
+		eglSurface = EGL_NO_SURFACE;
+		eglWidth = 8; // required for Gear VR apps with internal win surf mgmt
+		eglHeight = 8; // required for Gear VR apps with internal win surf mgmt
+		eglRatio = 0;
+		Window = nullptr;
+		AndroidWindow = nullptr;
+		bIsWndSurface = false; // True when the surface is attached to a HW window.
 	}
 };
 
-
 class AndroidEGL
 {
 public:
@@ -71,23 +173,23 @@ public:
 
 	bool IsInitialized();
 
-	void InitBackBuffer();
-	void DestroyBackBuffer();
+	void InitBackBuffer(struct FPlatformOpenGLSurfaceContext* Surface);
+	void DestroyBackBuffer(struct FPlatformOpenGLSurfaceContext* Surface);
 
 	void Init(APIVariant API, uint32 MajorVersion, uint32 MinorVersion);
 
-	void UnBindRender();
+	void UnBindRender(struct FPlatformOpenGLSurfaceContext* Surface);
+	void OnLostNativeWindow(FPlatformOpenGLSurfaceContext* Surface);
 
 	void Terminate();
 
 	// Initialize the surface of the context.
 	// when using the new window locking behavior, if WindowContainer is not set an off screen surface will be used.
 	// TODO: remove bCreateWndSurface (and possibly bUseSmallSurface), these can be inferred from WindowContainer.
-	void InitRenderSurface(bool bUseSmallSurface, bool bCreateWndSurface, const TOptional<FAndroidWindow::FNativeAccessor>& WindowContainer);
+	void InitRenderSurface(bool bUseSmallSurface, bool bCreateWndSurface, FPlatformOpenGLSurfaceContext* Surface);
 	void UpdateBuffersTransform();
 	bool IsOfflineSurfaceRequired();
 
-	void GetDimensions(uint32& OutWidth, uint32& OutHeight);
 	bool IsUsingRobustContext() const { return bIsEXTRobustContextActive; }
 
 	EGLDisplay GetDisplay() const;
@@ -98,15 +200,20 @@ public:
 
 	EGLContext CreateContext(EGLContext InParentContext);
 	int32 GetError();
-	EGLBoolean SetCurrentContext(EGLContext InContext, EGLSurface InSurface);
+	EGLBoolean SetCurrentContext(EGLContext InContext, FPlatformOpenGLSurfaceContext* InSurfaceInfo);
 
 	void AcquireCurrentRenderingContext();
 	void ReleaseContextOwnership();
 
-	GLuint GetResolveFrameBuffer();
+	// Methods to control rendering blocking during window recreation
+	void SetWaitingForWindowCreated(bool bWaiting) { bWaitingForWindowCreated.store(bWaiting); }
+	bool IsWaitingForWindowCreated() const { return bWaitingForWindowCreated.load(); }
+
 	EGLContext GetCurrentContext();
+	EGLSurface GetCurrentSurface(EGLint readdraw);
 
-	void SetCurrentRenderingContext();
+
+	void SetCurrentRenderingContext(FPlatformOpenGLSurfaceContext* InSurface);
 	bool ThreadHasRenderingContext();
 	FPlatformOpenGLContext* GetRenderingContext();
 
@@ -121,6 +228,10 @@ public:
 	// true if the current surface is associated with a window.
 	bool IsUsingWindowedSurface() const;
 
+	FPlatformOpenGLSurfaceContext* CreateSurfaceContextFromWindow(FAndroidWindow* InWindow);
+	FPlatformOpenGLSurfaceContext* FindSurfaceContextFromWindow(FAndroidWindow* InWindow);
+	void DestroySurfaceContext(FPlatformOpenGLSurfaceContext* SurfaceContext);
+
 protected:
 	AndroidEGL();
 	static AndroidEGL* Singleton;
@@ -129,8 +240,8 @@ private:
 	void InitEGL(APIVariant API);
 	void TerminateEGL();
 
-	void CreateEGLRenderSurface(ANativeWindow* InWindow, bool bCreateWndSurface);
-	void DestroyRenderSurface();
+	void CreateEGLRenderSurface(FPlatformOpenGLSurfaceContext* Surface, bool bCreateWndSurface);
+	void DestroyRenderSurface(FPlatformOpenGLSurfaceContext* Surface);
 
 	bool InitContexts();
 	void DestroyContext(EGLContext InContext);
@@ -143,7 +254,7 @@ private:
 	void LogConfigInfo(EGLConfig  EGLConfigInfo);
 
 	// Actual Update to the egl surface to match the GT's requested size.
-	void ResizeRenderContextSurface(const TOptional<FAndroidWindow::FNativeAccessor>& WindowContainer);
+	void ResizeRenderContextSurface(FPlatformOpenGLSurfaceContext* Surface);
 
 	bool bSupportsKHRCreateContext      = false;
 	bool bSupportsKHRSurfacelessContext = false;
@@ -151,6 +262,9 @@ private:
 	bool bSupportsEXTRobustContext      = false;
 	bool bIsEXTRobustContextActive      = false;
 
+	// Flag to block rendering while waiting for APP_EVENT_STATE_WINDOW_CREATED to be processed
+	std::atomic<bool> bWaitingForWindowCreated{false};
+
 	int* ContextAttributes = nullptr;
 };
 
diff --git a/Engine/Source/Runtime/OpenGLDrv/Private/Android/AndroidOpenGL.cpp b/Engine/Source/Runtime/OpenGLDrv/Private/Android/AndroidOpenGL.cpp
--- a/Engine/Source/Runtime/OpenGLDrv/Private/Android/AndroidOpenGL.cpp
+++ b/Engine/Source/Runtime/OpenGLDrv/Private/Android/AndroidOpenGL.cpp
@@ -46,8 +46,6 @@ static TAutoConsoleVariable<int32> CVarDisableFBFNonCoherent(
 
 struct FPlatformOpenGLDevice
 {
-	bool TargetDirty;
-
 	void SetCurrentRenderingContext();
 
 	FPlatformOpenGLDevice();
@@ -174,8 +172,6 @@ void FPlatformOpenGLDevice::Init()
 
 	FPlatformMisc::LowLevelOutputDebugString(TEXT("FPlatformOpenGLDevice:Init"));
 	bool bCreateSurface = !AndroidThunkCpp_IsXRImmersiveHMDApplication();
-	// With the new window system we may not have a window yet, the surface will be initially off screen.
-	AndroidEGL::GetInstance()->InitRenderSurface(false, bCreateSurface, TOptional<FAndroidWindow::FNativeAccessor>());
 
 	LoadEXT();
 	PlatformRenderingContextSetup(this);
@@ -193,7 +189,21 @@ void FPlatformOpenGLDevice::Init()
 	}
 	InitDefaultGLContextState();
 
-	AndroidEGL::GetInstance()->InitBackBuffer(); //can be done only after context is made current.
+ 	// Get the special case instance of FPlatformOpenGLSurfaceContext, it provides singleton behaviour for the old window method or offscreen surface for the new window method...
+	FPlatformOpenGLSurfaceContext* GLSurfaceContext = AndroidEGL::GetInstance()->FindSurfaceContextFromWindow(nullptr);
+	AndroidEGL::GetInstance()->InitBackBuffer(GLSurfaceContext);
+	if (!FAndroidMisc::UseNewWindowBehavior())
+	{
+		// during rhi init the HW window is locked by the GT and the hardware window is valid to use
+		check(GLSurfaceContext);
+		GLSurfaceContext->Window = (ANativeWindow*)FAndroidWindow::GetHardwareWindow_EventThread();
+		AndroidEGL::GetInstance()->InitRenderSurface(false, true, GLSurfaceContext);
+	}
+	else
+	{
+		// set the offscreen surface.
+		AndroidEGL::GetInstance()->InitRenderSurface(false, false, GLSurfaceContext);
+	}
 
 	OpenGLRemoteGLProgramCompileJNI.Init(FAndroidApplication::GetJavaEnv());
 
@@ -223,7 +233,7 @@ void FPlatformOpenGLDevice::Init()
 	if (FPlatformMisc::UseNewWindowBehavior())
 	{
 		UE::FAndroidPlatformDynamicRHI::SetRHIOnReleaseWindowCallback(
-			[](const TOptional<FAndroidWindow::FNativeAccessor>& WindowContainer) 
+			[](const TOptional<FAndroidWindow::FNativeAccessor>& WindowContainer)
 			{
 				check(IsInAndroidEventThread());
 				UE_LOG(LogRHI, Log, TEXT("OnReleaseWindowCallback event thread"));
@@ -234,11 +244,20 @@ void FPlatformOpenGLDevice::Init()
 						UE_LOG(LogRHI, Log, TEXT("OnReleaseWindowCallback: RT"));
 						RHICmdList.EnqueueLambda([&WindowContainer, OnComplete](FRHICommandListImmediate& ExecutingCmdList) mutable
 							{
+								// find the viewport using this window.
 								UE_LOG(LogRHI, Log, TEXT("GLES OnReleaseWindowCallback: RHI, set egl surface to offscreen (window %p)"), WindowContainer->GetANativeWindow());
-								AndroidEGL::GetInstance()->UnBindRender();
-								AndroidEGL::GetInstance()->InitRenderSurface(true, false, WindowContainer);
-								AndroidEGL::GetInstance()->SetCurrentRenderingContext();
-								UE_LOG(LogRHI, Log, TEXT("GLES OnReleaseWindowCallback: RHI, swap surface done."));
+								int ViewportsAffected = 0;
+								FOpenGLDynamicRHI::FGLViewportContainer Viewports = FOpenGLDynamicRHI::Get().GetViewportContainer();
+								for (FOpenGLViewport* Viewport : Viewports.Get())
+								{
+									if (Viewport->GetWindowHandle() == WindowContainer->Get())
+									{
+										ViewportsAffected++;
+										FPlatformOpenGLSurfaceContext* Surface = Viewport->GetGLSurfaceContext();
+										AndroidEGL::GetInstance()->OnLostNativeWindow(Surface);
+									}
+								}
+								UE_LOG(LogRHI, Log, TEXT("GLES OnReleaseWindowCallback: RHI, swap surface done. (%d viewports)"), ViewportsAffected);
 							});
 						RHICmdList.RHIThreadFence(true);
 						RHICmdList.ImmediateFlush(EImmediateFlushType::FlushRHIThread);
@@ -253,7 +272,13 @@ void FPlatformOpenGLDevice::Init()
 		UE::FAndroidPlatformDynamicRHI::SetRHIOnReInitWindowCallback(
 			[](const TOptional<FAndroidWindow::FNativeAccessor>& WindowContainer)
 			{
-				UE_LOG(LogRHI, Log, TEXT("OnReInitWindowCallback event thread"));
+				check(IsInAndroidEventThread());
+
+				UE_LOG(LogRHI, Log, TEXT("OnReInitWindowCallback event thread, blocking rendering until WINDOW_CREATED is processed"));
+
+				// Block rendering until APP_EVENT_STATE_WINDOW_CREATED is processed
+				AndroidEGL::GetInstance()->SetWaitingForWindowCreated(true);
+
 				FGraphEventRef OnComplete = FGraphEvent::CreateGraphEvent();
 
 				ENQUEUE_RENDER_COMMAND(OnAndroidFoundWindow)([&WindowContainer, OnComplete](FRHICommandListImmediate& RHICmdList) mutable
@@ -276,6 +301,17 @@ void FPlatformOpenGLDevice::Init()
 				OnComplete->Wait();
 				UE_LOG(LogRHI, Log, TEXT("OnReInitWindowCallback: done"));
 			});
+
+		FPlatformDynamicRHI::SetRHIOnWindowCreatedCallback(
+			[]()
+			{
+				// Unblock rendering now that viewport has been resized to match the new surface
+				if (AndroidEGL* EGL = AndroidEGL::GetInstance())
+				{
+					EGL->SetWaitingForWindowCreated(false);
+					UE_LOG(LogRHI, Display, TEXT("OnWindowCreatedProcessed: rendering unblocked"));
+				}
+			});
 	}
 }
 
@@ -305,18 +341,38 @@ void* PlatformGetWindow(FPlatformOpenGLContext* Context, void** AddParam)
 bool PlatformBlitToViewport(IRHICommandContext& RHICmdContext, FPlatformOpenGLDevice* Device, FOpenGLViewport& Viewport, uint32 BackbufferSizeX, uint32 BackbufferSizeY, bool bPresent,bool bLockToVsync )
 {
 	SCOPED_NAMED_EVENT(STAT_PlatformBlitToViewportTime, FColor::Red)
-	
+
+	if (FAndroidMisc::UseNewWindowBehavior())
+	{
+		if (AndroidEGL::GetInstance()->IsWaitingForWindowCreated())
+		{
+			UE_LOG(LogRHI, Log, TEXT("PlatformBlitToViewport: skipping frame, waiting for WINDOW_CREATED"));
+			return false;
+		}
+	}
+
 	FPlatformOpenGLContext* const Context = Viewport.GetGLContext();
 
+	FPlatformOpenGLSurfaceContext* Surface = Viewport.GetGLSurfaceContext();
+
+	// the context must be reapplied if the surface has changed
+	if (AndroidEGL::GetInstance()->SetCurrentContext(Context->eglContext, Surface) != EGL_TRUE)
+	{
+		UE_LOG(LogRHI, Error, TEXT("PlatformBlitToViewport: can't set current context"));
+		return false;
+	}
+
+
 	if (bPresent && AndroidEGL::GetInstance()->IsOfflineSurfaceRequired())
 	{
-		if (Device->TargetDirty)
+		check(Surface->ViewportFramebuffer);
+		if (Surface->TargetDirty)
 		{
 			VERIFY_GL_SCOPE();
-			glBindFramebuffer(GL_FRAMEBUFFER, Context->ViewportFramebuffer);
-			glFramebufferTexture2D(GL_FRAMEBUFFER, GL_COLOR_ATTACHMENT0, Context->BackBufferTarget, Context->BackBufferResource, 0);
+			glBindFramebuffer(GL_FRAMEBUFFER, Surface->ViewportFramebuffer);
+			glFramebufferTexture2D(GL_FRAMEBUFFER, GL_COLOR_ATTACHMENT0, Surface->BackBufferTarget, Surface->BackBufferResource, 0);
 
-			Device->TargetDirty = false;
+			Surface->TargetDirty = false;
 		}
 
 		{
@@ -324,7 +380,7 @@ bool PlatformBlitToViewport(IRHICommandContext& RHICmdContext, FPlatformOpenGLDe
 			glDisable(GL_FRAMEBUFFER_SRGB);
 			glBindFramebuffer(GL_DRAW_FRAMEBUFFER, 0);
 			FOpenGL::DrawBuffer(GL_BACK);
-			glBindFramebuffer(GL_READ_FRAMEBUFFER, Context->ViewportFramebuffer);
+			glBindFramebuffer(GL_READ_FRAMEBUFFER, Surface->ViewportFramebuffer);
 			FOpenGL::ReadBuffer(GL_COLOR_ATTACHMENT0);
 
 			FOpenGL::BlitFramebuffer(
@@ -337,7 +393,7 @@ bool PlatformBlitToViewport(IRHICommandContext& RHICmdContext, FPlatformOpenGLDe
 			glEnable(GL_FRAMEBUFFER_SRGB);
 
 			// Bind viewport FBO so driver knows we don't need backbuffer image anymore
-			glBindFramebuffer(GL_FRAMEBUFFER, Context->ViewportFramebuffer);
+			glBindFramebuffer(GL_FRAMEBUFFER, Surface->ViewportFramebuffer);
 		}
 	}
 
@@ -350,7 +406,7 @@ bool PlatformBlitToViewport(IRHICommandContext& RHICmdContext, FPlatformOpenGLDe
 	if (bPresent)
 	{
 		AndroidEGL::GetInstance()->UpdateBuffersTransform();
-		FAndroidPlatformRHIFramePacer::SwapBuffers(bLockToVsync);
+		FAndroidPlatformRHIFramePacer::SwapBuffers(Surface, bLockToVsync);
 	}
 	static IConsoleVariable* CVar = IConsoleManager::Get().FindConsoleVariable(TEXT("a.UseFrameTimeStampsForPacing"));
 	const bool bForceGPUFence = CVar ? CVar->GetInt() != 0 : false;
@@ -509,16 +565,6 @@ void FPlatformOpenGLDevice::LoadEXT()
 FPlatformOpenGLContext* PlatformCreateOpenGLContext(FPlatformOpenGLDevice* Device, void* InWindowHandle)
 {
 	AndroidEGL* AnEGL = AndroidEGL::GetInstance();
-	if (FAndroidMisc::UseNewWindowBehavior())
-	{
-		check(InWindowHandle);
-		FAndroidWindow* AndroidWindow = (FAndroidWindow *)InWindowHandle;
-
-		FAndroidWindow::FNativeAccessor Accessor = AndroidWindow->GetANativeAccessor(false);
-
-		AnEGL->InitRenderSurface(false, true, MoveTemp(Accessor));
-	}
-	// else - with the original window method we block on startup until the window + device are initialized and context fully created.
 	return AnEGL->GetRenderingContext();
 }
 
@@ -526,7 +572,41 @@ void PlatformDestroyOpenGLContext(FPlatformOpenGLDevice* Device, FPlatformOpenGL
 {
 }
 
-FOpenGLTexture* PlatformCreateBuiltinBackBuffer(FOpenGLDynamicRHI* OpenGLRHI, uint32 SizeX, uint32 SizeY)
+FPlatformOpenGLSurfaceContext* PlatformCreateOpenGLSurfaceContext(FPlatformOpenGLContext* SourceContext, void* InWindowHandle)
+{
+	check(IsInGameThread());
+	// create the surface context for the window.
+	AndroidEGL* AnEGL = AndroidEGL::GetInstance();
+	FAndroidWindow* AndroidWindow = (FAndroidWindow*)InWindowHandle;
+
+	FPlatformOpenGLSurfaceContext* Surface = AnEGL->CreateSurfaceContextFromWindow(AndroidWindow);
+	ENQUEUE_RENDER_COMMAND(CreateOpenGLSurfaceContext)([Surface](FRHICommandListImmediate& RHICmdList)
+		{
+			RHICmdList.EnqueueLambda([Surface](FRHICommandListImmediate&)
+				{
+					AndroidEGL* AnEGL = AndroidEGL::GetInstance();
+					AnEGL->InitBackBuffer(Surface);
+					AnEGL->InitRenderSurface(false, true, Surface);
+				});
+
+		});
+
+	// waiting for RHI, surface creation/destruction is not expected to be frequent.
+	FlushRenderingCommands();
+
+	return Surface;
+}
+
+void PlatformDestroyOpenGLSurfaceContext(FPlatformOpenGLDevice* Device, FPlatformOpenGLContext* Context, FPlatformOpenGLSurfaceContext* SurfaceContext)
+{
+	VERIFY_GL_SCOPE();
+
+	AndroidEGL* AnEGL = AndroidEGL::GetInstance();
+	AnEGL->DestroySurfaceContext(SurfaceContext);
+}
+
+
+FOpenGLTexture* PlatformCreateBuiltinBackBuffer(FOpenGLDynamicRHI* OpenGLRHI, FPlatformOpenGLSurfaceContext* SurfaceContext, uint32 SizeX, uint32 SizeY)
 {
 	check(IsInRenderingThread());
 	// Create the built-in back buffer if we disable backbuffer sampling.
@@ -549,23 +629,29 @@ FOpenGLTexture* PlatformCreateBuiltinBackBuffer(FOpenGLDynamicRHI* OpenGLRHI, ui
 	}
 }
 
-void PlatformResizeGLContext( FPlatformOpenGLDevice* Device, FPlatformOpenGLContext* Context, uint32 SizeX, uint32 SizeY, bool bFullscreen, bool bWasFullscreen, GLenum BackBufferTarget, GLuint BackBufferResource)
+void PlatformResizeGLContext( FPlatformOpenGLDevice* Device, FPlatformOpenGLContext* Context, FPlatformOpenGLSurfaceContext* SurfaceContext, uint32 SizeX, uint32 SizeY, bool bFullscreen, bool bWasFullscreen, GLenum BackBufferTarget, GLuint BackBufferResource)
 {
 	check(Context);
 	VERIFY_GL_SCOPE();
 
-	Context->BackBufferResource = BackBufferResource;
-	Context->BackBufferTarget = BackBufferTarget;
+	SurfaceContext->BackBufferResource = BackBufferResource;
+	SurfaceContext->BackBufferTarget = BackBufferTarget;
 
 	if (AndroidEGL::GetInstance()->IsOfflineSurfaceRequired())
 	{
-		Device->TargetDirty = true;
-		glBindFramebuffer(GL_FRAMEBUFFER, Context->ViewportFramebuffer);
-		glFramebufferTexture2D(GL_FRAMEBUFFER, GL_COLOR_ATTACHMENT0, BackBufferTarget, BackBufferResource, 0);
+		SurfaceContext->TargetDirty = true;
+		check(SurfaceContext->ViewportFramebuffer);
+		glBindFramebuffer(GL_FRAMEBUFFER, SurfaceContext->ViewportFramebuffer);
+		glFramebufferTexture2D(GL_FRAMEBUFFER, GL_COLOR_ATTACHMENT0, SurfaceContext->BackBufferTarget, SurfaceContext->BackBufferResource, 0);
 	}
 
-	glViewport(0, 0, SizeX, SizeY);
-	VERIFY_GL(glViewport);
+#if USE_ANDROID_STANDALONE
+	if (!FAndroidMisc::UseNewWindowBehavior())
+#endif
+	{
+		glViewport(0, 0, SizeX, SizeY);
+		VERIFY_GL(glViewport);
+	}
 }
 
 void PlatformGetSupportedResolution(uint32 &Width, uint32 &Height)
diff --git a/Engine/Source/Runtime/OpenGLDrv/Private/Android/AndroidOpenGLFramePacer.cpp b/Engine/Source/Runtime/OpenGLDrv/Private/Android/AndroidOpenGLFramePacer.cpp
--- a/Engine/Source/Runtime/OpenGLDrv/Private/Android/AndroidOpenGLFramePacer.cpp
+++ b/Engine/Source/Runtime/OpenGLDrv/Private/Android/AndroidOpenGLFramePacer.cpp
@@ -240,8 +240,9 @@ bool FAndroidOpenGLFramePacer::SupportsFramePace(int32 QueryFramePace)
 }
 
 extern bool AndroidThunkCpp_IsXRImmersiveHMDApplication();
+extern bool AndroidThunkCpp_IsStandaloneApplication();
 
-bool FAndroidOpenGLFramePacer::SwapBuffers(bool bLockToVsync)
+bool FAndroidOpenGLFramePacer::SwapBuffers(FPlatformOpenGLSurfaceContext* SurfaceContext, bool bLockToVsync)
 {
 	SCOPED_NAMED_EVENT(STAT_OpenGLSwapBuffersTime, FColor::Red)
 
@@ -257,10 +258,10 @@ bool FAndroidOpenGLFramePacer::SwapBuffers(bool bLockToVsync)
 
 	AndroidEGL* AnEGL = AndroidEGL::GetInstance();
 	EGLDisplay eglDisplay = AnEGL->GetDisplay();
-	EGLSurface eglSurface = AnEGL->GetSurface();
+	EGLSurface eglSurface = SurfaceContext->eglSurface;
 	int32 SyncInterval = FAndroidPlatformRHIFramePacer::GetLegacySyncInterval();
 
-	if (AnEGL->IsUsingWindowedSurface() == false && !AndroidThunkCpp_IsXRImmersiveHMDApplication())
+	if (SurfaceContext->IsUsingWindowedSurface() == false && !AndroidThunkCpp_IsXRImmersiveHMDApplication() && !AndroidThunkCpp_IsStandaloneApplication())
 	{
 		// offscreen surface, no swap to perform, limit update rate.
 		FPlatformProcess::Sleep(0.1);
@@ -273,7 +274,7 @@ bool FAndroidOpenGLFramePacer::SwapBuffers(bool bLockToVsync)
 	int32 CurrentFramePace = 0;
 	if (FAndroidPlatformRHIFramePacer::CVarUseSwappyForFramePacing.GetValueOnRenderThread() != 0 && (CurrentFramePace = FAndroidPlatformRHIFramePacer::GetFramePace()) != 0 && ensure(bSwappyInit))
 	{
-		ANativeWindow* CurrentNativeWindow = AnEGL->GetNativeWindow();
+		ANativeWindow* CurrentNativeWindow = SurfaceContext->Window;
 		if (CurrentNativeWindow != CachedNativeWindow)
 		{
 			CachedNativeWindow = CurrentNativeWindow;
diff --git a/Engine/Source/Runtime/OpenGLDrv/Private/Android/AndroidOpenGLPlatform.cpp b/Engine/Source/Runtime/OpenGLDrv/Private/Android/AndroidOpenGLPlatform.cpp
--- a/Engine/Source/Runtime/OpenGLDrv/Private/Android/AndroidOpenGLPlatform.cpp
+++ b/Engine/Source/Runtime/OpenGLDrv/Private/Android/AndroidOpenGLPlatform.cpp
@@ -48,7 +48,6 @@ void FAndroidAppEntry::ReleaseEGL()
 	AndroidEGL* EGL = AndroidEGL::GetInstance();
 	if (EGL->IsInitialized())
 	{
-		EGL->DestroyBackBuffer();
 		EGL->Terminate();
 	}
 }
diff --git a/Engine/Source/Runtime/OpenGLDrv/Private/Android/AndroidOpenGLPrivate.cpp b/Engine/Source/Runtime/OpenGLDrv/Private/Android/AndroidOpenGLPrivate.cpp
--- a/Engine/Source/Runtime/OpenGLDrv/Private/Android/AndroidOpenGLPrivate.cpp
+++ b/Engine/Source/Runtime/OpenGLDrv/Private/Android/AndroidOpenGLPrivate.cpp
@@ -23,12 +23,9 @@ FAndroidGPUInfo::FAndroidGPUInfo()
 		FAndroidAppEntry::PlatformInit();
 	}
 
-	// Do not create a window surface if the app is for XRImmersiveHMD (use small buffer)
-	bool bCreateSurface = !AndroidThunkCpp_IsXRImmersiveHMDApplication();
 	FPlatformMisc::LowLevelOutputDebugString(TEXT("FAndroidGPUInfo"));
 
-	EGL->InitRenderSurface(false, bCreateSurface, TOptional<FAndroidWindow::FNativeAccessor>());
-	EGL->SetCurrentRenderingContext();
+	EGL->SetCurrentRenderingContext(nullptr); // a windowed surface is not required here.
 
 	// get extensions
 	// Process the extension caps directly here, as FOpenGL might not yet be setup
diff --git a/Engine/Source/Runtime/OpenGLDrv/Private/Linux/LinuxOpenGLPlatform.cpp b/Engine/Source/Runtime/OpenGLDrv/Private/Linux/LinuxOpenGLPlatform.cpp
--- a/Engine/Source/Runtime/OpenGLDrv/Private/Linux/LinuxOpenGLPlatform.cpp
+++ b/Engine/Source/Runtime/OpenGLDrv/Private/Linux/LinuxOpenGLPlatform.cpp
@@ -46,6 +46,9 @@ struct FPlatformOpenGLContext
 	GLuint	VertexArrayObject;	// one has to be generated and set for each context (OpenGL 3.2 Core requirements)
 };
 
+struct FPlatformOpenGLSurfaceContext
+{
+};
 
 class FScopeContext
 {
@@ -272,6 +275,15 @@ FPlatformOpenGLContext* PlatformCreateOpenGLContext(FPlatformOpenGLDevice* Devic
 	return Context;
 }
 
+struct FPlatformOpenGLSurfaceContext* PlatformCreateOpenGLSurfaceContext(FPlatformOpenGLContext* SourceContext, void* InWindowHandle)
+{
+	return nullptr;
+}
+
+void PlatformDestroyOpenGLSurfaceContext(FPlatformOpenGLDevice* Device, FPlatformOpenGLContext* Context, FPlatformOpenGLSurfaceContext* SurfaceContext)
+{
+}
+
 /**
  * Release an OpenGL context.
  */
@@ -536,6 +548,7 @@ void PlatformNULLContextSetup()
  */
 void PlatformResizeGLContext(FPlatformOpenGLDevice* Device,
 	FPlatformOpenGLContext* Context,
+	struct FPlatformOpenGLSurfaceContext* SurfaceContext,
 	uint32 SizeX, uint32 SizeY,
 	bool bFullscreen,
 	bool bWasFullscreen,
@@ -599,6 +612,8 @@ void PlatformGetSupportedResolution(uint32& Width, uint32& Height)
 			}
 
 		}
+
+		SDL_free(Modes);
 	}
 	check(InitializedMode);
 	Width = BestWidth;
@@ -816,7 +831,7 @@ bool PlatformOpenGLThreadHasRenderingContext()
 
 // =============================================================
 
-FOpenGLTexture* PlatformCreateBuiltinBackBuffer(FOpenGLDynamicRHI* OpenGLRHI, uint32 SizeX, uint32 SizeY)
+FOpenGLTexture* PlatformCreateBuiltinBackBuffer(FOpenGLDynamicRHI* OpenGLRHI, FPlatformOpenGLSurfaceContext* SurfaceContext, uint32 SizeX, uint32 SizeY)
 {
 	return nullptr;
 }
diff --git a/Engine/Source/Runtime/OpenGLDrv/Private/OpenGL3.h b/Engine/Source/Runtime/OpenGLDrv/Private/OpenGL3.h
--- a/Engine/Source/Runtime/OpenGLDrv/Private/OpenGL3.h
+++ b/Engine/Source/Runtime/OpenGLDrv/Private/OpenGL3.h
@@ -14,6 +14,7 @@
 #include "OpenGLBase.h"
 
 struct FPlatformOpenGLContext;
+struct FPlatformOpenGLSurfaceContext;
 struct FPlatformOpenGLDevice;
 
 #define USE_OPENGL_NAME_CACHE 1
diff --git a/Engine/Source/Runtime/OpenGLDrv/Private/OpenGLDrv.cpp b/Engine/Source/Runtime/OpenGLDrv/Private/OpenGLDrv.cpp
--- a/Engine/Source/Runtime/OpenGLDrv/Private/OpenGLDrv.cpp
+++ b/Engine/Source/Runtime/OpenGLDrv/Private/OpenGLDrv.cpp
@@ -96,17 +96,17 @@ bool FOpenGLDynamicRHI::RHIEGLSupportsNoErrorContext() const
 void FOpenGLDynamicRHI::RHIInitEGLInstanceGLES2()
 {
 	AndroidEGL::GetInstance()->Init(AndroidEGL::AV_OpenGLES, 2, 0);
-	AndroidEGL::GetInstance()->InitRenderSurface(false, false, TOptional<FAndroidWindow::FNativeAccessor>());
+	AndroidEGL::GetInstance()->InitRenderSurface(false, false, nullptr);
 }
 
 void FOpenGLDynamicRHI::RHIInitEGLBackBuffer()
 {
-	AndroidEGL::GetInstance()->InitBackBuffer();
+	AndroidEGL::GetInstance()->InitBackBuffer(nullptr);
 }
 
 void FOpenGLDynamicRHI::RHIEGLSetCurrentRenderingContext()
 {
-	AndroidEGL::GetInstance()->SetCurrentRenderingContext();
+	AndroidEGL::GetInstance()->SetCurrentRenderingContext(nullptr);
 }
 
 void FOpenGLDynamicRHI::RHIEGLTerminateContext()
@@ -222,6 +222,11 @@ bool FOpenGLDynamicRHI::RHIMatchPrecachePSOInitializers(const FGraphicsPipelineS
 	return true;
 }
 
+FOpenGLDynamicRHI::FGLViewportContainer FOpenGLDynamicRHI::GetViewportContainer()
+{
+	return FGLViewportContainer(Viewports);
+}
+
 void FOpenGLDynamicRHI::InitializeStateResources()
 {
 	VERIFY_GL_SCOPE();
diff --git a/Engine/Source/Runtime/OpenGLDrv/Private/OpenGLDrv.h b/Engine/Source/Runtime/OpenGLDrv/Private/OpenGLDrv.h
--- a/Engine/Source/Runtime/OpenGLDrv/Private/OpenGLDrv.h
+++ b/Engine/Source/Runtime/OpenGLDrv/Private/OpenGLDrv.h
@@ -446,6 +446,32 @@ public:
 	
 	UE_API void ResolveTexture(FOpenGLTexture* Texture, uint32 MipIndex, uint32 ArrayIndex);
 
+	class FGLViewportContainer;
+ 	FGLViewportContainer GetViewportContainer();
+
+	class FGLViewportContainer
+	{
+		friend FGLViewportContainer FOpenGLDynamicRHI::GetViewportContainer();
+
+		static inline UE::FMutex ViewportMutex;
+		UE::TUniqueLock<UE::FMutex> ScopeLock;
+		TArray<FOpenGLViewport*>& Viewports;
+
+		FGLViewportContainer(TArray<FOpenGLViewport*>& ViewportsIn) : ScopeLock(ViewportMutex), Viewports(ViewportsIn)
+		{
+		}
+
+	public:
+		~FGLViewportContainer()
+		{
+		}
+
+		TArray<FOpenGLViewport*>& Get()
+		{
+			return Viewports;
+		}
+	};
+
 private:
 	UE_API FBoundShaderStateRHIRef RHICreateBoundShaderState_Internal(FRHIVertexDeclaration* VertexDeclaration, FRHIVertexShader* VertexShader, FRHIPixelShader* PixelShader, FRHIGeometryShader* GeometryShader, bool FromPSOFileCache);
 
diff --git a/Engine/Source/Runtime/OpenGLDrv/Private/OpenGLDrvPrivate.h b/Engine/Source/Runtime/OpenGLDrv/Private/OpenGLDrvPrivate.h
--- a/Engine/Source/Runtime/OpenGLDrv/Private/OpenGLDrvPrivate.h
+++ b/Engine/Source/Runtime/OpenGLDrv/Private/OpenGLDrvPrivate.h
@@ -171,6 +171,17 @@ FPlatformOpenGLContext* PlatformCreateOpenGLContext(FPlatformOpenGLDevice* Devic
  */
 void PlatformDestroyOpenGLContext(FPlatformOpenGLDevice* Device, FPlatformOpenGLContext* Context);
 
+/**
+ * Create an FPlatformOpenGLSurfaceContext that contains properties required for an individual output surface
+ * attached to the FPlatformOpenGLContext.
+ */
+struct FPlatformOpenGLSurfaceContext* PlatformCreateOpenGLSurfaceContext(FPlatformOpenGLContext* SourceContext, void* InWindowHandle);
+
+/**
+ * destroy the FPlatformOpenGLSurfaceContext
+ */
+void PlatformDestroyOpenGLSurfaceContext(FPlatformOpenGLDevice* Device, FPlatformOpenGLContext* Context, FPlatformOpenGLSurfaceContext* SurfaceContext);
+
 /**
  * Set up rendering OpenGL context on current thread. This is the only context that can be used for rendering,
  * and not only loading resources and framebuffer blitting (due to limitations of Windows OpenGL).
@@ -186,7 +197,7 @@ void PlatformFlushIfNeeded();
 void PlatformNULLContextSetup();
 
 // Creates a platform-specific back buffer
-class FOpenGLTexture* PlatformCreateBuiltinBackBuffer(FOpenGLDynamicRHI* OpenGLRHI, uint32 SizeX, uint32 SizeY);
+class FOpenGLTexture* PlatformCreateBuiltinBackBuffer(FOpenGLDynamicRHI* OpenGLRHI, struct FPlatformOpenGLSurfaceContext* SurfaceContext, uint32 SizeX, uint32 SizeY);
 
 /**
  * Main function for transferring data to on-screen buffers.
@@ -198,7 +209,7 @@ bool PlatformBlitToViewport(IRHICommandContext& RHICmdContext, FPlatformOpenGLDe
 /**
  * Resize the GL context for platform.
  */
-void PlatformResizeGLContext( FPlatformOpenGLDevice* Device, FPlatformOpenGLContext* Context, uint32 SizeX, uint32 SizeY, bool bFullscreen, bool bWasFullscreen, GLenum BackBufferTarget, GLuint BackBufferResource);
+void PlatformResizeGLContext( FPlatformOpenGLDevice* Device, FPlatformOpenGLContext* Context, struct FPlatformOpenGLSurfaceContext* SurfaceContext, uint32 SizeX, uint32 SizeY, bool bFullscreen, bool bWasFullscreen, GLenum BackBufferTarget, GLuint BackBufferResource);
 
 
 /**
diff --git a/Engine/Source/Runtime/OpenGLDrv/Private/OpenGLResources.h b/Engine/Source/Runtime/OpenGLDrv/Private/OpenGLResources.h
--- a/Engine/Source/Runtime/OpenGLDrv/Private/OpenGLResources.h
+++ b/Engine/Source/Runtime/OpenGLDrv/Private/OpenGLResources.h
@@ -1366,17 +1366,23 @@ public:
 	struct FPlatformOpenGLContext* GetGLContext() const { return OpenGLContext; }
 	FOpenGLDynamicRHI* GetOpenGLRHI() const { return OpenGLRHI; }
 
+	struct FPlatformOpenGLSurfaceContext* GetGLSurfaceContext() const { return OpenGLSurfaceContext; }
+
 	virtual void SetCustomPresent(FRHICustomPresent* InCustomPresent) override
 	{
 		CustomPresent = InCustomPresent;
 	}
 	FRHICustomPresent* GetCustomPresent() const { return CustomPresent.GetReference(); }
+
+	void* GetWindowHandle() { return WindowHandle; }
+
 private:
 
 	friend class FOpenGLDynamicRHI;
 
 	FOpenGLDynamicRHI* OpenGLRHI;
 	struct FPlatformOpenGLContext* OpenGLContext;
+	FPlatformOpenGLSurfaceContext* OpenGLSurfaceContext;
 	uint32 SizeX;
 	uint32 SizeY;
 	bool bIsFullscreen;
@@ -1385,6 +1391,7 @@ private:
 	TRefCountPtr<FOpenGLTexture> BackBuffer;
 	TUniquePtr<FOpenGLEventQuery> FrameSyncEvent;
 	FCustomPresentRHIRef CustomPresent;
+	void* WindowHandle;
 };
 
 class FOpenGLGPUFence final : public FRHIGPUFence
diff --git a/Engine/Source/Runtime/OpenGLDrv/Private/OpenGLViewport.cpp b/Engine/Source/Runtime/OpenGLDrv/Private/OpenGLViewport.cpp
--- a/Engine/Source/Runtime/OpenGLDrv/Private/OpenGLViewport.cpp
+++ b/Engine/Source/Runtime/OpenGLDrv/Private/OpenGLViewport.cpp
@@ -11,6 +11,9 @@
 #include "RHI.h"
 #include "OpenGLDrv.h"
 #include "OpenGLDrvPrivate.h"
+#if PLATFORM_ANDROID
+#include "Android/AndroidPlatformMisc.h"
+#endif
 
 static int32 GOpenGLPerFrameErrorCheck = 1;
 static FAutoConsoleVariableRef CVarPerFrameGLErrorCheck(
@@ -179,6 +182,7 @@ FOpenGLViewport::FOpenGLViewport(FOpenGLDynamicRHI* InOpenGLRHI,void* InWindowHa
 	, bIsFullscreen(false)
 	, PixelFormat(PreferredPixelFormat)
 	, bIsValid(true)
+	, WindowHandle(InWindowHandle)
 {
 	check(OpenGLRHI);
 #if !PLATFORM_ANDROID
@@ -187,13 +191,35 @@ FOpenGLViewport::FOpenGLViewport(FOpenGLDynamicRHI* InOpenGLRHI,void* InWindowHa
 	check(IsInGameThread());
 
 	// flush out old errors.
-	PlatformGlGetError();	
+	PlatformGlGetError();
+
+	OpenGLContext = PlatformCreateOpenGLContext(OpenGLRHI->PlatformDevice, WindowHandle);
 
-	OpenGLRHI->Viewports.Add(this);
+#if PLATFORM_ANDROID
+	if (!FAndroidMisc::UseNewWindowBehavior())
+	{
+		WindowHandle = nullptr; // The original window behaviour received a null handle.
+		OpenGLSurfaceContext = AndroidEGL::GetInstance()->FindSurfaceContextFromWindow(nullptr); // there's only one window in this mode.
+#if DO_CHECK
+		FOpenGLDynamicRHI::FGLViewportContainer Viewports = FOpenGLDynamicRHI::Get().GetViewportContainer();
+		// multiple viewports are not supported..
+		check(Viewports.Get().IsEmpty());
+#endif
+	}
+	else
+#endif
+	{
+		check(WindowHandle);
+		OpenGLSurfaceContext = PlatformCreateOpenGLSurfaceContext(OpenGLContext, WindowHandle);
+	}
 
-	OpenGLContext = PlatformCreateOpenGLContext(OpenGLRHI->PlatformDevice, InWindowHandle);
 	Resize(InSizeX, InSizeY, bInIsFullscreen);
 
+	{
+		FOpenGLDynamicRHI::FGLViewportContainer Viewports = FOpenGLDynamicRHI::Get().GetViewportContainer();
+		Viewports.Get().Add(this);
+	}
+
 	ENQUEUE_RENDER_COMMAND(CreateFrameSyncEvent)([this](FRHICommandListImmediate& RHICmdList)
 	{
 		RHICmdList.EnqueueLambda([this](FRHICommandListImmediate&)
@@ -207,6 +233,11 @@ FOpenGLViewport::~FOpenGLViewport()
 {
 	VERIFY_GL_SCOPE();
 
+	{
+		FOpenGLDynamicRHI::FGLViewportContainer Viewports = FOpenGLDynamicRHI::Get().GetViewportContainer();
+		Viewports.Get().Remove(this);
+	}
+
 	if (bIsFullscreen)
 	{
 		PlatformRestoreDesktopDisplayMode();
@@ -217,10 +248,10 @@ FOpenGLViewport::~FOpenGLViewport()
 	check(!IsValidRef(BackBuffer));
 
 	FrameSyncEvent = nullptr;
+	PlatformDestroyOpenGLSurfaceContext(OpenGLRHI->PlatformDevice, OpenGLContext, OpenGLSurfaceContext);
 	PlatformDestroyOpenGLContext(OpenGLRHI->PlatformDevice, OpenGLContext);
 
 	OpenGLContext = NULL;
-	OpenGLRHI->Viewports.Remove(this);
 }
 
 void FOpenGLViewport::WaitForFrameEventCompletion()
@@ -257,7 +288,7 @@ void FOpenGLViewport::Resize(uint32 InSizeX,uint32 InSizeY,bool bInIsFullscreen)
 
 		BackBuffer.SafeRelease();	// when the rest of the engine releases it, its framebuffers will be released too (those the engine knows about)
 
-		BackBuffer = PlatformCreateBuiltinBackBuffer(OpenGLRHI, InSizeX, InSizeY);
+		BackBuffer = PlatformCreateBuiltinBackBuffer(OpenGLRHI, OpenGLSurfaceContext, InSizeX, InSizeY);
 		if (!BackBuffer)
 		{
 			const FRHITextureCreateDesc Desc =
@@ -272,7 +303,7 @@ void FOpenGLViewport::Resize(uint32 InSizeX,uint32 InSizeY,bool bInIsFullscreen)
 
 		RHICmdList.EnqueueLambda([this, InSizeX, InSizeY, bInIsFullscreen, bWasFullscreen](FRHICommandListImmediate&)
 		{
-			PlatformResizeGLContext(OpenGLRHI->PlatformDevice, OpenGLContext, InSizeX, InSizeY, bInIsFullscreen, bWasFullscreen, BackBuffer->Target, BackBuffer->GetResource());
+			PlatformResizeGLContext(OpenGLRHI->PlatformDevice, OpenGLContext, OpenGLSurfaceContext, InSizeX, InSizeY, bInIsFullscreen, bWasFullscreen, BackBuffer->Target, BackBuffer->GetResource());
 		});
 	});	
 }
diff --git a/Engine/Source/Runtime/OpenGLDrv/Private/Windows/WindowsOpenGLPlatform.cpp b/Engine/Source/Runtime/OpenGLDrv/Private/Windows/WindowsOpenGLPlatform.cpp
--- a/Engine/Source/Runtime/OpenGLDrv/Private/Windows/WindowsOpenGLPlatform.cpp
+++ b/Engine/Source/Runtime/OpenGLDrv/Private/Windows/WindowsOpenGLPlatform.cpp
@@ -44,6 +44,10 @@ struct FPlatformOpenGLContext
 	GLenum	BackBufferTarget;
 };
 
+struct FPlatformOpenGLSurfaceContext
+{
+};
+
 class FScopeContext
 {
 public:
@@ -356,6 +360,15 @@ FPlatformOpenGLContext* PlatformCreateOpenGLContext(FPlatformOpenGLDevice* Devic
 	return Context;
 }
 
+FPlatformOpenGLSurfaceContext* PlatformCreateOpenGLSurfaceContext(FPlatformOpenGLContext* SourceContext, void* InWindowHandle)
+{
+	return nullptr;
+}
+
+void PlatformDestroyOpenGLSurfaceContext(FPlatformOpenGLDevice* Device, FPlatformOpenGLContext* Context, FPlatformOpenGLSurfaceContext* SurfaceContext)
+{
+}
+
 /**
  * Release an OpenGL context.
  */
@@ -554,7 +567,7 @@ void PlatformNULLContextSetup()
 /**
  * Resize the GL context.
  */
-void PlatformResizeGLContext(FPlatformOpenGLDevice* Device, FPlatformOpenGLContext* Context, uint32 SizeX, uint32 SizeY, bool bFullscreen, bool bWasFullscreen, GLenum BackBufferTarget, GLuint BackBufferResource)
+void PlatformResizeGLContext(FPlatformOpenGLDevice* Device, FPlatformOpenGLContext* Context, FPlatformOpenGLSurfaceContext* SurfaceContext, uint32 SizeX, uint32 SizeY, bool bFullscreen, bool bWasFullscreen, GLenum BackBufferTarget, GLuint BackBufferResource)
 {
 	FScopeLock ScopeLock(Device->ContextUsageGuard);
 	{
@@ -820,7 +833,7 @@ bool PlatformOpenGLThreadHasRenderingContext()
 
 // =============================================================
 
-FOpenGLTexture* PlatformCreateBuiltinBackBuffer(FOpenGLDynamicRHI* OpenGLRHI, uint32 SizeX, uint32 SizeY)
+FOpenGLTexture* PlatformCreateBuiltinBackBuffer(FOpenGLDynamicRHI* OpenGLRHI, FPlatformOpenGLSurfaceContext* SurfaceContext, uint32 SizeX, uint32 SizeY)
 {
 	if (FOpenGL::IsAndroidGLESCompatibilityModeEnabled())
 	{
diff --git a/Engine/Source/Runtime/RHI/Private/Android/AndroidDynamicRHI.cpp b/Engine/Source/Runtime/RHI/Private/Android/AndroidDynamicRHI.cpp
--- a/Engine/Source/Runtime/RHI/Private/Android/AndroidDynamicRHI.cpp
+++ b/Engine/Source/Runtime/RHI/Private/Android/AndroidDynamicRHI.cpp
@@ -117,6 +117,18 @@ namespace UE
 			OnRHIReleaseWindowCallbackType = MoveTemp(InRHIOnReleaseWindowCallback);
 		}
 
+		static FRHIOnWindowCreatedCallbackType OnRHIWindowCreatedCallback;
+
+		FRHIOnWindowCreatedCallbackType& GetRHIOnWindowCreatedCallback()
+		{
+			return OnRHIWindowCreatedCallback;
+		}
+
+		void SetRHIOnWindowCreatedCallback(FRHIOnWindowCreatedCallbackType&& InOnWindowCreatedCallback)
+		{
+			OnRHIWindowCreatedCallback = MoveTemp(InOnWindowCreatedCallback);
+		}
+
 	} // namespace FAndroidPlatformDynamicRHI
 } // namespace UE
 
diff --git a/Engine/Source/Runtime/RHI/Public/Android/AndroidDynamicRHI.h b/Engine/Source/Runtime/RHI/Public/Android/AndroidDynamicRHI.h
--- a/Engine/Source/Runtime/RHI/Public/Android/AndroidDynamicRHI.h
+++ b/Engine/Source/Runtime/RHI/Public/Android/AndroidDynamicRHI.h
@@ -16,6 +16,9 @@ namespace UE
 		FRHIReleaseWindowCallbackType& GetRHIOnReleaseWindowCallback();
 		void SetRHIOnReleaseWindowCallback(FRHIReleaseWindowCallbackType&& InOnReleaseWindowCallback);
 
+		typedef TUniqueFunction<void()> FRHIOnWindowCreatedCallbackType;
+		FRHIOnWindowCreatedCallbackType& GetRHIOnWindowCreatedCallback();
+		void SetRHIOnWindowCreatedCallback(FRHIOnWindowCreatedCallbackType&& InOnWindowCreatedCallback);
 
 		// Helper class to convert FGraphicsPipelineStateInitializer::EPSOPrecacheCompileType to the uint64 required by android's external PSO service.
 		class FPSOServicePriInfo
