Fix merge issues
#rnx

diff --git a/Engine/Source/Runtime/OpenGLDrv/Private/Android/AndroidOpenGL.cpp b/Engine/Source/Runtime/OpenGLDrv/Private/Android/AndroidOpenGL.cpp
--- a/Engine/Source/Runtime/OpenGLDrv/Private/Android/AndroidOpenGL.cpp
+++ b/Engine/Source/Runtime/OpenGLDrv/Private/Android/AndroidOpenGL.cpp
@@ -246,7 +246,6 @@ void FPlatformOpenGLDevice::Init()
 							{
 								// find the viewport using this window.
 								UE_LOGF(LogRHI, Log, "GLES OnReleaseWindowCallback: RHI, set egl surface to offscreen (window %p)", WindowContainer->GetANativeWindow());
-								AndroidEGL::GetInstance()->UnBindRender();
 								int ViewportsAffected = 0;
 								FOpenGLDynamicRHI::FGLViewportContainer Viewports = FOpenGLDynamicRHI::Get().GetViewportContainer();
 								for (FOpenGLViewport* Viewport : Viewports.Get())
@@ -258,7 +257,7 @@ void FPlatformOpenGLDevice::Init()
 										AndroidEGL::GetInstance()->OnLostNativeWindow(Surface);
 									}
 								}
-								UE_LOGF(LogRHI, Log, "GLES OnReleaseWindowCallback: RHI, swap surface done. (%d viewports)"), ViewportsAffected);
+								UE_LOGF(LogRHI, Log, "GLES OnReleaseWindowCallback: RHI, swap surface done. (%d viewports)", ViewportsAffected);
 							});
 						RHICmdList.RHIThreadFence(true);
 						RHICmdList.ImmediateFlush(EImmediateFlushType::FlushRHIThread);
